Chaos Cloth - Move the new Cloth Asset modules out of the Restricted area to the Engine/Plugins/Experimental folder (part 2, add files)

diff --git a/Engine/Plugins/Experimental/ChaosClothAsset/Source/ChaosClothAsset/Private/ChaosClothAsset/ClothAdapter.cpp b/Engine/Plugins/Experimental/ChaosClothAsset/Source/ChaosClothAsset/Private/ChaosClothAsset/ClothAdapter.cpp
new file mode 100644
--- /dev/null
+++ b/Engine/Plugins/Experimental/ChaosClothAsset/Source/ChaosClothAsset/Private/ChaosClothAsset/ClothAdapter.cpp
@@ -0,0 +1,56 @@
+// Copyright Epic Games, Inc. All Rights Reserved.
+
+#include "ChaosClothAsset/ClothAdapter.h"
+
+namespace UE::Chaos::ClothAsset
+{
+	FClothConstAdapter::FClothConstAdapter(const TSharedPtr<const FClothCollection>& InClothCollection)
+		: ClothCollection(InClothCollection)
+	{
+		check(ClothCollection.IsValid());
+	}
+
+	FClothConstAdapter::FClothConstAdapter(const FClothLodConstAdapter& ClothLodConstAdapter)
+		: ClothCollection(ClothLodConstAdapter.GetClothCollection())
+	{
+	}
+
+	FClothConstAdapter::FClothConstAdapter(const FClothPatternConstAdapter& ClothPatternConstAdapter)
+		: ClothCollection(ClothPatternConstAdapter.GetClothCollection())
+	{
+	}
+
+	FClothLodConstAdapter FClothConstAdapter::GetLod(int32 LodIndex) const
+	{
+		return FClothLodConstAdapter(ClothCollection, LodIndex);
+	}
+
+	FClothAdapter::FClothAdapter(const TSharedPtr<FClothCollection>& InClothCollection)
+		: FClothConstAdapter(InClothCollection)
+	{
+	}
+
+	int32 FClothAdapter::AddLod()
+	{
+		const int32 LodIndex = GetClothCollection()->AddElements(1, FClothCollection::LodsGroup);
+
+		FClothLodAdapter(GetClothCollection(), LodIndex).SetDefaults();
+
+		return LodIndex;
+	}
+
+	FClothLodAdapter FClothAdapter::GetLod(int32 LodIndex)
+	{
+		return FClothLodAdapter(GetClothCollection(), LodIndex);
+	}
+
+	void FClothAdapter::Reset()
+	{
+		const int32 NumLods = GetNumLods();
+		for (int32 LodIndex = 0; LodIndex < NumLods; ++LodIndex)
+		{
+			GetLod(LodIndex).Reset();
+		}
+		GetClothCollection()->EmptyGroup(FClothCollection::LodsGroup);
+	}
+}  // End namespace UE::Chaos::ClothAsset
diff --git a/Engine/Plugins/Experimental/ChaosClothAsset/Source/ChaosClothAsset/Private/ChaosClothAsset/ClothCollection.cpp b/Engine/Plugins/Experimental/ChaosClothAsset/Source/ChaosClothAsset/Private/ChaosClothAsset/ClothCollection.cpp
new file mode 100644
--- /dev/null
+++ b/Engine/Plugins/Experimental/ChaosClothAsset/Source/ChaosClothAsset/Private/ChaosClothAsset/ClothCollection.cpp
@@ -0,0 +1,102 @@
+// Copyright Epic Games, Inc. All Rights Reserved.
+
+#include "ChaosClothAsset/ClothCollection.h"
+
+namespace UE::Chaos::ClothAsset
+{
+	const FName UE::Chaos::ClothAsset::FClothCollection::SimVerticesGroup("SimVertices");
+	const FName UE::Chaos::ClothAsset::FClothCollection::SimFacesGroup("SimFaces");
+	const FName UE::Chaos::ClothAsset::FClothCollection::RenderVerticesGroup("RenderVertices");
+	const FName UE::Chaos::ClothAsset::FClothCollection::RenderFacesGroup("RenderFaces");
+	const FName UE::Chaos::ClothAsset::FClothCollection::WrapDeformersGroup("WrapDeformers");
+	const FName UE::Chaos::ClothAsset::FClothCollection::PatternsGroup("Patterns");
+	const FName UE::Chaos::ClothAsset::FClothCollection::StitchingsGroup("Stitchings");
+	const FName UE::Chaos::ClothAsset::FClothCollection::TethersGroup("Tethers");
+	const FName UE::Chaos::ClothAsset::FClothCollection::TetherBatchesGroup("TetherBatches");
+	const FName UE::Chaos::ClothAsset::FClothCollection::LodsGroup("Lods");
+
+	FClothCollection::FClothCollection()
+	{
+		Construct();
+	}
+
+	void FClothCollection::Serialize(FArchive& Ar)
+	{
+		::Chaos::FChaosArchive ChaosArchive(Ar);
+		Super::Serialize(ChaosArchive);
+	}
+
+	void FClothCollection::Construct()
+	{
+		// Dependencies
+		FManagedArrayCollection::FConstructionParameters SimVerticesDependency(SimVerticesGroup);
+		FManagedArrayCollection::FConstructionParameters SimFacesDependency(SimFacesGroup);
+		FManagedArrayCollection::FConstructionParameters RenderVerticesDependency(RenderVerticesGroup);
+		FManagedArrayCollection::FConstructionParameters RenderFacesDependency(RenderFacesGroup);
+		FManagedArrayCollection::FConstructionParameters WrapDeformersDependency(WrapDeformersGroup);
+		FManagedArrayCollection::FConstructionParameters PatternsDependency(PatternsGroup);
+		FManagedArrayCollection::FConstructionParameters StitchingsDependency(StitchingsGroup);
+		FManagedArrayCollection::FConstructionParameters TethersDependency(TethersGroup);
+		FManagedArrayCollection::FConstructionParameters TetherBatchesDependency(TetherBatchesGroup);
+
+		// Sim Vertices Group
+		AddExternalAttribute<FVector2f>("SimPosition", SimVerticesGroup, SimPosition);
+		AddExternalAttribute<FVector3f>("SimRestPosition", SimVerticesGroup, SimRestPosition);
+		AddExternalAttribute<FVector3f>("SimRestNormal", SimVerticesGroup, SimRestNormal);
+		//AddExternalAttribute<FVertexBoneData>("SimBoneData", SimVerticesGroup, SimBoneData);
+
+		// Sim Faces Group
+		AddExternalAttribute<FIntVector3>("SimIndices", SimFacesGroup, SimIndices, SimVerticesDependency);
+
+		// Render Vertices Group
+		AddExternalAttribute<FVector3f>("RenderPosition", RenderVerticesGroup, RenderPosition);
+		AddExternalAttribute<FVector3f>("RenderNormal", RenderVerticesGroup, RenderNormal);
+		AddExternalAttribute<FVector3f>("RenderTangentU", RenderVerticesGroup, RenderTangentU);
+		AddExternalAttribute<FVector3f>("RenderTangentV", RenderVerticesGroup, RenderTangentV);
+		AddExternalAttribute<TArray<FVector2f>>("RenderUVs", RenderVerticesGroup, RenderUVs);
+		AddExternalAttribute<FLinearColor>("RenderColor", RenderVerticesGroup, RenderColor);
+		//AddExternalAttribute<FVertexBoneData>("RenderBoneData", RenderVerticesGroup, RenderBoneData);
+
+		// Render Faces Group
+		AddExternalAttribute<FIntVector3>("RenderIndices", RenderFacesGroup, RenderIndices, RenderVerticesDependency);
+		AddExternalAttribute<int32>("RenderMaterialIndex", RenderFacesGroup, RenderMaterialIndex);
+
+		// TODO: FMeshToMeshVertData
+
+		// Patterns Group
+		AddExternalAttribute<int32>("SimVerticesStart", PatternsGroup, SimVerticesStart, SimVerticesDependency);
+		AddExternalAttribute<int32>("SimVerticesEnd", PatternsGroup, SimVerticesEnd, SimVerticesDependency);
+		AddExternalAttribute<int32>("SimFacesStart", PatternsGroup, SimFacesStart, SimFacesDependency);
+		AddExternalAttribute<int32>("SimFacesEnd", PatternsGroup, SimFacesEnd, SimFacesDependency);
+		AddExternalAttribute<int32>("RenderVerticesStart", PatternsGroup, RenderVerticesStart, RenderVerticesDependency);
+		AddExternalAttribute<int32>("RenderVerticesEnd", PatternsGroup, RenderVerticesEnd, RenderVerticesDependency);
+		AddExternalAttribute<int32>("RenderFacesStart", PatternsGroup, RenderFacesStart, RenderFacesDependency);
+		AddExternalAttribute<int32>("RenderFacesEnd", PatternsGroup, RenderFacesEnd, RenderFacesDependency);
+		AddExternalAttribute<int32>("WrapDeformerStart", PatternsGroup, WrapDeformerStart, WrapDeformersDependency);
+		AddExternalAttribute<int32>("WrapDeformerEnd", PatternsGroup, WrapDeformerEnd, WrapDeformersDependency);
+		AddExternalAttribute<int32>("NumWeights", PatternsGroup, NumWeights);
+		AddExternalAttribute<int32>("StatusFlags", PatternsGroup, StatusFlags);
+		AddExternalAttribute<int32>("SimMaterialIndex", PatternsGroup, SimMaterialIndex);
+
+		// Stitching Group
+		AddExternalAttribute<TSet<int32>>("StitchedVertices", StitchingsGroup, StitchedVertices, SimVerticesDependency);
+
+		// Tethers Group
+		AddExternalAttribute<int32>("TetherKinematicIndex", TethersGroup, TetherKinematicIndex, SimVerticesDependency);
+		AddExternalAttribute<int32>("TetherDynamicIndex", TethersGroup, TetherDynamicIndex, SimVerticesDependency);
+		AddExternalAttribute<float>("TetherReferenceLength", TethersGroup, TetherReferenceLength);
+
+		// Tether Batches Group
+		AddExternalAttribute<int32>("TetherStart", TetherBatchesGroup, TetherStart, TethersDependency);
+		AddExternalAttribute<int32>("TetherEnd", TetherBatchesGroup, TetherEnd, TethersDependency);
+
+		// LOD Group
+		AddExternalAttribute<int32>("PatternStart", LodsGroup, PatternStart, PatternsDependency);
+		AddExternalAttribute<int32>("PatternEnd", LodsGroup, PatternEnd, PatternsDependency);
+		AddExternalAttribute<int32>("StitchingStart", LodsGroup, StitchingStart, StitchingsDependency);
+		AddExternalAttribute<int32>("StitchingEnd", LodsGroup, StitchingEnd, StitchingsDependency);
+		AddExternalAttribute<int32>("TetherBatchStart", LodsGroup, TetherBatchStart, TetherBatchesDependency);
+		AddExternalAttribute<int32>("TetherBatchEnd", LodsGroup, TetherBatchEnd, TetherBatchesDependency);
+		AddExternalAttribute<int32>("LodBiasDepth", LodsGroup, LodBiasDepth);
+	}
+} // End namespace UE::Chaos::ClothAsset
diff --git a/Engine/Plugins/Experimental/ChaosClothAsset/Source/ChaosClothAsset/Private/ChaosClothAsset/ClothGeometryTools.cpp b/Engine/Plugins/Experimental/ChaosClothAsset/Source/ChaosClothAsset/Private/ChaosClothAsset/ClothGeometryTools.cpp
new file mode 100644
--- /dev/null
+++ b/Engine/Plugins/Experimental/ChaosClothAsset/Source/ChaosClothAsset/Private/ChaosClothAsset/ClothGeometryTools.cpp
@@ -0,0 +1,112 @@
+// Copyright Epic Games, Inc. All Rights Reserved.
+
+#include "ChaosClothAsset/ClothGeometryTools.h"
+#include "ChaosClothAsset/ClothCollection.h"
+#include "ChaosClothAsset/ClothAdapter.h"
+
+namespace UE::Chaos::ClothAsset
+{
+	void FClothGeometryTools::CopySimMeshToRenderMesh(const TSharedPtr<FClothCollection>& ClothCollection, int32 MaterialIndex)
+	{
+		FClothAdapter ClothAdapter(ClothCollection);
+
+		for (int32 LodIndex = 0; LodIndex < ClothAdapter.GetNumLods(); ++LodIndex)
+		{
+			FClothLodAdapter ClothLodAdapter = ClothAdapter.GetLod(LodIndex);
+
+			const TConstArrayView<FVector2f> LodSimPosition = ClothLodAdapter.GetPatternsSimPosition();
+			const TConstArrayView<FVector3f> LodSimRestPosition = ClothLodAdapter.GetPatternsSimRestPosition();
+
+			for (int32 PatternIndex = 0; PatternIndex < ClothLodAdapter.GetNumPatterns(); ++PatternIndex)
+			{
+				FClothPatternAdapter ClothPatternAdapter = ClothLodAdapter.GetPattern(PatternIndex);
+
+				const int32 NumVertices = ClothPatternAdapter.GetNumSimVertices();
+				const int32 NumFaces = ClothPatternAdapter.GetNumSimFaces();
+
+				if (!NumVertices || !NumFaces)
+				{
+					ClothPatternAdapter.SetNumRenderVertices(0);
+					ClothPatternAdapter.SetNumRenderFaces(0);
+					continue;
+				}
+
+				ClothPatternAdapter.SetNumRenderVertices(NumVertices);
+				ClothPatternAdapter.SetNumRenderFaces(NumFaces);
+
+				const TConstArrayView<FVector2f> SimPosition = ClothPatternAdapter.GetSimPosition();
+				const TConstArrayView<FVector3f> SimRestPosition = ClothPatternAdapter.GetSimRestPosition();
+				const TConstArrayView<FVector3f> SimRestNormal = ClothPatternAdapter.GetSimRestNormal();
+				const TArrayView<FVector3f> RenderTangentU = ClothPatternAdapter.GetRenderTangentU();
+				const TArrayView<FVector3f> RenderTangentV = ClothPatternAdapter.GetRenderTangentV();
+
+				FVector2f MinPosition(TNumericLimits<float>::Max());
+				FVector2f MaxPosition(TNumericLimits<float>::Lowest());
+
+				// Calculate UVs scale and zero out tangents
+				for (int32 VertexIndex = 0; VertexIndex < NumVertices; ++VertexIndex)
+				{
+					MinPosition = FVector2f::Min(MinPosition, SimPosition[VertexIndex]);
+					MaxPosition = FVector2f::Max(MaxPosition, SimPosition[VertexIndex]);
+
+					RenderTangentU[VertexIndex] = FVector3f::ZeroVector;
+					RenderTangentV[VertexIndex] = FVector3f::ZeroVector;
+				}
+				const FVector2f UVScale = MaxPosition - MinPosition;
+				const FVector2f UVInvScale(
+					UVScale.X < SMALL_NUMBER ? 0.f : 1.f / UVScale.X,
+					UVScale.Y < SMALL_NUMBER ? 0.f : 1.f / UVScale.Y);
+
+				// Face group
+				const TConstArrayView<FIntVector3> SimIndices = ClothPatternAdapter.GetSimIndices();
+				const TArrayView<FIntVector3> RenderIndices = ClothPatternAdapter.GetRenderIndices();
+				const TArrayView<int32> RenderMaterialIndex = ClothPatternAdapter.GetRenderMaterialIndex();
+				const TArrayView<FVector3f> LodRenderTangentU = ClothLodAdapter.GetPatternsRenderTangentU();
+				const TArrayView<FVector3f> LodRenderTangentV = ClothLodAdapter.GetPatternsRenderTangentV();
+				for (int32 FaceIndex = 0; FaceIndex < NumFaces; ++FaceIndex)
+				{
+					const FIntVector3& Face = SimIndices[FaceIndex];
+
+					RenderIndices[FaceIndex] = Face;
+					RenderMaterialIndex[FaceIndex] = MaterialIndex;
+
+					// Edges of the triangle : position delta
+					const FVector3f DeltaPos1 = LodSimRestPosition[Face[1]] - LodSimRestPosition[Face[0]];
+					const FVector3f DeltaPos2 = LodSimRestPosition[Face[2]] - LodSimRestPosition[Face[0]];
+
+					// UV delta
+					const FVector2f DeltaUV1 = LodSimPosition[Face[1]] - LodSimPosition[Face[0]];
+					const FVector2f DeltaUV2 = LodSimPosition[Face[2]] - LodSimPosition[Face[0]];
+					
+					// We can now use our formula to compute the tangent and the bitangent :
+					const float Denom = DeltaUV1.X * DeltaUV2.Y - DeltaUV1.Y * DeltaUV2.X;
+					const float R = (Denom < SMALL_NUMBER) ? 0.f : 1.f / Denom;
+					const FVector3f TangentU = (DeltaPos1 * DeltaUV2.Y - DeltaPos2 * DeltaUV1.Y) * R;
+					const FVector3f TangentV = (DeltaPos2 * DeltaUV1.X - DeltaPos1 * DeltaUV2.X) * R;
+
+					for (int32 PointIndex = 0; PointIndex < 3; ++PointIndex)
+					{
+						LodRenderTangentU[Face[PointIndex]] += TangentU;
+						LodRenderTangentV[Face[PointIndex]] += TangentV;
+					}
+				}
+
+				// Vertex group
+				const TArrayView<FVector3f> RenderPosition = ClothPatternAdapter.GetRenderPosition();
+				const TArrayView<FVector3f> RenderNormal = ClothPatternAdapter.GetRenderNormal();
+				const TArrayView<TArray<FVector2f>> RenderUVs = ClothPatternAdapter.GetRenderUVs();
+				const TArrayView<FLinearColor> RenderColor = ClothPatternAdapter.GetRenderColor();
+
+				for (int32 VertexIndex = 0; VertexIndex < NumVertices; ++VertexIndex)
+				{
+					RenderPosition[VertexIndex] = SimRestPosition[VertexIndex];
+					RenderNormal[VertexIndex] = SimRestNormal[VertexIndex];
+					RenderUVs[VertexIndex] = { SimPosition[VertexIndex] * UVInvScale };
+					RenderColor[VertexIndex] = FLinearColor::White;
+					RenderTangentU[VertexIndex].Normalize();
+					RenderTangentV[VertexIndex].Normalize();
+				}
+			}
+		}
+	}
+}  // End namespace UE::Chaos::ClothAsset
\ No newline at end of file
diff --git a/Engine/Plugins/Experimental/ChaosClothAsset/Source/ChaosClothAsset/Private/ChaosClothAsset/ClothLodAdapter.cpp b/Engine/Plugins/Experimental/ChaosClothAsset/Source/ChaosClothAsset/Private/ChaosClothAsset/ClothLodAdapter.cpp
new file mode 100644
--- /dev/null
+++ b/Engine/Plugins/Experimental/ChaosClothAsset/Source/ChaosClothAsset/Private/ChaosClothAsset/ClothLodAdapter.cpp
@@ -0,0 +1,449 @@
+// Copyright Epic Games, Inc. All Rights Reserved.
+
+#include "ChaosClothAsset/ClothLodAdapter.h"
+
+// Utility functions to unwrap a 3d sim mesh into a tailored cloth
+namespace UE::Chaos::ClothAsset::Private
+{
+	struct FIsland
+	{
+		TArray<uint32> Indices;  // 3x number of triangles
+		TArray<FVector2f> Positions;
+		TArray<FVector3f> RestPositions;  // Same size as Positions
+	};
+
+	enum class EIntersectCirclesResult
+	{
+		SingleIntersect,
+		DoubleIntersect,
+		Coincident,
+		Separate,
+		Contained
+	};
+
+	static EIntersectCirclesResult IntersectCircles(const FVector2f& C0, float R0, const FVector2f& C1, float R1, FVector2f& OutI0, FVector2f& OutI1)
+	{
+		const FVector2f C0C1 = C0 - C1;
+		const float D = C0C1.Length();
+		if (D < SMALL_NUMBER)
+		{
+			return EIntersectCirclesResult::Coincident;
+		}
+		else if (D > R0 + R1)
+		{
+			return EIntersectCirclesResult::Separate;
+		}
+		else if (D < FMath::Abs(R0 - R1))
+		{
+			return EIntersectCirclesResult::Contained;
+		}
+		const float SquareR0 = R0 * R0;
+		const float SquareR1 = R1 * R1;
+		const float SquareD = D * D;
+		const float A = (SquareD - SquareR1 + SquareR0) / (2.f * D);
+
+		OutI0 = OutI1 = C0 + A * (C1 - C0) / D;
+
+		if (FMath::Abs(A - R0) < SMALL_NUMBER)
+		{
+			return EIntersectCirclesResult::SingleIntersect;
+		}
+
+		const float SquareA = A * A;
+		const float H = FMath::Sqrt(SquareR0 - SquareA);
+
+		const FVector2f N(C0C1.Y, -C0C1.X);
+
+		OutI0 += N * H / D;
+		OutI1 -= N * H / D;
+
+		return EIntersectCirclesResult::DoubleIntersect;
+	}
+
+	static void UnwrapMesh(const TArray<FVector3f>& Positions, const TArray<uint32>& Indices, TArray<FIsland>& OutIslands)
+	{
+		OutIslands.Reset();
+
+		ensure(Indices.Num() % 3 == 0);
+		const uint32 NumTriangles = (uint32)Indices.Num() / 3;
+		if (!NumTriangles)
+		{
+			return;
+		}
+
+		// Gather edge information
+		TMap<FUintVector2, TArray<uint32>> EdgeToTriangles;
+		EdgeToTriangles.Reserve(NumTriangles * 2);  // Rough estimate for the number of edges
+
+		auto MakeSortedEdge = [](uint32 Index0, uint32 Index1)->FUintVector2
+		{
+			return Index0 < Index1 ? FUintVector2(Index0, Index1) : FUintVector2(Index1, Index0);
+		};
+
+		for (uint32 Triangle = 0; Triangle < NumTriangles; ++Triangle)
+		{
+			const uint32 Index0 = Indices[Triangle * 3 + 0];
+			const uint32 Index1 = Indices[Triangle * 3 + 1];
+			const uint32 Index2 = Indices[Triangle * 3 + 2];
+
+			const FUintVector2 Edge0 = MakeSortedEdge(Index0, Index1);
+			const FUintVector2 Edge1 = MakeSortedEdge(Index1, Index2);
+			const FUintVector2 Edge2 = MakeSortedEdge(Index2, Index0);
+
+			EdgeToTriangles.FindOrAdd(Edge0).Add(Triangle);
+			EdgeToTriangles.FindOrAdd(Edge1).Add(Triangle);
+			EdgeToTriangles.FindOrAdd(Edge2).Add(Triangle);
+		}
+
+		// Build pattern islands
+		TSet<uint32> VisitedTriangles;
+		VisitedTriangles.Reserve(NumTriangles);
+
+		constexpr float SquaredWeldingDistance = FMath::Square(0.01f);  // 0.1 mm
+
+		for (uint32 SeedTriangle = 0; SeedTriangle < NumTriangles; ++SeedTriangle)
+		{
+			if (VisitedTriangles.Contains(SeedTriangle))
+			{
+				continue;
+			}
+
+			const uint32 SeedIndex0 = Indices[SeedTriangle * 3 + 0];
+			const uint32 SeedIndex1 = Indices[SeedTriangle * 3 + 1];
+
+			if (FVector3f::DistSquared(Positions[SeedIndex0], Positions[SeedIndex1]) <= SquaredWeldingDistance)
+			{
+				continue;  // A degenerated triangle edge is not a good start
+			}
+
+			// Setup first visitor from seed, and add the first two points
+			FIsland& Island = OutIslands.AddDefaulted_GetRef();
+
+			Island.RestPositions.Add(Positions[SeedIndex0]);
+			Island.RestPositions.Add(Positions[SeedIndex1]);
+
+			const uint32 SeedIndex2D0 = Island.Positions.Add(FVector2f::ZeroVector);
+			const uint32 SeedIndex2D1 = Island.Positions.Add(FVector2f(FVector3f::Dist(Positions[SeedIndex0], Positions[SeedIndex1]), 0.f));
+
+			struct FVisitor
+			{
+				uint32 Triangle;
+				FUintVector2 OldEdge;
+				FUintVector2 NewEdge;
+				uint32 CrossEdgePoint;  // Keep the opposite point to orientate degenerate cases
+			} Visitor =
+			{
+				SeedTriangle,
+				FUintVector2(SeedIndex0, SeedIndex1),
+				FUintVector2(SeedIndex2D0, SeedIndex2D1),
+				(uint32)INDEX_NONE
+			};
+
+			VisitedTriangles.Add(SeedTriangle);
+
+			TQueue<FVisitor> Visitors;
+			do
+			{
+				const uint32 Triangle = Visitor.Triangle;
+				const uint32 CrossEdgePoint = Visitor.CrossEdgePoint;
+				const uint32 OldIndex0 = Visitor.OldEdge.X;
+				const uint32 OldIndex1 = Visitor.OldEdge.Y;
+				const uint32 NewIndex0 = Visitor.NewEdge.X;
+				const uint32 NewIndex1 = Visitor.NewEdge.Y;
+
+				// Find opposite index from this triangle edge
+				const uint32 TriangleIndex0 = Indices[Triangle * 3 + 0];
+				const uint32 TriangleIndex1 = Indices[Triangle * 3 + 1];
+				const uint32 TriangleIndex2 = Indices[Triangle * 3 + 2];
+
+				const uint32 OldIndex2 =
+					(OldIndex0 != TriangleIndex0 && OldIndex1 != TriangleIndex0) ? TriangleIndex0 :
+					(OldIndex0 != TriangleIndex1 && OldIndex1 != TriangleIndex1) ? TriangleIndex1 : TriangleIndex2;
+
+				// Find the 2D interesection of the two connecting adjacent edges using the 3D reference length
+				const FVector3f& P0 = Positions[OldIndex0];
+				const FVector3f& P1 = Positions[OldIndex1];
+				const FVector3f& P2 = Positions[OldIndex2];
+
+				const float R0 = FVector3f::Dist(P0, P2);
+				const float R1 = FVector3f::Dist(P1, P2);
+				const FVector2f C0 = Island.Positions[NewIndex0];
+				const FVector2f C1 = Island.Positions[NewIndex1];
+
+				FVector2f I0, I1;
+				const EIntersectCirclesResult IntersectCirclesResult = IntersectCircles(C0, R0, C1, R1, I0, I1);
+
+				FVector2f C2;
+				switch (IntersectCirclesResult)
+				{
+				case EIntersectCirclesResult::SingleIntersect:
+					C2 = I0;  // Degenerated C2 is on (C0C1)
+					break;
+				case EIntersectCirclesResult::DoubleIntersect:
+					C2 = (FVector2f::CrossProduct(C0 - C1, C0 - I0) > 0) ? I0 : I1;  // Keep correct winding order
+					break;
+				case EIntersectCirclesResult::Coincident:
+					check(CrossEdgePoint != INDEX_NONE); // We can't start on a degenerated triangle
+					C2 = C0 - (Island.Positions[CrossEdgePoint] - C0).GetSafeNormal() * R0;  // Degenerated C0 == C1, choose C2 on the opposite of the visitor opposite point
+					break;
+				case EIntersectCirclesResult::Separate: [[fallthrough]];
+				case EIntersectCirclesResult::Contained:
+					C2 = C0 - (C1 - C0).GetSafeNormal() * R0;  // Degenerated + some tolerance, C2 is on (C0C1)
+					break;
+				}
+
+				// Add the new position found for the opposite point
+				uint32 NewIndex2 = INDEX_NONE;
+				for (uint32 UsedIndex = 0; UsedIndex < (uint32)Island.Positions.Num(); ++UsedIndex)
+				{
+					if (FVector2f::DistSquared(Island.Positions[UsedIndex], C2) <= SquaredWeldingDistance &&
+						FVector3f::DistSquared(Island.RestPositions[UsedIndex], Positions[OldIndex2]) <= SquaredWeldingDistance)
+					{
+						NewIndex2 = UsedIndex;  // Both Rest and 2D positions match, reuse this index
+						break;
+					}
+				}
+				if (NewIndex2 == INDEX_NONE)
+				{
+					NewIndex2 = Island.Positions.Add(C2);
+					Island.RestPositions.Add(Positions[OldIndex2]);
+				}
+
+				// Add triangle to list of indices, unless it is degenerated to a segment
+				if (NewIndex0 != NewIndex1 && NewIndex1 != NewIndex2 && NewIndex2 != NewIndex0)
+				{
+					Island.Indices.Add(NewIndex0);
+					Island.Indices.Add(NewIndex1);
+					Island.Indices.Add(NewIndex2);
+				}
+
+				// Add neighbor triangles to the queue
+				const FUintVector2 OldEdgeList[3] =
+				{
+					FUintVector2(OldIndex1, OldIndex0),  // Reversed as to keep the correct winding order
+					FUintVector2(OldIndex2, OldIndex1),
+					FUintVector2(OldIndex0, OldIndex2)
+				};
+				const FUintVector3 NewEdgeList[3] =
+				{
+					FUintVector3(NewIndex1, NewIndex0, NewIndex2),  // Adds opposite point index
+					FUintVector3(NewIndex2, NewIndex1, NewIndex0),
+					FUintVector3(NewIndex0, NewIndex2, NewIndex1)
+				};
+				for (int32 Edge = 0; Edge < 3; ++Edge)
+				{
+					const uint32 EdgeIndex0 = OldEdgeList[Edge].X;
+					const uint32 EdgeIndex1 = OldEdgeList[Edge].Y;
+
+					const TArray<uint32>& NeighborTriangles = EdgeToTriangles.FindChecked(MakeSortedEdge(EdgeIndex0, EdgeIndex1));
+
+					for (const uint32 NeighborTriangle : NeighborTriangles)
+					{
+						if (!VisitedTriangles.Contains(NeighborTriangle))
+						{
+							// Mark neighboring triangle as visited
+							VisitedTriangles.Add(NeighborTriangle);
+
+							// Enqueue next triangle
+							Visitors.Enqueue(FVisitor
+								{
+									NeighborTriangle,
+									FUintVector2(EdgeIndex0, EdgeIndex1),
+									FUintVector2(NewEdgeList[Edge].X, NewEdgeList[Edge].Y),
+									NewEdgeList[Edge].Z,  // Pass the cross edge 2D opposite point to help define orientation of any degenerated triangles
+								});
+						}
+					}
+				}
+			} while (Visitors.Dequeue(Visitor));
+		}
+	}
+}  // End namespace UE::ClothAsset::Private
+
+// Cloth LOD adapter
+namespace UE::Chaos::ClothAsset
+{
+	FClothLodConstAdapter::FClothLodConstAdapter(const TSharedPtr<const FClothCollection>& InClothCollection, int32 InLodIndex)
+		: ClothCollection(InClothCollection)
+		, LodIndex(InLodIndex)
+	{
+		check(ClothCollection.IsValid());
+		check(LodIndex >= 0 && LodIndex < ClothCollection->NumElements(FClothCollection::LodsGroup));
+	}
+
+	FClothLodConstAdapter::FClothLodConstAdapter(const FClothPatternConstAdapter& ClothPatternConstAdapter)
+		: ClothCollection(ClothPatternConstAdapter.GetClothCollection())
+		, LodIndex(ClothPatternConstAdapter.GetLodIndex())
+	{
+	}
+
+	FClothPatternConstAdapter FClothLodConstAdapter::GetPattern(int32 PatternIndex) const
+	{
+		return FClothPatternConstAdapter(ClothCollection, LodIndex, PatternIndex);
+	}
+
+	int32 FClothLodConstAdapter::GetNumElements(const TManagedArray<int32>& StartArray, const TManagedArray<int32>& EndArray) const
+	{
+		const int32 ElementIndex = GetElementIndex();
+		const int32 Start = StartArray[ElementIndex];
+		const int32 End = EndArray[ElementIndex];
+		check(Start != INDEX_NONE || End == INDEX_NONE);  // Best to avoid situations where only one boundary of the range is set to INDEX_NONE
+
+		return Start == INDEX_NONE ? 0 : End - Start + 1;
+	}
+
+	template<typename T>
+	TConstArrayView<T> FClothLodConstAdapter::GetElements(const TManagedArray<T>& ElementArray, const TManagedArray<int32>& StartArray, const TManagedArray<int32>& EndArray) const
+	{
+		const int32 ElementIndex = GetElementIndex();
+		const int32 Start = StartArray[ElementIndex];
+		const int32 End = EndArray[ElementIndex];
+		check(Start != INDEX_NONE || End == INDEX_NONE);  // Best to avoid situations where only one boundary of the range is set to INDEX_NONE
+
+		return Start == INDEX_NONE ? TConstArrayView<T>() : TConstArrayView<T>(ElementArray.GetData() + Start, End - Start + 1);
+	}
+
+	template<bool bStart, bool bEnd>
+	TTuple<int32, int32> FClothLodConstAdapter::GetPatternsElementsStartEnd(const TManagedArray<int32>& StartArray, const TManagedArray<int32>& EndArray) const
+	{
+		const int32 ElementIndex = GetElementIndex();
+		const int32 PatternStart = ClothCollection->PatternStart[ElementIndex];
+		const int32 PatternEnd = ClothCollection->PatternEnd[ElementIndex];
+
+		// Find Start and End indices for the entire LOD minding empty patterns on the way
+		int32 Start = INDEX_NONE;
+		int32 End = INDEX_NONE;
+		for (int32 PatternIndex = PatternStart; PatternIndex <= PatternEnd; ++PatternIndex)
+		{
+			if (bStart && StartArray[PatternIndex] != INDEX_NONE)
+			{
+				Start = (Start == INDEX_NONE) ? StartArray[PatternIndex] : FMath::Min(Start, StartArray[PatternIndex]);
+			}
+			if (bEnd && EndArray[PatternIndex] != INDEX_NONE)
+			{
+				End = (End == INDEX_NONE) ? EndArray[PatternIndex] : FMath::Max(End, EndArray[PatternIndex]);
+			}
+		}
+		return TTuple<int32, int32>(Start, End);
+	}
+
+	int32 FClothLodConstAdapter::GetPatternsNumElements(const TManagedArray<int32>& StartArray, const TManagedArray<int32>& EndArray) const
+	{
+		const TTuple<int32, int32> StartEnd = GetPatternsElementsStartEnd<true, true>(StartArray, EndArray);
+		const int32 Start = StartEnd.Get<0>();
+		const int32 End = StartEnd.Get<1>();
+		check(Start != INDEX_NONE || End == INDEX_NONE);  // Best to avoid situations where only one boundary of the range is set to INDEX_NONE
+
+		return Start == INDEX_NONE ? 0 : End - Start + 1;
+	}
+
+	template<typename T>
+	TConstArrayView<T> FClothLodConstAdapter::GetPatternsElements(const TManagedArray<T>& ElementArray, const TManagedArray<int32>& StartArray, const TManagedArray<int32>& EndArray) const
+	{
+		const TTuple<int32, int32> StartEnd = GetPatternsElementsStartEnd<true, true>(StartArray, EndArray);
+		const int32 Start = StartEnd.Get<0>();
+		const int32 End = StartEnd.Get<1>();
+
+		check(Start != INDEX_NONE || End == INDEX_NONE);  // Best to avoid situations where only one boundary of the range is set to INDEX_NONE
+		return Start == INDEX_NONE ? TConstArrayView<T>() : TConstArrayView<T>(ElementArray.GetData() + Start, End - Start + 1);
+	}
+
+	FClothLodAdapter::FClothLodAdapter(const TSharedPtr<FClothCollection>& InClothCollection, int32 InLodIndex)
+		: FClothLodConstAdapter(InClothCollection, InLodIndex)
+	{
+	}
+
+	int32 FClothLodAdapter::AddPattern()
+	{
+		const int32 PatternElementIndex = GetClothCollection()->AddElements(1, FClothCollection::PatternsGroup);
+
+		GetClothCollection()->PatternEnd[GetElementIndex()] = PatternElementIndex;
+
+		// If this is the first pattern being added, set also the start
+		int32& PatternStart = GetClothCollection()->PatternStart[GetElementIndex()];
+		if (PatternStart == INDEX_NONE)
+		{
+			PatternStart = (GetLodIndex() > 0) ? GetClothCollection()->PatternEnd[GetLodIndex() - 1] + 1 : 0;
+		}
+
+		const int32 PatternIndex = PatternElementIndex - PatternStart;
+
+		GetPattern(PatternIndex).SetDefaults();
+
+		return PatternIndex;
+	}
+
+	FClothPatternAdapter FClothLodAdapter::GetPattern(int32 PatternIndex)
+	{
+		return FClothPatternAdapter(GetClothCollection(), GetLodIndex(), PatternIndex);
+	}
+
+	void FClothLodAdapter::Reset()
+	{
+		const int32 ElementIndex = GetElementIndex();
+
+		const int32 NumPatterns = GetNumPatterns();
+		for (int32 PatternIndex = 0; PatternIndex < NumPatterns; ++PatternIndex)
+		{
+			GetPattern(PatternIndex).Reset();
+		}
+		GetClothCollection()->RemoveElements(FClothCollection::PatternsGroup, NumPatterns, GetClothCollection()->PatternStart[ElementIndex]);
+
+		const int32 NumStitchings = GetNumStitchings();
+		for (int32 StitchingIndex = 0; StitchingIndex < NumStitchings; ++StitchingIndex)
+		{
+			//GetStitching(StitchingIndex).Reset();  // TODO
+		}
+		GetClothCollection()->RemoveElements(FClothCollection::StitchingsGroup, NumStitchings, GetClothCollection()->StitchingStart[ElementIndex]);
+
+		const int32 NumTetherBatches = GetNumTetherBatches();
+		for (int32 TetherBatchIndex = 0; TetherBatchIndex < NumTetherBatches; ++TetherBatchIndex)
+		{
+			//TetherBatch(TetherBatchIndex).Reset();  // TODO
+		}
+		GetClothCollection()->RemoveElements(FClothCollection::TetherBatchesGroup, NumTetherBatches, GetClothCollection()->TetherBatchStart[ElementIndex]);
+
+		SetDefaults();
+	}
+
+	void FClothLodAdapter::SetDefaults()
+	{
+		const int32 ElementIndex = GetElementIndex();
+
+		GetClothCollection()->PatternStart[ElementIndex] = INDEX_NONE;
+		GetClothCollection()->PatternEnd[ElementIndex] = INDEX_NONE;
+		GetClothCollection()->StitchingStart[ElementIndex] = INDEX_NONE;
+		GetClothCollection()->StitchingEnd[ElementIndex] = INDEX_NONE;
+		GetClothCollection()->TetherBatchStart[ElementIndex] = INDEX_NONE;
+		GetClothCollection()->TetherBatchEnd[ElementIndex] = INDEX_NONE;
+		GetClothCollection()->LodBiasDepth[ElementIndex] = 0;
+	}
+
+	template<typename T>
+	TArrayView<T> FClothLodAdapter::GetPatternsElements(TManagedArray<T>& ElementArray, const TManagedArray<int32>& StartArray, const TManagedArray<int32>& EndArray)
+	{
+		const TTuple<int32, int32> StartEnd = GetPatternsElementsStartEnd<true, true>(StartArray, EndArray);
+		const int32 Start = StartEnd.Get<0>();
+		const int32 End = StartEnd.Get<1>();
+
+		check(Start != INDEX_NONE || End == INDEX_NONE);  // Best to avoid situations where only one boundary of the range is set to INDEX_NONE
+		return Start == INDEX_NONE ? TArrayView<T>() : TArrayView<T>(ElementArray.GetData() + Start, End - Start + 1);
+	}
+
+	void FClothLodAdapter::Initialize(const TArray<FVector3f>& Positions, const TArray<uint32>& Indices)
+	{
+		using namespace UE::Chaos::ClothAsset::Private;
+
+		TArray<FIsland> Islands;
+		UnwrapMesh(Positions, Indices, Islands);  // Unwrap to 2D and reconstruct indices on 3D mesh
+
+		for (FIsland& Island : Islands)
+		{
+			if (Island.Indices.Num() && Island.Positions.Num() && Island.RestPositions.Num())
+			{
+				FClothPatternAdapter Pattern = AddGetPattern();
+				Pattern.Initialize(Island.Positions, Island.RestPositions, Island.Indices);
+			}
+		}
+	}
+}  // End namespace UE::Chaos::ClothAsset
diff --git a/Engine/Plugins/Experimental/ChaosClothAsset/Source/ChaosClothAsset/Private/ChaosClothAsset/ClothPatternAdapter.cpp b/Engine/Plugins/Experimental/ChaosClothAsset/Source/ChaosClothAsset/Private/ChaosClothAsset/ClothPatternAdapter.cpp
new file mode 100644
--- /dev/null
+++ b/Engine/Plugins/Experimental/ChaosClothAsset/Source/ChaosClothAsset/Private/ChaosClothAsset/ClothPatternAdapter.cpp
@@ -0,0 +1,272 @@
+// Copyright Epic Games, Inc. All Rights Reserved.
+
+#include "ChaosClothAsset/ClothPatternAdapter.h"
+#include "Modules/ModuleManager.h"
+
+namespace UE::Chaos::ClothAsset
+{
+	FClothPatternConstAdapter::FClothPatternConstAdapter(const TSharedPtr<const FClothCollection>& InClothCollection, int32 InLodIndex, int32 InPatternIndex)
+		: ClothCollection(InClothCollection)
+		, LodIndex(InLodIndex)
+		, PatternIndex(InPatternIndex)
+	{
+		check(ClothCollection.IsValid());
+		check(LodIndex >= 0 && LodIndex < ClothCollection->NumElements(FClothCollection::LodsGroup));
+		check(PatternIndex >= 0 && PatternIndex < ClothCollection->NumElements(FClothCollection::PatternsGroup));
+	}
+
+	TConstArrayView<float> FClothPatternConstAdapter::GetWeightMap(const FName& Name) const
+	{
+		const int32 ElementIndex = GetElementIndex();
+		const int32 SimVerticesStart = GetClothCollection()->SimVerticesStart[ElementIndex];
+		if (SimVerticesStart != INDEX_NONE)
+		{
+			const int32 SimVerticesEnd = GetClothCollection()->SimVerticesEnd[ElementIndex];
+			check(SimVerticesEnd != INDEX_NONE);
+
+			const TManagedArray<float>& ManagedArray = GetClothCollection()->GetAttribute<float>(Name, FClothCollection::SimVerticesGroup);
+
+			return TConstArrayView<float>(ManagedArray.GetData() + SimVerticesStart, SimVerticesEnd - SimVerticesStart + 1);
+		}
+
+		return TConstArrayView<float>();
+	}
+
+	int32 FClothPatternConstAdapter::GetNumElements(const TManagedArray<int32>& StartArray, const TManagedArray<int32>& EndArray) const
+	{
+		const int32 ElementIndex = GetElementIndex();
+		const int32 Start = StartArray[ElementIndex];
+		const int32 End = EndArray[ElementIndex];
+		check(Start != INDEX_NONE || End == INDEX_NONE);  // Best to avoid situations where only one boundary of the range is set to INDEX_NONE
+
+		return Start == INDEX_NONE ? 0 : End - Start + 1;
+	}
+
+	template<typename T>
+	TConstArrayView<T> FClothPatternConstAdapter::GetElements(const TManagedArray<T>& ElementArray, const TManagedArray<int32>& StartArray, const TManagedArray<int32>& EndArray) const
+	{
+		const int32 ElementIndex = GetElementIndex();
+		const int32 Start = StartArray[ElementIndex];
+		const int32 End = EndArray[ElementIndex];
+		check(Start != INDEX_NONE || End == INDEX_NONE);  // Best to avoid situations where only one boundary of the range is set to INDEX_NONE
+
+		return Start == INDEX_NONE ? TConstArrayView<T>() : TConstArrayView<T>(ElementArray.GetData() + Start, End - Start + 1);
+	}
+
+	FClothPatternAdapter::FClothPatternAdapter(const TSharedPtr<FClothCollection>& InClothCollection, int32 InLodIndex, int32 InPatternIndex)
+		: FClothPatternConstAdapter(InClothCollection, InLodIndex, InPatternIndex)
+	{
+	}
+
+	void FClothPatternAdapter::Reset()
+	{
+		TArray<int32> Array;
+		const TArrayView<int32>& ArrayView = static_cast<const TArrayView<int32>&>(TArrayView<int32>(Array));
+
+		const int32 ElementIndex = GetElementIndex();
+
+		GetClothCollection()->RemoveElements(FClothCollection::SimVerticesGroup, GetNumSimVertices(), GetClothCollection()->SimVerticesStart[ElementIndex]);
+		GetClothCollection()->RemoveElements(FClothCollection::SimFacesGroup, GetNumSimFaces(), GetClothCollection()->SimFacesStart[ElementIndex]);
+		GetClothCollection()->RemoveElements(FClothCollection::RenderVerticesGroup, GetNumRenderVertices(), GetClothCollection()->RenderVerticesStart[ElementIndex]);
+		GetClothCollection()->RemoveElements(FClothCollection::RenderFacesGroup, GetNumRenderFaces(), GetClothCollection()->WrapDeformerStart[ElementIndex]);
+
+		SetDefaults();
+	}
+
+	int32 FClothPatternAdapter::SetNumSimVertices(int32 InNumSimVertices)
+	{
+		return SetNumElements(
+			InNumSimVertices,
+			FClothCollection::SimVerticesGroup,
+			GetClothCollection()->SimVerticesStart,
+			GetClothCollection()->SimVerticesEnd);
+	}
+
+	int32 FClothPatternAdapter::SetNumSimFaces(int32 InNumSimFaces)
+	{
+		return SetNumElements(
+			InNumSimFaces,
+			FClothCollection::SimFacesGroup,
+			GetClothCollection()->SimFacesStart,
+			GetClothCollection()->SimFacesEnd);
+	}
+
+	int32 FClothPatternAdapter::SetNumRenderVertices(int32 InNumRenderVertices)
+	{
+		return SetNumElements(
+			InNumRenderVertices,
+			FClothCollection::RenderVerticesGroup,
+			GetClothCollection()->RenderVerticesStart,
+			GetClothCollection()->RenderVerticesEnd);
+	}
+
+	int32 FClothPatternAdapter::SetNumRenderFaces(int32 InNumRenderFaces)
+	{
+		return SetNumElements(
+			InNumRenderFaces,
+			FClothCollection::RenderFacesGroup,
+			GetClothCollection()->RenderFacesStart,
+			GetClothCollection()->RenderFacesEnd);
+	}
+
+	void FClothPatternAdapter::SetDefaults()
+	{
+		const int32 ElementIndex = GetElementIndex();
+
+		GetClothCollection()->SimVerticesStart[ElementIndex] = INDEX_NONE;
+		GetClothCollection()->SimVerticesEnd[ElementIndex] = INDEX_NONE;
+		GetClothCollection()->SimFacesStart[ElementIndex] = INDEX_NONE;
+		GetClothCollection()->SimFacesEnd[ElementIndex] = INDEX_NONE;
+		GetClothCollection()->RenderVerticesStart[ElementIndex] = INDEX_NONE;
+		GetClothCollection()->RenderVerticesEnd[ElementIndex] = INDEX_NONE;
+		GetClothCollection()->RenderFacesStart[ElementIndex] = INDEX_NONE;
+		GetClothCollection()->RenderFacesEnd[ElementIndex] = INDEX_NONE;
+		GetClothCollection()->NumWeights[ElementIndex] = 0;
+		GetClothCollection()->StatusFlags[ElementIndex] = 0;
+		GetClothCollection()->SimMaterialIndex[ElementIndex] = 0;
+	}
+
+	int32 FClothPatternAdapter::SetNumElements(int32 InNumElements, const FName& GroupName, TManagedArray<int32>& StartArray, TManagedArray<int32>& EndArray)
+	{
+		check(InNumElements >= 0);
+
+		const int32 ElementIndex = GetElementIndex();
+
+		int32& Start = StartArray[ElementIndex];
+		int32& End = EndArray[ElementIndex];
+		check(Start != INDEX_NONE || End == INDEX_NONE);  // Best to avoid situations where only one boundary of the range is set to INDEX_NONE
+
+		const int32 NumElements = (Start == INDEX_NONE) ? 0 : End - Start + 1;
+
+		if (const int32 Delta = InNumElements - NumElements)
+		{
+			if (Delta > 0)
+			{
+				// Find a previous valid index range to insert after when the range is empty
+				auto ComputeEnd = [&EndArray](int32 ElementIndex)->int32
+				{
+					for (int32 Index = ElementIndex; Index >= 0; --Index)
+					{
+						if (EndArray[Index] != INDEX_NONE)
+						{
+							return EndArray[Index];
+						}
+					}
+					return INDEX_NONE;
+				};
+
+				// Grow the array
+				const int32 Position = ComputeEnd(ElementIndex) + 1;
+				GetClothCollection()->InsertElements(Delta, Position, GroupName);
+
+				// Update Start/End
+				if (!NumElements)
+				{
+					Start = Position;
+				}
+				End = Start + InNumElements - 1;
+			}
+			else
+			{
+				// Shrink the array
+				const int32 Position = Start + InNumElements;
+				GetClothCollection()->RemoveElements(GroupName, -Delta, Position);
+
+				// Update Start/End
+				if (InNumElements)
+				{
+					End = Position - 1;
+				}
+				else
+				{
+					End = Start = INDEX_NONE;  // It is important to set the start & end to INDEX_NONE so that they never get automatically re-indexed by the managed array collection
+				}
+			}
+		}
+		return Start;
+	}
+
+	template<typename T>
+	TArrayView<T> FClothPatternAdapter::GetElements(TManagedArray<T>& ElementArray, const TManagedArray<int32>& StartArray, const TManagedArray<int32>& EndArray) const
+	{
+		const int32 ElementIndex = GetElementIndex();
+		const int32 Start = StartArray[ElementIndex];
+		const int32 End = EndArray[ElementIndex];
+		check(Start != INDEX_NONE || End == INDEX_NONE);  // Best to avoid situations where only one boundary of the range is set to INDEX_NONE
+
+		return Start == INDEX_NONE ? TArrayView<T>() : TArrayView<T>(ElementArray.GetData() + Start, End - Start + 1);
+	}
+
+	void FClothPatternAdapter::Initialize(const TArray<FVector2f>& Positions, const TArray<FVector3f>& RestPositions, const TArray<uint32>& Indices)
+	{
+		const int32 NumSimVertices = Positions.Num();
+		check(NumSimVertices == RestPositions.Num());
+
+		const int32 SimVerticesStart = SetNumSimVertices(NumSimVertices);
+
+		for (int32 Index = 0; Index < NumSimVertices; ++Index)
+		{
+			const int32 CollectionVertexIndex = SimVerticesStart + Index;
+			GetClothCollection()->SimPosition[CollectionVertexIndex] = Positions[Index];
+			GetClothCollection()->SimRestPosition[CollectionVertexIndex] = RestPositions[Index];
+			GetClothCollection()->SimRestNormal[CollectionVertexIndex] = FVector3f::ZeroVector;
+		}
+
+		const int32 NumSimFaces = (uint32)Indices.Num() / 3;
+		check(NumSimFaces * 3 == Indices.Num());
+
+		const int32 SimFacesStart = SetNumSimFaces(NumSimFaces);
+
+		// Face indices always index from the first vertex of the LOD, but these indices are indexed from the start of the pattern and need to be offset
+		const int32 LodSimVerticesStart = GetClothCollection()->SimVerticesStart[GetClothCollection()->PatternStart[GetLodIndex()]];
+		const int32 Offset = SimVerticesStart - LodSimVerticesStart;
+
+		for (int32 Index = 0; Index < NumSimFaces; ++Index)
+		{
+			// Indices from the start of the collection
+			const int32 CollectionVertexIndex0 = SimVerticesStart + Indices[Index * 3];
+			const int32 CollectionVertexIndex1 = SimVerticesStart + Indices[Index * 3 + 1];
+			const int32 CollectionVertexIndex2 = SimVerticesStart + Indices[Index * 3 + 2];
+
+			// Indices from the start of the LOD
+			const int32 LodVertexIndex0 = CollectionVertexIndex0 - LodSimVerticesStart;
+			const int32 LodVertexIndex1 = CollectionVertexIndex1 - LodSimVerticesStart;
+			const int32 LodVertexIndex2 = CollectionVertexIndex2 - LodSimVerticesStart;
+
+			GetClothCollection()->SimIndices[SimFacesStart + Index] = FIntVector3(LodVertexIndex0, LodVertexIndex1, LodVertexIndex2);
+
+			// Calculate face normal contribution
+			const FVector3f& Pos0 = GetClothCollection()->SimRestPosition[CollectionVertexIndex0];
+			const FVector3f& Pos1 = GetClothCollection()->SimRestPosition[CollectionVertexIndex1];
+			const FVector3f& Pos2 = GetClothCollection()->SimRestPosition[CollectionVertexIndex2];
+			const FVector3f Normal = (Pos1 - Pos0).Cross(Pos2 - Pos0).GetSafeNormal();
+			GetClothCollection()->SimRestNormal[CollectionVertexIndex0] += Normal;
+			GetClothCollection()->SimRestNormal[CollectionVertexIndex1] += Normal;
+			GetClothCollection()->SimRestNormal[CollectionVertexIndex2] += Normal;
+		}
+
+		// Normalize normals
+		for (int32 Index = 0; Index < NumSimVertices; ++Index)
+		{
+			const int32 CollectionVertexIndex = SimVerticesStart + Index;
+			GetClothCollection()->SimRestNormal[CollectionVertexIndex] =
+				GetClothCollection()->SimRestNormal[CollectionVertexIndex].GetSafeNormal(UE_SMALL_NUMBER, FVector3f::XAxisVector);
+		}
+	}
+
+	TArrayView<float> FClothPatternAdapter::GetWeightMap(const FName& Name)
+	{
+		const int32 ElementIndex = GetElementIndex();
+		const int32 SimVerticesStart = GetClothCollection()->SimVerticesStart[ElementIndex];
+		if (SimVerticesStart != INDEX_NONE)
+		{
+			const int32 SimVerticesEnd = GetClothCollection()->SimVerticesEnd[ElementIndex];
+			check(SimVerticesEnd != INDEX_NONE);
+
+			TManagedArray<float>& ManagedArray = GetClothCollection()->ModifyAttribute<float>(Name, FClothCollection::SimVerticesGroup);
+
+			return TArrayView<float>(ManagedArray.GetData() + SimVerticesStart, SimVerticesEnd - SimVerticesStart + 1);
+		}
+		return TArrayView<float>();
+	}
+}  // End namespace UE::Chaos::ClothAsset
diff --git a/Engine/Plugins/Experimental/ChaosClothAsset/Source/ChaosClothAsset/Private/ChaosClothAsset/ClothPresetCollection.cpp b/Engine/Plugins/Experimental/ChaosClothAsset/Source/ChaosClothAsset/Private/ChaosClothAsset/ClothPresetCollection.cpp
new file mode 100644
--- /dev/null
+++ b/Engine/Plugins/Experimental/ChaosClothAsset/Source/ChaosClothAsset/Private/ChaosClothAsset/ClothPresetCollection.cpp
@@ -0,0 +1,31 @@
+// Copyright Epic Games, Inc. All Rights Reserved.
+
+#include "ChaosClothAsset/ClothPresetCollection.h"
+#include "Chaos/ChaosArchive.h"
+
+namespace UE::Chaos::ClothAsset
+{
+	const FName FClothPresetCollection::PropertyGroup("Property");
+
+	FClothPresetCollection::FClothPresetCollection()
+	{
+		Construct();
+	}
+
+	void FClothPresetCollection::Serialize(FArchive& Ar)
+	{
+		::Chaos::FChaosArchive ChaosArchive(Ar);
+		Super::Serialize(ChaosArchive);
+	}
+
+	void FClothPresetCollection::Construct()
+	{
+		// Property Group
+		AddExternalAttribute<FString>("Name", PropertyGroup, Name);
+		AddExternalAttribute<FVector3f>("LowValue", PropertyGroup, LowValue);
+		AddExternalAttribute<FVector3f>("HighValue", PropertyGroup, HighValue);
+		AddExternalAttribute<FString>("StringValue", PropertyGroup, StringValue);
+		AddExternalAttribute<bool>("Enable", PropertyGroup, Enable);
+		AddExternalAttribute<bool>("Animatable", PropertyGroup, Animatable);
+	}
+}  // End namespace UE::Chaos::ClothAsset
diff --git a/Engine/Plugins/Experimental/ChaosClothAsset/Source/ChaosClothAsset/Public/ChaosClothAsset/ClothAdapter.h b/Engine/Plugins/Experimental/ChaosClothAsset/Source/ChaosClothAsset/Public/ChaosClothAsset/ClothAdapter.h
new file mode 100644
--- /dev/null
+++ b/Engine/Plugins/Experimental/ChaosClothAsset/Source/ChaosClothAsset/Public/ChaosClothAsset/ClothAdapter.h
@@ -0,0 +1,71 @@
+// Copyright Epic Games, Inc. All Rights Reserved.
+
+#pragma once
+
+#include "ChaosClothAsset/ClothCollection.h"
+#include "ChaosClothAsset/ClothLodAdapter.h"
+
+namespace UE::Chaos::ClothAsset
+{
+	/**
+	 * Const cloth adapter object to provide a more convenient object oriented access to the cloth collection.
+	 */
+	class CHAOSCLOTHASSET_API FClothConstAdapter
+	{
+	public:
+		FClothConstAdapter(const TSharedPtr<const FClothCollection>& InClothCollection);
+		FClothConstAdapter(const FClothLodConstAdapter& ClothLodConstAdapter);
+		FClothConstAdapter(const FClothPatternConstAdapter& ClothPatternConstAdapter);
+		virtual ~FClothConstAdapter() = default;
+
+		FClothConstAdapter(const FClothConstAdapter& Other) : ClothCollection(Other.ClothCollection) {}
+		FClothConstAdapter(FClothConstAdapter&& Other) : ClothCollection(MoveTemp(Other.ClothCollection)) {}
+		FClothConstAdapter& operator=(const FClothConstAdapter& Other) { ClothCollection = Other.ClothCollection; return *this; }
+		FClothConstAdapter& operator=(FClothConstAdapter&& Other) { ClothCollection = MoveTemp(Other.ClothCollection); return *this; }
+
+		/** Return the specified LOD. */
+		FClothLodConstAdapter GetLod(int32 LodIndex) const;
+
+		/** Return the number of LODs contained in this Cloth. */
+		int32 GetNumLods() const { return ClothCollection->NumElements(FClothCollection::LodsGroup); }
+
+		/** Return the underlaying cloth collection this adapter has been created with. */
+		const TSharedPtr<const FClothCollection>& GetClothCollection() const { return ClothCollection; }
+
+	private:
+		TSharedPtr<const FClothCollection> ClothCollection;
+	};
+
+	/**
+	 * Cloth adapter object to provide a more convenient object oriented access to the cloth collection.
+	 */
+	class CHAOSCLOTHASSET_API FClothAdapter : public FClothConstAdapter
+	{
+	public:
+		FClothAdapter(const TSharedPtr<FClothCollection>& InClothCollection);
+		virtual ~FClothAdapter() override = default;
+
+		FClothAdapter(const FClothAdapter& Other) : FClothConstAdapter(Other) {}
+		FClothAdapter(FClothAdapter&& Other) : FClothConstAdapter(MoveTemp(Other)) {}
+		FClothAdapter& operator=(const FClothAdapter& Other) { FClothConstAdapter::operator=(Other); return *this; }
+		FClothAdapter& operator=(FClothAdapter&& Other) { FClothConstAdapter::operator=(MoveTemp(Other)); return *this; }
+
+		/** Add a new LOD to this cloth. */
+		int32 AddLod();
+
+		/** Return the specified LOD. */
+		FClothLodAdapter GetLod(int32 LodIndex);
+
+		/** Add a new LOD to this cloth, and return the cloth LOD adapter set to its index. */
+		FClothLodAdapter AddGetLod() { return GetLod(AddLod()); }
+
+		/** Remove all LODs from this cloth. */
+		void Reset();
+
+		/** Add a new weight map to this cloth. Access is then done per pattern. */
+		void AddWeightMap(const FName& Name) { GetClothCollection()->AddAttribute<float>(Name, FClothCollection::SimVerticesGroup); }
+
+		/** Return the underlaying cloth collection this adapter has been created with. */
+		TSharedPtr<FClothCollection> GetClothCollection() { return ConstCastSharedPtr<FClothCollection>(FClothConstAdapter::GetClothCollection()); }
+	};
+}  // End namespace UE::Chaos::ClothAsset
diff --git a/Engine/Plugins/Experimental/ChaosClothAsset/Source/ChaosClothAsset/Public/ChaosClothAsset/ClothCollection.h b/Engine/Plugins/Experimental/ChaosClothAsset/Source/ChaosClothAsset/Public/ChaosClothAsset/ClothCollection.h
new file mode 100644
--- /dev/null
+++ b/Engine/Plugins/Experimental/ChaosClothAsset/Source/ChaosClothAsset/Public/ChaosClothAsset/ClothCollection.h
@@ -0,0 +1,127 @@
+// Copyright Epic Games, Inc. All Rights Reserved.
+
+#pragma once
+
+#include "CoreMinimal.h"
+#include "GeometryCollection/ManagedArrayCollection.h"
+#include "Templates/SharedPointer.h"
+
+// TODO:
+// - Transition Sim LOD data
+// - Move wrap deformer into struct similar to FMeshToMeshVertData
+// - Add FVertexBoneData to MANAGED_ARRAY_TYPE
+
+namespace UE::Chaos::ClothAsset
+{
+	// TODO: Do we want to move the skinning data away from the static render data?
+	///**
+	// * Skinning data per vertex
+	// */
+	//struct FVertexBoneData
+	//{
+	//	static const int8 MaxTotalInfluences = 12;  // Up to MAX_TOTAL_INFLUENCES bone indices that this vert is weighted to
+	//
+	//	int32 NumInfluences;
+	//	uint16 BoneIndices[MaxTotalInfluences];
+	//	float BoneWeights[MaxTotalInfluences];
+	//};
+
+	/**
+	 * Tailored Cloth Asset Collection containing draping and pattern information.
+	 */
+	class CHAOSCLOTHASSET_API FClothCollection : public FManagedArrayCollection
+	{
+	public:
+		typedef FManagedArrayCollection Super;
+
+		FClothCollection();
+		FClothCollection(FClothCollection&) = delete;
+		FClothCollection& operator=(const FClothCollection&) = delete;
+		FClothCollection(FClothCollection&&) = default;
+		FClothCollection& operator=(FClothCollection&&) = default;
+
+		using Super::Serialize;
+		void Serialize(FArchive& Ar);
+
+		// Attribute groups, predefined data member of the FClothLod object.
+		static const FName SimVerticesGroup;  // Contains patterns' 2D positions, 3D draped position (rest)
+		static const FName SimFacesGroup;  // Contains indices to sim vertex
+		static const FName RenderVerticesGroup;  // Contains pattern's 3D render model
+		static const FName RenderFacesGroup;  // Contains indices to render vertex
+		static const FName WrapDeformersGroup;  // Contains the wrap deformers cloth capture information
+		static const FName PatternsGroup;  // Contains pattern relationships to other groups
+		static const FName StitchingsGroup;  // Contains stitched sim vertex indices for quick welding or constraint creations
+		static const FName TethersGroup;  // Tethers information
+		static const FName TetherBatchesGroup;  // Tethers parallel batch processing information
+		static const FName LodsGroup;  // Lod split information
+
+		// Sim Vertices Group
+		TManagedArray<FVector2f> SimPosition;
+		TManagedArray<FVector3f> SimRestPosition;
+		TManagedArray<FVector3f> SimRestNormal;  // Used for capture, maxdistance, backstop authoring ...etc
+		//TManagedArray<FVertexBoneData> SimBoneData;
+
+		// Sim Faces Group
+		TManagedArray<FIntVector3> SimIndices;  // The indices point to the elements in the Sim Vertices arrays but don't include the LOD start offset
+
+		// Render Vertices Group
+		TManagedArray<FVector3f> RenderPosition;
+		TManagedArray<FVector3f> RenderNormal;
+		TManagedArray<FVector3f> RenderTangentU;
+		TManagedArray<FVector3f> RenderTangentV;
+		TManagedArray<TArray<FVector2f>> RenderUVs;
+		TManagedArray<FLinearColor> RenderColor;
+		//TManagedArray<FVertexBoneData> RenderBoneData;
+
+		// Render Faces Group
+		TManagedArray<FIntVector3> RenderIndices;  // The indices point to the elements in the Render Vertices arrays but don't include the LOD start offset
+		TManagedArray<int32> RenderMaterialIndex;  // Render material per triangle
+
+		// TODO: FMeshToMeshVertData
+		//// Wrap Deformers Group (render mesh captured, sim mesh capturing)
+		//TManagedArray<FVector4f> CapturedPosition;  // Barycentric coordinate and distance for captured render positions
+		//TManagedArray<FVector4f> CapturedNormal;  // Barycentric coordinate and distance for captured render normals
+		//TManagedArray<FVector4f> CapturedTangent;  // Barycentric coordinate and distance for captured render tangents
+		//TManagedArray<int32> CapturingFace;  // Sim triangle index
+		//TManagedArray<float> CapturingWeight;  // Weight for when using multiple capture influence per point, 1.0 otherwise
+
+		// Patterns Group
+		TManagedArray<int32> SimVerticesStart;
+		TManagedArray<int32> SimVerticesEnd;
+		TManagedArray<int32> SimFacesStart;
+		TManagedArray<int32> SimFacesEnd;
+		TManagedArray<int32> RenderVerticesStart;
+		TManagedArray<int32> RenderVerticesEnd;
+		TManagedArray<int32> RenderFacesStart;
+		TManagedArray<int32> RenderFacesEnd;
+		TManagedArray<int32> WrapDeformerStart;
+		TManagedArray<int32> WrapDeformerEnd;
+		TManagedArray<int32> NumWeights;  // Number of weights stored between WrapDeformerStart and WrapDeformerEnd
+		TManagedArray<int32> StatusFlags;  // Whether this pattern should be rendered, simulated, ...etc
+		TManagedArray<int32> SimMaterialIndex;  // Cloth material per pattern
+
+		// Stitching Group
+		TManagedArray<TSet<int32>> StitchedVertices;  // Stitched vertex indices
+
+		// Tethers Group
+		TManagedArray<int32> TetherKinematicIndex;
+		TManagedArray<int32> TetherDynamicIndex;
+		TManagedArray<float> TetherReferenceLength;
+
+		// Tether Batches Group
+		TManagedArray<int32> TetherStart;
+		TManagedArray<int32> TetherEnd;
+
+		// LOD Group
+		TManagedArray<int32> PatternStart;
+		TManagedArray<int32> PatternEnd;
+		TManagedArray<int32> StitchingStart;
+		TManagedArray<int32> StitchingEnd;
+		TManagedArray<int32> TetherBatchStart;
+		TManagedArray<int32> TetherBatchEnd;
+		TManagedArray<int32> LodBiasDepth;  // The number of LODs covered by each Sim LOD (for the wrap deformer)
+
+	protected:
+		void Construct();
+	};
+}  // End namespace UE::Chaos::ClothAsset
diff --git a/Engine/Plugins/Experimental/ChaosClothAsset/Source/ChaosClothAsset/Public/ChaosClothAsset/ClothGeometryTools.h b/Engine/Plugins/Experimental/ChaosClothAsset/Source/ChaosClothAsset/Public/ChaosClothAsset/ClothGeometryTools.h
new file mode 100644
--- /dev/null
+++ b/Engine/Plugins/Experimental/ChaosClothAsset/Source/ChaosClothAsset/Public/ChaosClothAsset/ClothGeometryTools.h
@@ -0,0 +1,17 @@
+// Copyright Epic Games, Inc. All Rights Reserved.
+
+#pragma once
+
+#include "ChaosClothAsset/ClothCollection.h"
+
+namespace UE::Chaos::ClothAsset
+{
+	/**
+	 * Geometry tools operating on cloth collections.
+	 */
+	struct CHAOSCLOTHASSET_API FClothGeometryTools
+	{
+		/** Turn the sim mesh portion of this ClothCollection into a render mesh. */
+		static void CopySimMeshToRenderMesh(const TSharedPtr<FClothCollection>& ClothCollection, int32 MaterialIndex);
+	};
+}  // End namespace UE::Chaos::ClothAsset
\ No newline at end of file
diff --git a/Engine/Plugins/Experimental/ChaosClothAsset/Source/ChaosClothAsset/Public/ChaosClothAsset/ClothLodAdapter.h b/Engine/Plugins/Experimental/ChaosClothAsset/Source/ChaosClothAsset/Public/ChaosClothAsset/ClothLodAdapter.h
new file mode 100644
--- /dev/null
+++ b/Engine/Plugins/Experimental/ChaosClothAsset/Source/ChaosClothAsset/Public/ChaosClothAsset/ClothLodAdapter.h
@@ -0,0 +1,144 @@
+// Copyright Epic Games, Inc. All Rights Reserved.
+
+#pragma once
+
+#include "ChaosClothAsset/ClothCollection.h"
+#include "ChaosClothAsset/ClothPatternAdapter.h"
+
+namespace UE::Chaos::ClothAsset
+{
+	/**
+	 * Cloth LOD adapter const object to provide a more convenient object oriented access to the cloth collection.
+	 */
+	class CHAOSCLOTHASSET_API FClothLodConstAdapter
+	{
+	public:
+		FClothLodConstAdapter(const TSharedPtr<const FClothCollection>& InClothCollection, int32 InLodIndex);
+		FClothLodConstAdapter(const FClothPatternConstAdapter& ClothPatternConstAdapter);
+
+		virtual ~FClothLodConstAdapter() = default;
+
+		FClothLodConstAdapter(const FClothLodConstAdapter& Other) : ClothCollection(Other.ClothCollection), LodIndex(Other.LodIndex) {}
+		FClothLodConstAdapter(FClothLodConstAdapter&& Other) : ClothCollection(MoveTemp(Other.ClothCollection)), LodIndex(Other.LodIndex) {}
+		FClothLodConstAdapter& operator=(const FClothLodConstAdapter& Other) { ClothCollection = Other.ClothCollection; LodIndex = Other.LodIndex; return *this; }
+		FClothLodConstAdapter& operator=(FClothLodConstAdapter&& Other) { ClothCollection = MoveTemp(Other.ClothCollection); LodIndex = Other.LodIndex; return *this; }
+
+		FClothPatternConstAdapter GetPattern(int32  PatternIndex) const;
+
+		int32 GetNumStitchings() const { return GetNumElements(ClothCollection->StitchingStart, ClothCollection->StitchingEnd); }
+		int32 GetNumTetherBatches() const { return GetNumElements(ClothCollection->TetherBatchStart, ClothCollection->TetherBatchEnd); }
+
+		// Patterns Group
+		int32 GetNumPatterns() const { return GetNumElements(ClothCollection->PatternStart, ClothCollection->PatternEnd); }
+		TConstArrayView<int32> GetSimVerticesStart() const { return GetElements(GetClothCollection()->SimVerticesStart, GetClothCollection()->PatternStart, GetClothCollection()->PatternEnd); }
+		TConstArrayView<int32> GetSimVerticesEnd() const { return GetElements(GetClothCollection()->SimVerticesEnd, GetClothCollection()->PatternStart, GetClothCollection()->PatternEnd); }
+		TConstArrayView<int32> GetSimFacesStart() const { return GetElements(GetClothCollection()->SimFacesStart, GetClothCollection()->PatternStart, GetClothCollection()->PatternEnd); }
+		TConstArrayView<int32> GetSimFacesEnd() const { return GetElements(GetClothCollection()->SimFacesEnd, GetClothCollection()->PatternStart, GetClothCollection()->PatternEnd); }
+		TConstArrayView<int32> GetRenderVerticesStart() const { return GetElements(GetClothCollection()->RenderVerticesStart, GetClothCollection()->PatternStart, GetClothCollection()->PatternEnd); }
+		TConstArrayView<int32> GetRenderVerticesEnd() const { return GetElements(GetClothCollection()->RenderVerticesEnd, GetClothCollection()->PatternStart, GetClothCollection()->PatternEnd); }
+		TConstArrayView<int32> GetRenderFacesStart() const { return GetElements(GetClothCollection()->RenderFacesStart, GetClothCollection()->PatternStart, GetClothCollection()->PatternEnd); }
+		TConstArrayView<int32> GetRenderFacesEnd() const { return GetElements(GetClothCollection()->RenderFacesEnd, GetClothCollection()->PatternStart, GetClothCollection()->PatternEnd); }
+
+		// Patterns Sim Vertices Group, use these to access all pattern at once, and when using the SimIndices
+		int32 GetPatternsSimVerticesStart() const { return GetPatternsElementsStartEnd<true, false>(GetClothCollection()->SimVerticesStart, GetClothCollection()->SimVerticesEnd).Get<0>(); }
+		int32 GetPatternsSimVerticesEnd() const { return GetPatternsElementsStartEnd<false, true>(GetClothCollection()->SimVerticesStart, GetClothCollection()->SimVerticesEnd).Get<1>(); }
+		TTuple<int32, int32> GetPatternsSimVerticesRange() const { return GetPatternsElementsStartEnd<true, true>(GetClothCollection()->SimVerticesStart, GetClothCollection()->SimVerticesEnd); }
+		int32 GetPatternsNumSimVertices() const { return GetPatternsNumElements(GetClothCollection()->SimVerticesStart, GetClothCollection()->SimVerticesEnd); }
+
+		TConstArrayView<FVector2f> GetPatternsSimPosition() const { return GetPatternsElements(GetClothCollection()->SimPosition, GetClothCollection()->SimVerticesStart, GetClothCollection()->SimVerticesEnd); }
+		TConstArrayView<FVector3f> GetPatternsSimRestPosition() const { return GetPatternsElements(GetClothCollection()->SimRestPosition, GetClothCollection()->SimVerticesStart, GetClothCollection()->SimVerticesEnd); }
+		TConstArrayView<FVector3f> GetPatternsSimRestNormal() const { return GetPatternsElements(GetClothCollection()->SimRestNormal, GetClothCollection()->SimVerticesStart, GetClothCollection()->SimVerticesEnd); }
+
+		// Patterns Render Vertices Group, use these to access all pattern at once, and when using the RenderIndices
+		int32 GetPatternsRenderVerticesStart() const { return GetPatternsElementsStartEnd<true, false>(GetClothCollection()->RenderVerticesStart, GetClothCollection()->RenderVerticesEnd).Get<0>(); }
+		int32 GetPatternsRenderVerticesEnd() const { return GetPatternsElementsStartEnd<false, true>(GetClothCollection()->RenderVerticesStart, GetClothCollection()->RenderVerticesEnd).Get<1>(); }
+		TTuple<int32, int32> GetPatternsRenderVerticesRange() const { return GetPatternsElementsStartEnd<true, true>(GetClothCollection()->RenderVerticesStart, GetClothCollection()->RenderVerticesEnd); }
+		int32 GetPatternsNumRenderVertices() const { return GetPatternsNumElements(GetClothCollection()->RenderVerticesStart, GetClothCollection()->RenderVerticesEnd); }
+
+		TConstArrayView<FVector3f> GetPatternsRenderPosition() const { return GetPatternsElements(GetClothCollection()->RenderPosition, GetClothCollection()->RenderVerticesStart, GetClothCollection()->RenderVerticesEnd); }
+		TConstArrayView<FVector3f> GetPatternsRenderNormal() const { return GetPatternsElements(GetClothCollection()->RenderNormal, GetClothCollection()->RenderVerticesStart, GetClothCollection()->RenderVerticesEnd); }
+		TConstArrayView<FVector3f> GetPatternsRenderTangentU() const { return GetPatternsElements(GetClothCollection()->RenderTangentU, GetClothCollection()->RenderVerticesStart, GetClothCollection()->RenderVerticesEnd); }
+		TConstArrayView<FVector3f> GetPatternsRenderTangentV() const { return GetPatternsElements(GetClothCollection()->RenderTangentV, GetClothCollection()->RenderVerticesStart, GetClothCollection()->RenderVerticesEnd); }
+		TConstArrayView<TArray<FVector2f>> GePatternstRenderUVs() const { return GetPatternsElements(GetClothCollection()->RenderUVs, GetClothCollection()->RenderVerticesStart, GetClothCollection()->RenderVerticesEnd); }
+		TConstArrayView<FLinearColor> GetPatternsRenderColor() const { return GetPatternsElements(GetClothCollection()->RenderColor, GetClothCollection()->RenderVerticesStart, GetClothCollection()->RenderVerticesEnd); }
+
+		/** Return the element index within the cloth collection. */
+		int32 GetElementIndex() const { return LodIndex; }
+
+		/** Return the LOD index this adapter has been created with. */
+		int32 GetLodIndex() const { return LodIndex; }
+
+		/** Return the underlaying cloth collection this adapter has been created with. */
+		const TSharedPtr<const FClothCollection>& GetClothCollection() const { return ClothCollection; }
+
+	protected:
+		template<bool bStart, bool bEnd>
+		TTuple<int32, int32> GetPatternsElementsStartEnd(const TManagedArray<int32>& StartArray, const TManagedArray<int32>& EndArray) const;
+
+	private:
+		int32 GetNumElements(const TManagedArray<int32>& StartArray, const TManagedArray<int32>& EndArray) const;
+		template<typename T>
+		TConstArrayView<T> GetElements(const TManagedArray<T>& ElementArray, const TManagedArray<int32>& StartArray, const TManagedArray<int32>& EndArray) const;
+
+		int32 GetPatternsNumElements(const TManagedArray<int32>& StartArray, const TManagedArray<int32>& EndArray) const;
+		template<typename T>
+		TConstArrayView<T> GetPatternsElements(const TManagedArray<T>& ElementArray, const TManagedArray<int32>& StartArray, const TManagedArray<int32>& EndArray) const;
+
+		TSharedPtr<const FClothCollection> ClothCollection;
+		int32 LodIndex;
+	};
+
+	/**
+	 * Cloth LOD adapter object to provide a more convenient object oriented access to the cloth collection.
+	 */
+	class CHAOSCLOTHASSET_API FClothLodAdapter : public FClothLodConstAdapter
+	{
+	public:
+		FClothLodAdapter(const TSharedPtr<FClothCollection>& InClothCollection, int32 InLodIndex);
+		virtual ~FClothLodAdapter() override = default;
+
+		FClothLodAdapter(const FClothLodAdapter& Other) : FClothLodConstAdapter(Other) {}
+		FClothLodAdapter(FClothLodAdapter&& Other) : FClothLodConstAdapter(MoveTemp(Other)) {}
+		FClothLodAdapter& operator=(const FClothLodAdapter& Other) { FClothLodConstAdapter::operator=(Other); return *this; }
+		FClothLodAdapter& operator=(FClothLodAdapter&& Other) { FClothLodConstAdapter::operator=(MoveTemp(Other)); return *this; }
+
+		/** Add a new pattern to this cloth LOD. */
+		int32 AddPattern();
+
+		/** Return the specified pattern. */
+		FClothPatternAdapter GetPattern(int32 PatternIndex);
+
+		/** Add a new pattern to this cloth LOD, and return the cloth pattern adapter set to its index. */
+		FClothPatternAdapter AddGetPattern() { return GetPattern(AddPattern()); }
+
+		/** Remove all cloth patterns from this cloth LOD. */
+		void Reset();
+
+		/** Initialize the cloth LOD using the specified 3D triangle mesh and unwrapping it into a 2D geometry for generating the patterns. */
+		void Initialize(const TArray<FVector3f>& Positions, const TArray<uint32>& Indices);
+
+		/** Return the underlaying Cloth Collection this adapter has been created with. */
+		TSharedPtr<FClothCollection> GetClothCollection() { return ConstCastSharedPtr<FClothCollection>(FClothLodConstAdapter::GetClothCollection()); }
+
+		// Patterns Sim Vertices Group, use these to access all pattern at once, and when using the SimIndices
+		TArrayView<FVector2f> GetPatternsSimPosition() { return GetPatternsElements(GetClothCollection()->SimPosition, GetClothCollection()->SimVerticesStart, GetClothCollection()->SimVerticesEnd); }
+		TArrayView<FVector3f> GetPatternsSimRestPosition() { return GetPatternsElements(GetClothCollection()->SimRestPosition, GetClothCollection()->SimVerticesStart, GetClothCollection()->SimVerticesEnd); }
+		TArrayView<FVector3f> GetPatternsSimRestNormal() { return GetPatternsElements(GetClothCollection()->SimRestNormal, GetClothCollection()->SimVerticesStart, GetClothCollection()->SimVerticesEnd); }
+
+		// Patterns Render Vertices Group, use these to access all pattern at once, and when using the RenderIndices
+		TArrayView<FVector3f> GetPatternsRenderPosition() { return GetPatternsElements(GetClothCollection()->RenderPosition, GetClothCollection()->RenderVerticesStart, GetClothCollection()->RenderVerticesEnd); }
+		TArrayView<FVector3f> GetPatternsRenderNormal() { return GetPatternsElements(GetClothCollection()->RenderNormal, GetClothCollection()->RenderVerticesStart, GetClothCollection()->RenderVerticesEnd); }
+		TArrayView<FVector3f> GetPatternsRenderTangentU() { return GetPatternsElements(GetClothCollection()->RenderTangentU, GetClothCollection()->RenderVerticesStart, GetClothCollection()->RenderVerticesEnd); }
+		TArrayView<FVector3f> GetPatternsRenderTangentV() { return GetPatternsElements(GetClothCollection()->RenderTangentV, GetClothCollection()->RenderVerticesStart, GetClothCollection()->RenderVerticesEnd); }
+		TArrayView<TArray<FVector2f>> GePatternstRenderUVs() { return GetPatternsElements(GetClothCollection()->RenderUVs, GetClothCollection()->RenderVerticesStart, GetClothCollection()->RenderVerticesEnd); }
+		TArrayView<FLinearColor> GetPatternsRenderColor() { return GetPatternsElements(GetClothCollection()->RenderColor, GetClothCollection()->RenderVerticesStart, GetClothCollection()->RenderVerticesEnd); }
+
+	private:
+		friend class FClothAdapter;
+
+		void SetDefaults();
+
+		template<typename T>
+		TArrayView<T> GetPatternsElements(TManagedArray<T>& ElementArray, const TManagedArray<int32>& StartArray, const TManagedArray<int32>& EndArray);
+	};
+}  // End namespace UE::Chaos::ClothAsset
diff --git a/Engine/Plugins/Experimental/ChaosClothAsset/Source/ChaosClothAsset/Public/ChaosClothAsset/ClothPatternAdapter.h b/Engine/Plugins/Experimental/ChaosClothAsset/Source/ChaosClothAsset/Public/ChaosClothAsset/ClothPatternAdapter.h
new file mode 100644
--- /dev/null
+++ b/Engine/Plugins/Experimental/ChaosClothAsset/Source/ChaosClothAsset/Public/ChaosClothAsset/ClothPatternAdapter.h
@@ -0,0 +1,155 @@
+// Copyright Epic Games, Inc. All Rights Reserved.
+
+#pragma once
+
+#include "ChaosClothAsset/ClothCollection.h"
+
+namespace UE::Chaos::ClothAsset
+{
+	/**
+	 * Cloth pattern const adapter object to provide a more convenient object oriented access to the cloth collection.
+	 */
+	class CHAOSCLOTHASSET_API FClothPatternConstAdapter
+	{
+	public:
+		FClothPatternConstAdapter(const TSharedPtr<const FClothCollection>& InClothCollection, int32 InLodIndex, int32 InPatternIndex);
+		virtual ~FClothPatternConstAdapter() = default;
+
+		FClothPatternConstAdapter(const FClothPatternConstAdapter& Other) : ClothCollection(Other.ClothCollection), LodIndex(Other.LodIndex), PatternIndex(Other.PatternIndex) {}
+		FClothPatternConstAdapter(FClothPatternConstAdapter&& Other) : ClothCollection(MoveTemp(Other.ClothCollection)), LodIndex(Other.LodIndex), PatternIndex(Other.PatternIndex) {}
+		FClothPatternConstAdapter& operator=(const FClothPatternConstAdapter& Other) { ClothCollection = Other.ClothCollection; LodIndex = Other.LodIndex; PatternIndex = Other.PatternIndex; return *this; }
+		FClothPatternConstAdapter& operator=(FClothPatternConstAdapter&& Other) { ClothCollection = MoveTemp(Other.ClothCollection); LodIndex = Other.LodIndex; PatternIndex = Other.PatternIndex; return *this; }
+
+		// Sim Vertices Group
+		// Note: Use the FClothLodConstAdapter accessors instead of these for the array indices to match the SimIndices values
+		int32 GetNumSimVertices() const { return GetNumElements(GetClothCollection()->SimVerticesStart, GetClothCollection()->SimVerticesEnd); }
+		TConstArrayView<FVector2f> GetSimPosition() const { return GetElements(GetClothCollection()->SimPosition, GetClothCollection()->SimVerticesStart, GetClothCollection()->SimVerticesEnd); }
+		TConstArrayView<FVector3f> GetSimRestPosition() const { return GetElements(GetClothCollection()->SimRestPosition, GetClothCollection()->SimVerticesStart, GetClothCollection()->SimVerticesEnd); }
+		TConstArrayView<FVector3f> GetSimRestNormal() const { return GetElements(GetClothCollection()->SimRestNormal, GetClothCollection()->SimVerticesStart, GetClothCollection()->SimVerticesEnd); }
+
+		// Sim Faces Group, note: SimIndices points to the LOD arrays, not the pattern arrays
+		int32 GetNumSimFaces() const { return GetNumElements(GetClothCollection()->SimFacesStart, GetClothCollection()->SimFacesEnd); }
+		TConstArrayView<FIntVector3> GetSimIndices() const { return GetElements(GetClothCollection()->SimIndices, GetClothCollection()->SimFacesStart, GetClothCollection()->SimFacesEnd); }
+
+		// Render Vertices Group
+		// Note: Use the FClothLodConstAdapter accessors instead of these for the array indices to match the RenderIndices values
+		int32 GetNumRenderVertices() const { return GetNumElements(GetClothCollection()->RenderVerticesStart, GetClothCollection()->RenderVerticesEnd); }
+		TConstArrayView<FVector3f> GetRenderPosition() const { return GetElements(GetClothCollection()->RenderPosition, GetClothCollection()->RenderVerticesStart, GetClothCollection()->RenderVerticesEnd); }
+		TConstArrayView<FVector3f> GetRenderNormal() const { return GetElements(GetClothCollection()->RenderNormal, GetClothCollection()->RenderVerticesStart, GetClothCollection()->RenderVerticesEnd); }
+		TConstArrayView<FVector3f> GetRenderTangentU() const { return GetElements(GetClothCollection()->RenderTangentU, GetClothCollection()->RenderVerticesStart, GetClothCollection()->RenderVerticesEnd); }
+		TConstArrayView<FVector3f> GetRenderTangentV() const { return GetElements(GetClothCollection()->RenderTangentV, GetClothCollection()->RenderVerticesStart, GetClothCollection()->RenderVerticesEnd); }
+		TConstArrayView<TArray<FVector2f>> GetRenderUVs() const { return GetElements(GetClothCollection()->RenderUVs, GetClothCollection()->RenderVerticesStart, GetClothCollection()->RenderVerticesEnd); }
+		TConstArrayView<FLinearColor> GetRenderColor() const { return GetElements(GetClothCollection()->RenderColor, GetClothCollection()->RenderVerticesStart, GetClothCollection()->RenderVerticesEnd); }
+
+		// Render Faces Group, note: RenderIndices points to the LOD arrays, not the pattern arrays
+		int32 GetNumRenderFaces() const { return GetNumElements(GetClothCollection()->RenderFacesStart, GetClothCollection()->RenderFacesEnd); }
+		TConstArrayView<FIntVector3> GetRenderIndices() const { return GetElements(GetClothCollection()->RenderIndices, GetClothCollection()->RenderFacesStart, GetClothCollection()->RenderFacesEnd); }
+		TConstArrayView<int32> GetRenderMaterialIndex() const { return GetElements(GetClothCollection()->RenderMaterialIndex, GetClothCollection()->RenderFacesStart, GetClothCollection()->RenderFacesEnd); }
+
+		// Wrap Deformers Group
+		int32 GetNumWrapDeformers() const { return GetNumElements(GetClothCollection()->WrapDeformerStart, GetClothCollection()->WrapDeformerEnd); }
+
+		/** Return the element index within the cloth collection. */
+		int32 GetElementIndex() const { return ClothCollection->PatternStart[LodIndex] + PatternIndex; }
+
+		/** Return the LOD index this adapter has been created with. */
+		int32 GetLodIndex() const { return LodIndex; }
+
+		/** Return the Pattern index this adapter has been created with. */
+		int32 GetPatternIndex() const { return PatternIndex; }
+
+		/**
+		 * Return a const view on the specified vertex weightmap.
+		 * The attribute must exist otherwise the function will assert.
+		 * @see FClothAdapter::AddWeightMap, GetNumSimVertices.
+		 */
+		TConstArrayView<float> GetWeightMap(const FName& Name) const;
+
+		/** Return the underlaying cloth collection this adapter has been created with. */
+		const TSharedPtr<const FClothCollection>& GetClothCollection() const { return ClothCollection; }
+
+	private:
+		int32 GetNumElements(const TManagedArray<int32>& StartArray, const TManagedArray<int32>& EndArray) const;
+		template<typename T>
+		TConstArrayView<T> GetElements(const TManagedArray<T>& ElementArray, const TManagedArray<int32>& StartArray, const TManagedArray<int32>& EndArray) const;
+
+		TSharedPtr<const FClothCollection> ClothCollection;
+		int32 LodIndex;
+		int32 PatternIndex;
+	};
+
+	/**
+	 * Cloth pattern adapter object to provide a more convenient object oriented access to the cloth collection.
+	 */
+	class CHAOSCLOTHASSET_API FClothPatternAdapter: public FClothPatternConstAdapter
+	{
+	public:
+		FClothPatternAdapter(const TSharedPtr<FClothCollection>& InClothCollection, int32 InLodIndex, int32 InPatternIndex);
+		virtual ~FClothPatternAdapter() override = default;
+
+		FClothPatternAdapter(const FClothPatternAdapter& Other) : FClothPatternConstAdapter(Other) {}
+		FClothPatternAdapter(FClothPatternAdapter&& Other) : FClothPatternConstAdapter(MoveTemp(Other)) {}
+		FClothPatternAdapter& operator=(const FClothPatternAdapter& Other) { FClothPatternConstAdapter::operator=(Other); return *this; }
+		FClothPatternAdapter& operator=(FClothPatternAdapter&& Other) { FClothPatternConstAdapter::operator=(MoveTemp(Other)); return *this; }
+
+		/** Remove all geometry from this cloth Pattern. */
+		void Reset();
+
+		/** Grow or shrink the space reserved for simulation vertices for this pattern within the cloth collection and return its start index. */
+		int32 SetNumSimVertices(int32 NumSimVertices);
+
+		/** Grow or shrink the space reserved for simulation faces for this pattern within the cloth collection and return its start index. */
+		int32 SetNumSimFaces(int32 InNumSimFaces);
+
+		/** Grow or shrink the space reserved for render vertices for this pattern within the cloth collection and return its start index. */
+		int32 SetNumRenderVertices(int32 NumRenderVertices);
+
+		/** Grow or shrink the space reserved for render faces for this pattern within the cloth collection and return its start index. */
+		int32 SetNumRenderFaces(int32 InNumRenderFaces);
+
+		/** Initialize the cloth pattern using the specified 3D and 2D positions, and topology. */
+		void Initialize(const TArray<FVector2f>& Positions, const TArray<FVector3f>& RestPositions, const TArray<uint32>& Indices);
+
+		// Sim Vertices Group
+		// Note: Use the FClothLodAdapter accessors instead of these for the array indices to match the SimIndices values
+		TArrayView<FVector2f> GetSimPosition() { return GetElements(GetClothCollection()->SimPosition, GetClothCollection()->SimVerticesStart, GetClothCollection()->SimVerticesEnd); }
+		TArrayView<FVector3f> GetSimRestPosition() { return GetElements(GetClothCollection()->SimRestPosition, GetClothCollection()->SimVerticesStart, GetClothCollection()->SimVerticesEnd); }
+		TArrayView<FVector3f> GetSimRestNormal() { return GetElements(GetClothCollection()->SimRestNormal, GetClothCollection()->SimVerticesStart, GetClothCollection()->SimVerticesEnd); }
+
+		// Sim Faces Group, note: SimIndices points to the LOD arrays, not the pattern arrays
+		TArrayView<FIntVector3> GetSimIndices() { return GetElements(GetClothCollection()->SimIndices, GetClothCollection()->SimFacesStart, GetClothCollection()->SimFacesEnd); }
+
+		// Render Vertices Group
+		// Note: Use the FClothLodAdapter accessors instead of these for the array indices to match the RenderIndices values
+		TArrayView<FVector3f> GetRenderPosition() { return GetElements(GetClothCollection()->RenderPosition, GetClothCollection()->RenderVerticesStart, GetClothCollection()->RenderVerticesEnd); }
+		TArrayView<FVector3f> GetRenderNormal() { return GetElements(GetClothCollection()->RenderNormal, GetClothCollection()->RenderVerticesStart, GetClothCollection()->RenderVerticesEnd); }
+		TArrayView<FVector3f> GetRenderTangentU() { return GetElements(GetClothCollection()->RenderTangentU, GetClothCollection()->RenderVerticesStart, GetClothCollection()->RenderVerticesEnd); }
+		TArrayView<FVector3f> GetRenderTangentV() { return GetElements(GetClothCollection()->RenderTangentV, GetClothCollection()->RenderVerticesStart, GetClothCollection()->RenderVerticesEnd); }
+		TArrayView<TArray<FVector2f>> GetRenderUVs() { return GetElements(GetClothCollection()->RenderUVs, GetClothCollection()->RenderVerticesStart, GetClothCollection()->RenderVerticesEnd); }
+		TArrayView<FLinearColor> GetRenderColor() { return GetElements(GetClothCollection()->RenderColor, GetClothCollection()->RenderVerticesStart, GetClothCollection()->RenderVerticesEnd); }
+
+		// Render Faces Group, note: RenderIndices points to the LOD arrays, not the pattern arrays
+		TArrayView<FIntVector3> GetRenderIndices() { return GetElements(GetClothCollection()->RenderIndices, GetClothCollection()->RenderFacesStart, GetClothCollection()->RenderFacesEnd); }
+		TArrayView<int32> GetRenderMaterialIndex() { return GetElements(GetClothCollection()->RenderMaterialIndex, GetClothCollection()->RenderFacesStart, GetClothCollection()->RenderFacesEnd); }
+
+		/**
+		 * Return a view on the specified vertex weightmap.
+		 * The attribute must exist otherwise the function will assert.
+		 * @see FClothAdapter::AddWeightMap, GetNumSimVertices.
+		 */
+		TArrayView<float> GetWeightMap(const FName& Name);
+
+		/** Return the underlaying Cloth Collection this adapter has been created with. */
+		TSharedPtr<FClothCollection> GetClothCollection() { return ConstCastSharedPtr<FClothCollection>(FClothPatternConstAdapter::GetClothCollection()); }
+
+	private:
+		friend class FClothLodAdapter;
+
+		void SetDefaults();
+
+		int32 SetNumElements(int32 InNumElements, const FName& GroupName, TManagedArray<int32>& StartArray, TManagedArray<int32>& EndArray);
+
+		template<typename T>
+		TArrayView<T> GetElements(TManagedArray<T>& ElementArray, const TManagedArray<int32>& StartArray, const TManagedArray<int32>& EndArray) const;
+	};
+}  // End namespace UE::Chaos::ClothAsset
diff --git a/Engine/Plugins/Experimental/ChaosClothAsset/Source/ChaosClothAsset/Public/ChaosClothAsset/ClothPresetCollection.h b/Engine/Plugins/Experimental/ChaosClothAsset/Source/ChaosClothAsset/Public/ChaosClothAsset/ClothPresetCollection.h
new file mode 100644
--- /dev/null
+++ b/Engine/Plugins/Experimental/ChaosClothAsset/Source/ChaosClothAsset/Public/ChaosClothAsset/ClothPresetCollection.h
@@ -0,0 +1,69 @@
+// Copyright Epic Games, Inc. All Rights Reserved.
+
+#pragma once
+
+#include "CoreMinimal.h"
+#include "GeometryCollection/ManagedArrayCollection.h"
+
+namespace UE::Chaos::ClothAsset
+{
+	/**
+	 * Tailored Cloth Preset Collection containing draping and pattern information.
+	 */
+	class CHAOSCLOTHASSET_API FClothPresetCollection : public FManagedArrayCollection
+	{
+	public:
+		typedef FManagedArrayCollection Super;
+
+		FClothPresetCollection();
+		FClothPresetCollection(FClothPresetCollection&) = delete;
+		FClothPresetCollection& operator=(const FClothPresetCollection&) = delete;
+		FClothPresetCollection(FClothPresetCollection&&) = default;
+		FClothPresetCollection& operator=(FClothPresetCollection&&) = default;
+
+		using Super::Serialize;
+		void Serialize(FArchive& Ar);
+
+		// Attribute groups, predefined data member of the FClothLod object.
+		static const FName PropertyGroup;  // Fabric properties
+
+		// PropertyGroup
+		TManagedArray<FString> Name;
+		TManagedArray<FVector3f> LowValue;
+		TManagedArray<FVector3f> HighValue;
+		TManagedArray<FString> StringValue;
+		TManagedArray<bool> Enable;
+		TManagedArray<bool> Animatable;
+
+	protected:
+		void Construct();
+
+		// TODO: Add predefined sim properties:
+		//	int32 MassMode;
+		//	float MassValue;
+		//	FVector2f EdgeStiffness;
+		//	FVector2f BendingStiffness;
+		//	FVector2f AreaStiffness;
+		//	bool bUseGeodesicTether;
+		//	FVector2f TetherStiffness;
+		//	FVector2f TetherScale;
+		//	FVector2f MaxDistance;
+		//	FVector2f BackstopDistance;
+		//	FVector2f BackstopRadius;
+		//	FVector2f CollisionThickness;
+		//	FVector2f FrictionCoefficient;
+		//	bool bUseCCD;
+		//	float DampingCoefficient;
+		//	float LocalDampingCoefficient;
+		//	FVector3f Drag;
+		//	FVector3f Lift;
+		//	float GravityScale;
+		//	bool bGravityOverride;
+		//	FVector3f Gravity;
+		//	FVector2f AnimDriveStiffness;
+		//	FVector2f AnimDriveDamping;
+		//	float LinearVelocityScale;
+		//	float AngularVelocityScale;
+		//	float FictitiousAngularScale;
+	};
+}  // End namespace UE::Chaos::ClothAsset
diff --git a/Engine/Plugins/Experimental/ChaosClothAsset/Source/ChaosClothAssetEngine/Private/ChaosClothAsset/ClothAsset.cpp b/Engine/Plugins/Experimental/ChaosClothAsset/Source/ChaosClothAssetEngine/Private/ChaosClothAsset/ClothAsset.cpp
new file mode 100644
--- /dev/null
+++ b/Engine/Plugins/Experimental/ChaosClothAsset/Source/ChaosClothAssetEngine/Private/ChaosClothAsset/ClothAsset.cpp
@@ -0,0 +1,480 @@
+// Copyright Epic Games, Inc. All Rights Reserved.
+
+#include "ChaosClothAsset/ClothAsset.h"
+#include "ChaosClothAsset/ClothAssetBuilder.h"
+#include "ChaosClothAsset/ClothAdapter.h"
+#include "ChaosClothAsset/ClothGeometryTools.h"
+#include "Animation/Skeleton.h"
+#include "Rendering/SkeletalMeshModel.h"
+#include "EngineUtils.h"
+#if WITH_EDITOR
+#include "IMeshBuilderModule.h"
+#include "DerivedDataCacheInterface.h"
+#endif
+
+#include UE_INLINE_GENERATED_CPP_BY_NAME(ClothAsset)
+
+// If Chaos cloth asset derived data needs to be rebuilt (new format, serialization differences, etc.) replace the version GUID below with a new one. 
+// In case of merge conflicts with DDC versions, you MUST generate a new GUID and set this new GUID as the version.
+#define CHAOS_CLOTH_ASSET_DERIVED_DATA_VERSION TEXT("19823D996CA54F279B9A9FA8ED7A8EB6")
+
+DEFINE_LOG_CATEGORY_STATIC(LogChaosClothAsset, Log, All);
+
+UChaosClothAsset::UChaosClothAsset(const FObjectInitializer& ObjectInitializer)
+	: Super(ObjectInitializer)
+	, DisableBelowMinLodStripping(FPerPlatformBool(false))
+	, ClothCollection(MakeShared<UE::Chaos::ClothAsset::FClothCollection>())
+#if WITH_EDITORONLY_DATA
+	, MeshModel(MakeShareable(new FSkeletalMeshModel()))
+#endif
+{
+}
+
+FMatrix UChaosClothAsset::GetComposedRefPoseMatrix(FName InBoneName) const
+{
+	FMatrix LocalPose(FMatrix::Identity);
+
+	if (InBoneName != NAME_None)
+	{
+		const int32 BoneIndex = GetRefSkeleton().FindBoneIndex(InBoneName);
+		if (BoneIndex != INDEX_NONE)
+		{
+			return GetComposedRefPoseMatrix(BoneIndex);
+		}
+		// TODO: Might need to add sockets like on the SkeletalMesh
+	}
+
+	return LocalPose;
+}
+
+void UChaosClothAsset::Serialize(FArchive& Ar)
+{
+	Super::Serialize(Ar);
+
+	bool bCooked = Ar.IsCooking();
+	Ar << bCooked;
+
+	ClothCollection->Serialize(Ar);
+
+	Ar << GetRefSkeleton();
+	if (Ar.IsLoading())
+	{
+		const bool bRebuildNameMap = false;
+		GetRefSkeleton().RebuildRefSkeleton(GetSkeleton(), bRebuildNameMap);
+	}
+
+	if (bCooked && !IsTemplate() && !Ar.IsCountingMemory())
+	{
+		if (Ar.IsLoading())
+		{
+			SetResourceForRendering(MakeUnique<FSkeletalMeshRenderData>());
+		}
+		GetResourceForRendering()->Serialize(Ar, this);
+	}
+}
+
+void UChaosClothAsset::PostLoad()
+{
+	Super::PostLoad();
+}
+
+void UChaosClothAsset::BeginPostLoadInternal(FSkinnedAssetPostLoadContext& Context)
+{
+#if WITH_EDITOR
+	TRACE_CPUPROFILER_EVENT_SCOPE(UChaosClothAsset::BeginPostLoadInternal);
+
+	checkf(IsInGameThread(), TEXT("Cannot execute function UChaosClothAsset::BeginPostLoadInternal asynchronously. Asset: %s"), *this->GetFullName());
+	SetInternalFlags(EInternalObjectFlags::Async);
+
+	// Lock all properties that should not be modified/accessed during async post-load
+	AcquireAsyncProperty();
+
+	// This scope allows us to use any locked properties without causing stalls
+	FSkinnedAssetAsyncBuildScope AsyncBuildScope(this);
+
+	BuildModel();
+#endif // #if WITH_EDITOR
+}
+
+void UChaosClothAsset::ExecutePostLoadInternal(FSkinnedAssetPostLoadContext& Context)
+{
+#if WITH_EDITOR
+	TRACE_CPUPROFILER_EVENT_SCOPE(UChaosClothAsset::ExecutePostLoadInternal);
+
+	// This scope allows us to use any locked properties without causing stalls
+	FSkinnedAssetAsyncBuildScope AsyncBuildScope(this);
+
+	if (!GetOutermost()->bIsCookedForEditor)
+	{
+		if (GetResourceForRendering() == nullptr)
+		{
+			CacheDerivedData(&Context);
+			Context.bHasCachedDerivedData = true;
+		}
+	}
+#endif // WITH_EDITOR
+}
+
+void UChaosClothAsset::FinishPostLoadInternal(FSkinnedAssetPostLoadContext& Context)
+{
+#if WITH_EDITOR
+	TRACE_CPUPROFILER_EVENT_SCOPE(UChaosClothAsset::FinishPostLoadInternal);
+
+	checkf(IsInGameThread(), TEXT("Cannot execute function UChaosClothAsset::FinishPostLoadInternal asynchronously. Asset: %s"), *this->GetFullName());
+	ClearInternalFlags(EInternalObjectFlags::Async);
+
+	// This scope allows us to use any locked properties without causing stalls
+	FSkinnedAssetAsyncBuildScope AsyncBuildScope(this);
+#endif
+
+	if (FApp::CanEverRender())
+	{
+		InitResources();
+	}
+
+	CalculateInvRefMatrices();
+	CalculateBounds();
+
+#if WITH_EDITOR
+	ReleaseAsyncProperty();
+#endif
+}
+
+void UChaosClothAsset::BeginDestroy()
+{
+	check(IsInGameThread());
+
+	Super::BeginDestroy();
+
+	// Release the mesh's render resources now
+	ReleaseResources();
+}
+
+bool UChaosClothAsset::IsReadyForFinishDestroy()
+{
+	if (!Super::IsReadyForFinishDestroy())
+	{
+		return false;
+	}
+
+	ReleaseResources();
+
+	// see if we have hit the resource flush fence
+	return ReleaseResourcesFence.IsFenceComplete();
+}
+
+void UChaosClothAsset::InitResources()
+{
+	LLM_SCOPE_BYNAME(TEXT("ClothAsset/InitResources"));
+
+	if (SkeletalMeshRenderData.IsValid())
+	{
+		TArray<UMorphTarget*> DummyMorphTargets;  // Not even used by InitResources() at the moment
+		SkeletalMeshRenderData->InitResources(false, DummyMorphTargets, this);
+	}
+}
+
+void UChaosClothAsset::ReleaseResources()
+{
+	if (SkeletalMeshRenderData && SkeletalMeshRenderData->IsInitialized())
+	{
+		if (GIsEditor && !GIsPlayInEditorWorld)
+		{
+			// Flush the rendering command to be sure there is no command left that can create/modify a rendering resource
+			FlushRenderingCommands();
+		}
+
+		SkeletalMeshRenderData->ReleaseResources();
+
+		// Insert a fence to signal when these commands completed
+		ReleaseResourcesFence.BeginFence();
+	}
+}
+
+void UChaosClothAsset::CalculateInvRefMatrices()
+{
+	auto GetRefPoseMatrix = [this](int32 BoneIndex)->FMatrix
+	{
+		check(BoneIndex >= 0 && BoneIndex < RefSkeleton.GetRawBoneNum());
+		FTransform BoneTransform = RefSkeleton.GetRawRefBonePose()[BoneIndex];
+		BoneTransform.NormalizeRotation();  // Make sure quaternion is normalized!
+		return BoneTransform.ToMatrixWithScale();
+	};
+
+	const int32 NumRealBones = RefSkeleton.GetRawBoneNum();
+
+	RefBasesInvMatrix.Empty(NumRealBones);
+	RefBasesInvMatrix.AddUninitialized(NumRealBones);
+
+	// Reset cached mesh-space ref pose
+	TArray<FMatrix> ComposedRefPoseMatrices;
+	ComposedRefPoseMatrices.SetNumUninitialized(NumRealBones);
+
+	// Precompute the Mesh.RefBasesInverse
+	for (int32 BoneIndex = 0; BoneIndex < NumRealBones; ++BoneIndex)
+	{
+		// Render the default pose
+		ComposedRefPoseMatrices[BoneIndex] = GetRefPoseMatrix(BoneIndex);
+
+		// Construct mesh-space skeletal hierarchy
+		if (BoneIndex > 0)
+		{
+			int32 Parent = RefSkeleton.GetRawParentIndex(BoneIndex);
+			ComposedRefPoseMatrices[BoneIndex] = ComposedRefPoseMatrices[BoneIndex] * ComposedRefPoseMatrices[Parent];
+		}
+
+		FVector XAxis, YAxis, ZAxis;
+		ComposedRefPoseMatrices[BoneIndex].GetScaledAxes(XAxis, YAxis, ZAxis);
+		if (XAxis.IsNearlyZero(UE_SMALL_NUMBER) &&
+			YAxis.IsNearlyZero(UE_SMALL_NUMBER) &&
+			ZAxis.IsNearlyZero(UE_SMALL_NUMBER))
+		{
+			// This is not allowed, warn them
+			UE_LOG(
+				LogChaosClothAsset,
+				Warning,
+				TEXT("Reference Pose for asset %s for joint (%s) includes NIL matrix. Zero scale isn't allowed on ref pose."),
+				*GetPathName(),
+				*RefSkeleton.GetBoneName(BoneIndex).ToString());
+		}
+
+		// Precompute inverse so we can use from-refpose-skin vertices
+		RefBasesInvMatrix[BoneIndex] = FMatrix44f(ComposedRefPoseMatrices[BoneIndex].Inverse());
+	}
+}
+
+void UChaosClothAsset::CalculateBounds()
+{
+	using namespace UE::Chaos::ClothAsset;
+
+	FBox BoundingBox(ForceInit);
+
+	const FClothConstAdapter Cloth(ClothCollection);
+
+	for (int32 LodIndex = 0; LodIndex < Cloth.GetNumLods(); ++LodIndex)
+	{
+		const FClothLodConstAdapter ClothLod = Cloth.GetLod(LodIndex);
+
+		for (int32 ClothPatternIndex = 0; ClothPatternIndex < ClothLod.GetNumPatterns(); ++ClothPatternIndex)
+		{
+			const FClothPatternConstAdapter ClothPattern = ClothLod.GetPattern(ClothPatternIndex);
+			const int32 PatternElementIndex = ClothPattern.GetElementIndex();
+
+			const int32 RenderVerticesStart = ClothCollection->RenderVerticesStart[PatternElementIndex];
+			const int32 RenderVerticesEnd = ClothCollection->RenderVerticesEnd[PatternElementIndex];
+
+			for (int32 RenderVertexIndex = RenderVerticesStart; RenderVertexIndex <= RenderVerticesEnd; ++RenderVertexIndex)
+			{
+				BoundingBox += (FVector)ClothCollection->RenderPosition[RenderVertexIndex];
+			}
+		}
+	}
+
+	Bounds = FBoxSphereBounds(BoundingBox);
+}
+
+void UChaosClothAsset::Build()
+{
+	using namespace UE::Chaos::ClothAsset;
+
+	// Release render resources
+	ReleaseResources();
+
+	// Set a new Guid to invalidate the DDC
+	AssetGuid = FGuid::NewGuid();
+
+	// Rebuild Skeleton
+	RefSkeleton = Skeleton->GetReferenceSkeleton();  // TODO: Does a copy work? Try to see what's exactly needed here...
+
+	constexpr bool bRebuildNameMap = false;
+	RefSkeleton.RebuildRefSkeleton(Skeleton, bRebuildNameMap);
+
+	CalculateInvRefMatrices();
+	CalculateBounds();
+
+	// Add LODs to the render data
+	const FClothConstAdapter Cloth(ClothCollection);
+	const int32 NumLods = Cloth.GetNumLods();
+
+	// Rebuild LOD Infos
+	LODInfo.Reset(NumLods);
+	LODInfo.AddDefaulted(NumLods);  // TODO: Expose some properties to fill up the LOD infos
+
+	// Rebuild LOD Model
+#if WITH_EDITORONLY_DATA
+	BuildModel();
+#endif
+
+	// Load/save render data from/to DDC
+#if WITH_EDITOR
+	CacheDerivedData(nullptr);
+#endif
+
+	if (FApp::CanEverRender())
+	{
+		InitResources();
+	}
+}
+
+#if WITH_EDITORONLY_DATA
+void UChaosClothAsset::BuildModel()
+{
+	const TArray<IClothAssetBuilderClassProvider*> ClassProviders = IModularFeatures::Get().GetModularFeatureImplementations<IClothAssetBuilderClassProvider>(IClothAssetBuilderClassProvider::FeatureName);
+	if (const TSubclassOf<UClothAssetBuilder> ClothAssetBuilderClass = ClassProviders.Num() ? ClassProviders[0]->GetClothAssetBuilderClass() : nullptr)
+	{
+		if (const UClothAssetBuilder* const ClothAssetBuilder = ClothAssetBuilderClass->GetDefaultObject<UClothAssetBuilder>())
+		{
+			using namespace UE::Chaos::ClothAsset;
+
+			const FClothConstAdapter Cloth(GetClothCollection());
+			const int32 NumLods = Cloth.GetNumLods();
+
+			check(MeshModel);  // MeshModel should always be created in the Cloth Asset constructor WITH_EDITORONLY_DATA
+			MeshModel->LODModels.Empty();
+
+			for (int32 LodIndex = 0; LodIndex < NumLods; ++LodIndex)
+			{
+				FSkeletalMeshLODModel* LODModel = new FSkeletalMeshLODModel();
+				ClothAssetBuilder->BuildLod(*LODModel, *this, LodIndex);
+				MeshModel->LODModels.Add(LODModel);
+			}
+		}
+	}
+}
+#endif  // #if WITH_EDITORONLY_DATA
+
+const FMeshUVChannelInfo* UChaosClothAsset::GetUVChannelData(int32 MaterialIndex) const
+{
+	if (GetMaterials().IsValidIndex(MaterialIndex))
+	{
+		// TODO: enable ensure when UVChannelData is setup
+		//ensure(GetMaterials()[MaterialIndex].UVChannelData.bInitialized);
+		return &GetMaterials()[MaterialIndex].UVChannelData;
+	}
+
+	return nullptr;
+}
+
+FSkeletalMeshRenderData* UChaosClothAsset::GetResourceForRendering() const
+{
+	WaitUntilAsyncPropertyReleased(EClothAssetAsyncProperties::RenderData);
+	return SkeletalMeshRenderData.Get();
+}
+
+void UChaosClothAsset::SetResourceForRendering(TUniquePtr<FSkeletalMeshRenderData>&& InSkeletalMeshRenderData)
+{
+	WaitUntilAsyncPropertyReleased(EClothAssetAsyncProperties::RenderData);
+	SkeletalMeshRenderData = MoveTemp(InSkeletalMeshRenderData);
+}
+
+void UChaosClothAsset::WaitUntilAsyncPropertyReleased(EClothAssetAsyncProperties AsyncProperties, ESkinnedAssetAsyncPropertyLockType LockType) const
+{
+	// Cast strongly typed enum to uint64
+	WaitUntilAsyncPropertyReleasedInternal((uint64)AsyncProperties, LockType);
+}
+
+FString UChaosClothAsset::GetAsyncPropertyName(uint64 Property) const
+{
+	return StaticEnum<EClothAssetAsyncProperties>()->GetNameByValue(Property).ToString();
+}
+
+#if WITH_EDITOR
+void UChaosClothAsset::CacheDerivedData(FSkinnedAssetPostLoadContext* Context)
+{
+	TRACE_CPUPROFILER_EVENT_SCOPE(UChaosClothAsset::CacheDerivedData);
+
+	if (!GetOutermost()->bIsCookedForEditor)
+	{
+		if (GetResourceForRendering() == nullptr)
+		{
+			// Cache derived data for the running platform.
+			ITargetPlatform* RunningPlatform = GetTargetPlatformManagerRef().GetRunningTargetPlatform();
+			check(RunningPlatform);
+
+			// Create the render data
+			SetResourceForRendering(MakeUnique<FSkeletalMeshRenderData>());
+			// Load render data from DDC, or generate it and save to DDC
+			GetResourceForRendering()->Cache(RunningPlatform, this, Context);
+		}
+	}
+}
+
+FString UChaosClothAsset::BuildDerivedDataKey(const ITargetPlatform* TargetPlatform)
+{
+	FString KeySuffix(TEXT(""));
+	KeySuffix += AssetGuid.ToString();
+
+	FString TmpPartialKeySuffix;
+	//Synchronize the user data that are part of the key
+	GetImportedModel()->SyncronizeLODUserSectionsData();
+
+	// Model GUID is not generated so exclude GetImportedModel()->GetIdString() from DDC key.
+
+	// Add the hashed string generated from the model data 
+	TmpPartialKeySuffix = GetImportedModel()->GetLODModelIdString();
+	KeySuffix += TmpPartialKeySuffix;
+
+	//Add the max gpu bone per section
+	const int32 MaxGPUSkinBones = FGPUBaseSkinVertexFactory::GetMaxGPUSkinBones(TargetPlatform);
+	KeySuffix += FString::FromInt(MaxGPUSkinBones);
+	// Add unlimited bone influences mode
+	IMeshBuilderModule::GetForPlatform(TargetPlatform).AppendToDDCKey(KeySuffix, true);
+	const bool bUnlimitedBoneInfluences = FGPUBaseSkinVertexFactory::GetUnlimitedBoneInfluences();
+	KeySuffix += bUnlimitedBoneInfluences ? "1" : "0";
+
+	// Add LODInfoArray
+	TmpPartialKeySuffix = TEXT("");
+	TArray<FSkeletalMeshLODInfo>& LODInfos = GetLODInfoArray();
+	for (int32 LODIndex = 0; LODIndex < GetLODNum(); ++LODIndex)
+	{
+		check(LODInfos.IsValidIndex(LODIndex));
+		FSkeletalMeshLODInfo& LOD = LODInfos[LODIndex];
+		LOD.BuildGUID = LOD.ComputeDeriveDataCacheKey(nullptr);	// TODO: FSkeletalMeshLODGroupSettings
+		TmpPartialKeySuffix += LOD.BuildGUID.ToString(EGuidFormats::Digits);
+	}
+	KeySuffix += TmpPartialKeySuffix;
+
+	//KeySuffix += GetHasVertexColors() ? "1" : "0";
+	//KeySuffix += GetVertexColorGuid().ToString(EGuidFormats::Digits);
+
+	//if (GetEnableLODStreaming(TargetPlatform))
+	//{
+	//	const int32 MaxNumStreamedLODs = GetMaxNumStreamedLODs(TargetPlatform);
+	//	const int32 MaxNumOptionalLODs = GetMaxNumOptionalLODs(TargetPlatform);
+	//	KeySuffix += *FString::Printf(TEXT("1%08x%08x"), MaxNumStreamedLODs, MaxNumOptionalLODs);
+	//}
+	//else
+	//{
+	//	KeySuffix += TEXT("0zzzzzzzzzzzzzzzz");
+	//}
+
+	//if (TargetPlatform->GetPlatformInfo().PlatformGroupName == TEXT("Desktop")
+	//	&& GStripSkeletalMeshLodsDuringCooking != 0
+	//	&& GSkeletalMeshKeepMobileMinLODSettingOnDesktop != 0)
+	//{
+	//	KeySuffix += TEXT("_MinMLOD");
+	//}
+
+	return FDerivedDataCacheInterface::BuildCacheKey(
+		TEXT("CHAOSCLOTH"),
+		CHAOS_CLOTH_ASSET_DERIVED_DATA_VERSION,
+		*KeySuffix
+	);
+}
+
+bool UChaosClothAsset::IsInitialBuildDone() const
+{
+	//We are consider built if we have a valid lod model
+	return GetImportedModel() != nullptr &&
+		GetImportedModel()->LODModels.Num() > 0 &&
+		GetImportedModel()->LODModels[0].Sections.Num() > 0;
+}
+#endif // WITH_EDITOR
+
+void UChaosClothAsset::CopySimMeshToRenderMesh(int32 MaterialIndex)
+{
+	using namespace UE::Chaos::ClothAsset;
+
+	check(ClothCollection.IsValid());
+	FClothGeometryTools::CopySimMeshToRenderMesh(ClothCollection, MaterialIndex);
+}
