Chaos Cloth - Remove the Cloth Asset plugin IsBeta setting so that the IsExperimental setting can be applied instead.

diff --git a/Engine/Plugins/Experimental/ChaosClothAsset/ChaosClothAsset.uplugin b/Engine/Plugins/Experimental/ChaosClothAsset/ChaosClothAsset.uplugin
--- a/Engine/Plugins/Experimental/ChaosClothAsset/ChaosClothAsset.uplugin
+++ b/Engine/Plugins/Experimental/ChaosClothAsset/ChaosClothAsset.uplugin
@@ -12,7 +12,6 @@
 	"SupportURL": "",
 	"EnabledByDefault": false,
 	"CanContainContent": true,
-	"IsBetaVersion" : true,
 	"IsExperimentalVersion": true,
 	"Installed": false,
 	"Modules": [
