ChaosCloth: Fix private dependency header included from a public folder

diff --git a/Engine/Plugins/Experimental/ChaosClothAsset/Source/ChaosClothAssetEngine/Private/ChaosClothAsset/ClothPreset.cpp b/Engine/Plugins/Experimental/ChaosClothAsset/Source/ChaosClothAssetEngine/Private/ChaosClothAsset/ClothPreset.cpp
--- a/Engine/Plugins/Experimental/ChaosClothAsset/Source/ChaosClothAssetEngine/Private/ChaosClothAsset/ClothPreset.cpp
+++ b/Engine/Plugins/Experimental/ChaosClothAsset/Source/ChaosClothAssetEngine/Private/ChaosClothAsset/ClothPreset.cpp
@@ -1,6 +1,7 @@
 // Copyright Epic Games, Inc. All Rights Reserved.
 
 #include "ChaosClothAsset/ClothPreset.h"
+#include "ChaosClothAsset/ClothPresetCollection.h"
 
 #include UE_INLINE_GENERATED_CPP_BY_NAME(ClothPreset)
 
diff --git a/Engine/Plugins/Experimental/ChaosClothAsset/Source/ChaosClothAssetEngine/Public/ChaosClothAsset/ClothAsset.h b/Engine/Plugins/Experimental/ChaosClothAsset/Source/ChaosClothAssetEngine/Public/ChaosClothAsset/ClothAsset.h
--- a/Engine/Plugins/Experimental/ChaosClothAsset/Source/ChaosClothAssetEngine/Public/ChaosClothAsset/ClothAsset.h
+++ b/Engine/Plugins/Experimental/ChaosClothAsset/Source/ChaosClothAssetEngine/Public/ChaosClothAsset/ClothAsset.h
@@ -3,7 +3,6 @@
 #pragma once
 
 #include "CoreMinimal.h"
-#include "ChaosClothAsset/ClothCollection.h"
 #include "ChaosClothAsset/ClothPreset.h"
 #include "Engine/SkeletalMesh.h"  // For FSkeletalMeshLODInfo
 #include "Engine/SkinnedAsset.h"
@@ -11,6 +10,7 @@
 #include "RenderCommandFence.h"
 #include "ClothAsset.generated.h"
 
+namespace UE::Chaos::ClothAsset { class FClothCollection; }
 class FSkeletalMeshRenderData;
 class FSkeletalMeshModel;
 
@@ -122,7 +122,7 @@ private:
 	virtual FString GetAsyncPropertyName(uint64 Property) const override;
 	//~ End USkinnedAsset interface
 
-	/** 
+	/**
 	 * Wait for the asset to finish compilation to protect internal skinned asset data from race conditions during async build.
 	 * This should be called before accessing all async accessible properties.
 	 */
diff --git a/Engine/Plugins/Experimental/ChaosClothAsset/Source/ChaosClothAssetEngine/Public/ChaosClothAsset/ClothPreset.h b/Engine/Plugins/Experimental/ChaosClothAsset/Source/ChaosClothAssetEngine/Public/ChaosClothAsset/ClothPreset.h
--- a/Engine/Plugins/Experimental/ChaosClothAsset/Source/ChaosClothAssetEngine/Public/ChaosClothAsset/ClothPreset.h
+++ b/Engine/Plugins/Experimental/ChaosClothAsset/Source/ChaosClothAssetEngine/Public/ChaosClothAsset/ClothPreset.h
@@ -3,10 +3,11 @@
 #pragma once
 
 #include "CoreMinimal.h"
-#include "ChaosClothAsset/ClothPresetCollection.h"
 
 #include "ClothPreset.generated.h"
 
+namespace UE::Chaos::ClothAsset { class FClothPresetCollection; }
+
 /**
  * Preset property categories.
  */
