Chaos Cloth:
- Add support for Seams in the cloth collection.
- Improve Cloth Adapter API with some of the missing getters from the cloth collection data.
- Move the generic SetNumElements/GetNumElements and GetElements functions to the cloth collection so that they can be used in all adapters.

diff --git a/Engine/Plugins/Experimental/ChaosClothAsset/Source/ChaosClothAsset/Private/ChaosClothAsset/ClothCollection.cpp b/Engine/Plugins/Experimental/ChaosClothAsset/Source/ChaosClothAsset/Private/ChaosClothAsset/ClothCollection.cpp
--- a/Engine/Plugins/Experimental/ChaosClothAsset/Source/ChaosClothAsset/Private/ChaosClothAsset/ClothCollection.cpp
+++ b/Engine/Plugins/Experimental/ChaosClothAsset/Source/ChaosClothAsset/Private/ChaosClothAsset/ClothCollection.cpp
@@ -10,7 +10,7 @@ namespace UE::Chaos::ClothAsset
 	const FName UE::Chaos::ClothAsset::FClothCollection::RenderFacesGroup("RenderFaces");
 	const FName UE::Chaos::ClothAsset::FClothCollection::WrapDeformersGroup("WrapDeformers");
 	const FName UE::Chaos::ClothAsset::FClothCollection::PatternsGroup("Patterns");
-	const FName UE::Chaos::ClothAsset::FClothCollection::StitchingsGroup("Stitchings");
+	const FName UE::Chaos::ClothAsset::FClothCollection::SeamsGroup("Seams");
 	const FName UE::Chaos::ClothAsset::FClothCollection::TethersGroup("Tethers");
 	const FName UE::Chaos::ClothAsset::FClothCollection::TetherBatchesGroup("TetherBatches");
 	const FName UE::Chaos::ClothAsset::FClothCollection::LodsGroup("Lods");
@@ -35,7 +35,7 @@ namespace UE::Chaos::ClothAsset
 		FManagedArrayCollection::FConstructionParameters RenderFacesDependency(RenderFacesGroup);
 		FManagedArrayCollection::FConstructionParameters WrapDeformersDependency(WrapDeformersGroup);
 		FManagedArrayCollection::FConstructionParameters PatternsDependency(PatternsGroup);
-		FManagedArrayCollection::FConstructionParameters StitchingsDependency(StitchingsGroup);
+		FManagedArrayCollection::FConstructionParameters StitchingsDependency(SeamsGroup);
 		FManagedArrayCollection::FConstructionParameters TethersDependency(TethersGroup);
 		FManagedArrayCollection::FConstructionParameters TetherBatchesDependency(TetherBatchesGroup);
 
@@ -78,8 +78,9 @@ namespace UE::Chaos::ClothAsset
 		AddExternalAttribute<int32>("StatusFlags", PatternsGroup, StatusFlags);
 		AddExternalAttribute<int32>("SimMaterialIndex", PatternsGroup, SimMaterialIndex);
 
-		// Stitching Group
-		AddExternalAttribute<TSet<int32>>("StitchedVertices", StitchingsGroup, StitchedVertices, SimVerticesDependency);
+		// Seams Group
+		AddExternalAttribute<FIntVector2>("SeamPatterns", SeamsGroup, SeamPatterns, SimVerticesDependency);
+		AddExternalAttribute<TArray<FIntVector2>>("SeamStitches", SeamsGroup, SeamStitches, SimVerticesDependency);
 
 		// Tethers Group
 		AddExternalAttribute<int32>("TetherKinematicIndex", TethersGroup, TetherKinematicIndex, SimVerticesDependency);
@@ -93,10 +94,87 @@ namespace UE::Chaos::ClothAsset
 		// LOD Group
 		AddExternalAttribute<int32>("PatternStart", LodsGroup, PatternStart, PatternsDependency);
 		AddExternalAttribute<int32>("PatternEnd", LodsGroup, PatternEnd, PatternsDependency);
-		AddExternalAttribute<int32>("StitchingStart", LodsGroup, StitchingStart, StitchingsDependency);
-		AddExternalAttribute<int32>("StitchingEnd", LodsGroup, StitchingEnd, StitchingsDependency);
+		AddExternalAttribute<int32>("StitchingStart", LodsGroup, SeamStart, StitchingsDependency);
+		AddExternalAttribute<int32>("StitchingEnd", LodsGroup, SeamEnd, StitchingsDependency);
 		AddExternalAttribute<int32>("TetherBatchStart", LodsGroup, TetherBatchStart, TetherBatchesDependency);
 		AddExternalAttribute<int32>("TetherBatchEnd", LodsGroup, TetherBatchEnd, TetherBatchesDependency);
 		AddExternalAttribute<int32>("LodBiasDepth", LodsGroup, LodBiasDepth);
 	}
+
+	int32 FClothCollection::SetNumElements(int32 InNumElements, const FName& GroupName, TManagedArray<int32>& StartArray, TManagedArray<int32>& EndArray, int32 StartEndIndex)
+	{
+		check(InNumElements >= 0);
+
+		int32& Start = StartArray[StartEndIndex];
+		int32& End = EndArray[StartEndIndex];
+		check(Start != INDEX_NONE || End == INDEX_NONE);  // Best to avoid situations where only one boundary of the range is set to INDEX_NONE
+
+		const int32 NumElements = (Start == INDEX_NONE) ? 0 : End - Start + 1;
+
+		if (const int32 Delta = InNumElements - NumElements)
+		{
+			if (Delta > 0)
+			{
+				// Find a previous valid index range to insert after when the range is empty
+				auto ComputeEnd = [&EndArray](int32 StartEndIndex)->int32
+				{
+					for (int32 Index = StartEndIndex; Index >= 0; --Index)
+					{
+						if (EndArray[Index] != INDEX_NONE)
+						{
+							return EndArray[Index];
+						}
+					}
+					return INDEX_NONE;
+				};
+
+				// Grow the array
+				const int32 Position = ComputeEnd(StartEndIndex) + 1;
+				InsertElements(Delta, Position, GroupName);
+
+				// Update Start/End
+				if (!NumElements)
+				{
+					Start = Position;
+				}
+				End = Start + InNumElements - 1;
+			}
+			else
+			{
+				// Shrink the array
+				const int32 Position = Start + InNumElements;
+				RemoveElements(GroupName, -Delta, Position);
+
+				// Update Start/End
+				if (InNumElements)
+				{
+					End = Position - 1;
+				}
+				else
+				{
+					End = Start = INDEX_NONE;  // It is important to set the start & end to INDEX_NONE so that they never get automatically re-indexed by the managed array collection
+				}
+			}
+		}
+		return Start;
+	}
+
+	int32 FClothCollection::GetNumElements(const TManagedArray<int32>& StartArray, const TManagedArray<int32>& EndArray, int32 StartEndIndex) const
+	{
+		const int32 Start = StartArray[StartEndIndex];
+		const int32 End = EndArray[StartEndIndex];
+		check(Start != INDEX_NONE || End == INDEX_NONE);  // Best to avoid situations where only one boundary of the range is set to INDEX_NONE
+
+		return Start == INDEX_NONE ? 0 : End - Start + 1;
+	}
+
+	int32 FClothCollection::GetPatternsNumElements(const TManagedArray<int32>& StartArray, const TManagedArray<int32>& EndArray, int32 LodIndex) const
+	{
+		const TTuple<int32, int32> StartEnd = GetPatternsElementsStartEnd(StartArray, EndArray, LodIndex);
+		const int32 Start = StartEnd.Get<0>();
+		const int32 End = StartEnd.Get<1>();
+		check(Start != INDEX_NONE || End == INDEX_NONE);  // Best to avoid situations where only one boundary of the range is set to INDEX_NONE
+
+		return Start == INDEX_NONE ? 0 : End - Start + 1;
+	}
 } // End namespace UE::Chaos::ClothAsset
diff --git a/Engine/Plugins/Experimental/ChaosClothAsset/Source/ChaosClothAsset/Private/ChaosClothAsset/ClothLodAdapter.cpp b/Engine/Plugins/Experimental/ChaosClothAsset/Source/ChaosClothAsset/Private/ChaosClothAsset/ClothLodAdapter.cpp
--- a/Engine/Plugins/Experimental/ChaosClothAsset/Source/ChaosClothAsset/Private/ChaosClothAsset/ClothLodAdapter.cpp
+++ b/Engine/Plugins/Experimental/ChaosClothAsset/Source/ChaosClothAsset/Private/ChaosClothAsset/ClothLodAdapter.cpp
@@ -5,11 +5,13 @@
 // Utility functions to unwrap a 3d sim mesh into a tailored cloth
 namespace UE::Chaos::ClothAsset::Private
 {
+	// Triangle islands to become patterns, although in this case all the seams are internal (same pattern)
 	struct FIsland
 	{
 		TArray<uint32> Indices;  // 3x number of triangles
 		TArray<FVector2f> Positions;
 		TArray<FVector3f> RestPositions;  // Same size as Positions
+		TArray<uint32> SourceIndices;  // Index in the original welded position array
 	};
 
 	enum class EIntersectCirclesResult
@@ -59,42 +61,53 @@ namespace UE::Chaos::ClothAsset::Private
 
 		return EIntersectCirclesResult::DoubleIntersect;
 	}
-
-	static void UnwrapMesh(const TArray<FVector3f>& Positions, const TArray<uint32>& Indices, TArray<FIsland>& OutIslands)
+;
+	static FUintVector2 MakeSortedUintVector2(uint32 Index0, uint32 Index1)
 	{
-		OutIslands.Reset();
+		return Index0 < Index1 ? FUintVector2(Index0, Index1) : FUintVector2(Index1, Index0);
+	}
 
+	static void BuildEdgeMap(const TArray<uint32>& Indices, TMap<FUintVector2, TArray<uint32>>& OutEdgeToTriangles)
+	{
 		ensure(Indices.Num() % 3 == 0);
 		const uint32 NumTriangles = (uint32)Indices.Num() / 3;
+		OutEdgeToTriangles.Empty(NumTriangles * 2);  // Rough estimate for the number of edges
 		if (!NumTriangles)
 		{
 			return;
 		}
 
-		// Gather edge information
-		TMap<FUintVector2, TArray<uint32>> EdgeToTriangles;
-		EdgeToTriangles.Reserve(NumTriangles * 2);  // Rough estimate for the number of edges
-
-		auto MakeSortedEdge = [](uint32 Index0, uint32 Index1)->FUintVector2
-		{
-			return Index0 < Index1 ? FUintVector2(Index0, Index1) : FUintVector2(Index1, Index0);
-		};
-
 		for (uint32 Triangle = 0; Triangle < NumTriangles; ++Triangle)
 		{
 			const uint32 Index0 = Indices[Triangle * 3 + 0];
 			const uint32 Index1 = Indices[Triangle * 3 + 1];
 			const uint32 Index2 = Indices[Triangle * 3 + 2];
 
-			const FUintVector2 Edge0 = MakeSortedEdge(Index0, Index1);
-			const FUintVector2 Edge1 = MakeSortedEdge(Index1, Index2);
-			const FUintVector2 Edge2 = MakeSortedEdge(Index2, Index0);
+			const FUintVector2 Edge0 = MakeSortedUintVector2(Index0, Index1);
+			const FUintVector2 Edge1 = MakeSortedUintVector2(Index1, Index2);
+			const FUintVector2 Edge2 = MakeSortedUintVector2(Index2, Index0);
+
+			OutEdgeToTriangles.FindOrAdd(Edge0).Add(Triangle);
+			OutEdgeToTriangles.FindOrAdd(Edge1).Add(Triangle);
+			OutEdgeToTriangles.FindOrAdd(Edge2).Add(Triangle);
+		}
+	}
+
+	static void UnwrapMesh(const TArray<FVector3f>& Positions, const TArray<uint32>& Indices, TArray<FIsland>& OutIslands)
+	{
+		OutIslands.Reset();
 
-			EdgeToTriangles.FindOrAdd(Edge0).Add(Triangle);
-			EdgeToTriangles.FindOrAdd(Edge1).Add(Triangle);
-			EdgeToTriangles.FindOrAdd(Edge2).Add(Triangle);
+		ensure(Indices.Num() % 3 == 0);
+		const uint32 NumTriangles = (uint32)Indices.Num() / 3;
+		if (!NumTriangles)
+		{
+			return;
 		}
 
+		// Gather edge information
+		TMap<FUintVector2, TArray<uint32>> EdgeToTrianglesMap;
+		BuildEdgeMap(Indices, EdgeToTrianglesMap);
+
 		// Build pattern islands
 		TSet<uint32> VisitedTriangles;
 		VisitedTriangles.Reserve(NumTriangles);
@@ -160,7 +173,7 @@ namespace UE::Chaos::ClothAsset::Private
 					(OldIndex0 != TriangleIndex0 && OldIndex1 != TriangleIndex0) ? TriangleIndex0 :
 					(OldIndex0 != TriangleIndex1 && OldIndex1 != TriangleIndex1) ? TriangleIndex1 : TriangleIndex2;
 
-				// Find the 2D interesection of the two connecting adjacent edges using the 3D reference length
+				// Find the 2D intersection of the two connecting adjacent edges using the 3D reference length
 				const FVector3f& P0 = Positions[OldIndex0];
 				const FVector3f& P1 = Positions[OldIndex1];
 				const FVector3f& P2 = Positions[OldIndex2];
@@ -215,6 +228,9 @@ namespace UE::Chaos::ClothAsset::Private
 					Island.Indices.Add(NewIndex0);
 					Island.Indices.Add(NewIndex1);
 					Island.Indices.Add(NewIndex2);
+					Island.SourceIndices.Add(OldIndex0);
+					Island.SourceIndices.Add(OldIndex1);
+					Island.SourceIndices.Add(OldIndex2);
 				}
 
 				// Add neighbor triangles to the queue
@@ -235,7 +251,7 @@ namespace UE::Chaos::ClothAsset::Private
 					const uint32 EdgeIndex0 = OldEdgeList[Edge].X;
 					const uint32 EdgeIndex1 = OldEdgeList[Edge].Y;
 
-					const TArray<uint32>& NeighborTriangles = EdgeToTriangles.FindChecked(MakeSortedEdge(EdgeIndex0, EdgeIndex1));
+					const TArray<uint32>& NeighborTriangles = EdgeToTrianglesMap.FindChecked(MakeSortedUintVector2(EdgeIndex0, EdgeIndex1));
 
 					for (const uint32 NeighborTriangle : NeighborTriangles)
 					{
@@ -258,6 +274,84 @@ namespace UE::Chaos::ClothAsset::Private
 			} while (Visitors.Dequeue(Visitor));
 		}
 	}
+
+	struct FSeam
+	{
+		TSet<FIntVector2> Stitches;
+		FIntVector2 Patterns;
+	};
+
+	// Rebuild the seam information from the torn/unwrapped mesh islands data
+	// Note that the isolated mesh islands are not technically patterns despite being considered so, 
+	// since they aren't sewed together in the source welded mesh.
+	// The algorithm will have to be slightly modified to be used with provided UV panels.
+	static void BuildSeams(const TArray<FIsland>& Islands, TArray<FSeam>& OutSeams)
+	{
+		OutSeams.Reset();
+
+		for (int32 IslandIndex = 0; IslandIndex < Islands.Num(); ++IslandIndex)
+		{
+			const FIsland& Island = Islands[IslandIndex];
+			
+			FSeam Seam;
+			Seam.Patterns = FIntVector2(IslandIndex);  // Just patching an unwrap here, will only do pattern to same
+
+			// Gather edge information for the source mesh
+			TMap<FUintVector2, TArray<uint32>> SourceEdgeToTrianglesMap;
+			BuildEdgeMap(Island.SourceIndices, SourceEdgeToTrianglesMap);
+
+			auto GetTriangleEdgeMatchingSourceEdge = [&Island](uint32 Triangle, const FUintVector2& SourceEdge) -> FUintVector2
+				{
+					const uint32 TriangleBase = Triangle * 3;
+					const uint32 TriangleIndex0 = Island.SourceIndices[TriangleBase + 0];
+					const uint32 TriangleIndex1 = Island.SourceIndices[TriangleBase + 1];
+					return
+						(SourceEdge[0] == TriangleIndex0 || SourceEdge[1] == TriangleIndex0) ?
+						(SourceEdge[0] == TriangleIndex1 || SourceEdge[1] == TriangleIndex1) ?
+						MakeSortedUintVector2(Island.Indices[TriangleBase + 0], Island.Indices[TriangleBase + 1]) :  // Edge 01
+						MakeSortedUintVector2(Island.Indices[TriangleBase + 2], Island.Indices[TriangleBase + 0]) :  // Edge 20
+						MakeSortedUintVector2(Island.Indices[TriangleBase + 1], Island.Indices[TriangleBase + 2]);   // Edge 12
+				};
+
+			// Look for disconnected triangles
+			for (const TPair<FUintVector2, TArray<uint32>>& SourceEdgeToTriangles : SourceEdgeToTrianglesMap)
+			{
+				const TArray<uint32>& Triangles = SourceEdgeToTriangles.Value;
+				const uint32 NumTriangles = (uint32)Triangles.Num();
+				if (NumTriangles > 1)
+				{
+					const FUintVector2& SourceEdge = SourceEdgeToTriangles.Key;
+
+					for (uint32 TriangleIndex0 = 0; TriangleIndex0 < NumTriangles - 1; ++TriangleIndex0)
+					{
+						const FUintVector2 Edge0 = GetTriangleEdgeMatchingSourceEdge(TriangleIndex0, SourceEdge);
+						check(Edge0[0] <= (uint32)TNumericLimits<int32>::Max() && Edge0[1] <= (uint32)TNumericLimits<int32>::Max());
+
+						for (uint32 TriangleIndex1 = TriangleIndex0 + 1; TriangleIndex1 < NumTriangles; ++TriangleIndex1)
+						{
+							const FUintVector2 Edge1 = GetTriangleEdgeMatchingSourceEdge(TriangleIndex1, SourceEdge);
+							check(Edge1[0] <= (uint32)TNumericLimits<int32>::Max() && Edge1[1] <= (uint32)TNumericLimits<int32>::Max());
+
+							if (Edge0[0] != Edge1[0])
+							{
+								Seam.Stitches.Emplace(FIntVector2(MakeSortedUintVector2(Edge0[0], Edge1[0])));
+							}
+							if (Edge0[1] != Edge1[1])
+							{
+								Seam.Stitches.Emplace(FIntVector2(MakeSortedUintVector2(Edge0[1], Edge1[1])));
+							}
+						}
+					}
+				}
+			}
+
+			// Add this island's seams
+			if (Seam.Stitches.Num())
+			{
+				OutSeams.Emplace(Seam);
+			}
+		}
+	}
 }  // End namespace UE::ClothAsset::Private
 
 // Cloth LOD adapter
@@ -282,70 +376,115 @@ namespace UE::Chaos::ClothAsset
 		return FClothPatternConstAdapter(ClothCollection, LodIndex, PatternIndex);
 	}
 
-	int32 FClothLodConstAdapter::GetNumElements(const TManagedArray<int32>& StartArray, const TManagedArray<int32>& EndArray) const
+	void FClothLodConstAdapter::BuildSimulationMesh(TArray<FVector3f>& Positions, TArray<FVector3f>& Normals, TArray<uint32>& Indices) const
 	{
-		const int32 ElementIndex = GetElementIndex();
-		const int32 Start = StartArray[ElementIndex];
-		const int32 End = EndArray[ElementIndex];
-		check(Start != INDEX_NONE || End == INDEX_NONE);  // Best to avoid situations where only one boundary of the range is set to INDEX_NONE
+		const int32 NumSimVertices = GetPatternsNumSimVertices();
 
-		return Start == INDEX_NONE ? 0 : End - Start + 1;
-	}
+		// Initialize welding map
+		TArray<int32> WeldingMap;
+		WeldingMap.SetNumUninitialized(NumSimVertices);
+		for (int32 SimVertexIndex = 0; SimVertexIndex < NumSimVertices; ++SimVertexIndex)
+		{
+			WeldingMap[SimVertexIndex] = SimVertexIndex;
+		}
 
-	template<typename T>
-	TConstArrayView<T> FClothLodConstAdapter::GetElements(const TManagedArray<T>& ElementArray, const TManagedArray<int32>& StartArray, const TManagedArray<int32>& EndArray) const
-	{
-		const int32 ElementIndex = GetElementIndex();
-		const int32 Start = StartArray[ElementIndex];
-		const int32 End = EndArray[ElementIndex];
-		check(Start != INDEX_NONE || End == INDEX_NONE);  // Best to avoid situations where only one boundary of the range is set to INDEX_NONE
+		TMap<int32, TSet<int32>> WeldingGroups;
 
-		return Start == INDEX_NONE ? TConstArrayView<T>() : TConstArrayView<T>(ElementArray.GetData() + Start, End - Start + 1);
-	}
+		auto UpdateWeldingMap =
+			[&WeldingMap, &WeldingGroups](int32 Index0, int32 Index1)
+		{
+			if (WeldingMap[Index0] != WeldingMap[Index1])
+			{
+				if (WeldingMap[Index0] > WeldingMap[Index1])
+				{
+					Swap(Index0, Index1);
+				}
 
-	template<bool bStart, bool bEnd>
-	TTuple<int32, int32> FClothLodConstAdapter::GetPatternsElementsStartEnd(const TManagedArray<int32>& StartArray, const TManagedArray<int32>& EndArray) const
-	{
-		const int32 ElementIndex = GetElementIndex();
-		const int32 PatternStart = ClothCollection->PatternStart[ElementIndex];
-		const int32 PatternEnd = ClothCollection->PatternEnd[ElementIndex];
+				TSet<int32>* WeldingGroup0 = WeldingGroups.Find(WeldingMap[Index0]);
+				if (!WeldingGroup0)
+				{
+					WeldingGroup0 = &WeldingGroups.Add(WeldingMap[Index0]);
+					WeldingGroup0->Add(Index0);
+				}
+
+				if (TSet<int32>* const WeldingGroup1 = WeldingGroups.Find(WeldingMap[Index1]))
+				{
+					// Update group1 mapped indices
+					for (int32 Index : *WeldingGroup1)
+					{
+						WeldingMap[Index] = WeldingMap[Index0];
+					}
+
+					// Merge group0 & group1
+					WeldingGroup0->Append(*WeldingGroup1);
+
+					// Remove group1
+					WeldingGroups.Remove(WeldingMap[Index1]);
+				}
+				else
+				{
+					WeldingMap[Index1] = WeldingMap[Index0];
+					WeldingGroup0->Add(Index1);
+				}
+			}
+		};
 
-		// Find Start and End indices for the entire LOD minding empty patterns on the way
-		int32 Start = INDEX_NONE;
-		int32 End = INDEX_NONE;
-		for (int32 PatternIndex = PatternStart; PatternIndex <= PatternEnd; ++PatternIndex)
+		// Apply all seams
+		const int32 NumSeams = GetNumSeams();
+		const TConstArrayView<TArray<FIntVector2>> SeamStitches = GetSeamStitches();
+		const TConstArrayView<FIntVector2> SeamPatternFirst = GetSeamPatterns();
+		for (int32 SeamIndex = 0; SeamIndex < NumSeams; ++SeamIndex)
 		{
-			if (bStart && StartArray[PatternIndex] != INDEX_NONE)
+			for (const FIntVector2& Stitch : SeamStitches[SeamIndex])
 			{
-				Start = (Start == INDEX_NONE) ? StartArray[PatternIndex] : FMath::Min(Start, StartArray[PatternIndex]);
+				UpdateWeldingMap(Stitch[0], Stitch[1]);
 			}
-			if (bEnd && EndArray[PatternIndex] != INDEX_NONE)
+		}
+
+		// Calculate the number of welded vertices
+		int32 NumWeldedVertices = NumSimVertices;
+		for (const TPair<int32, TSet<int32>>& WeldingGroup : WeldingGroups)
+		{
+			NumWeldedVertices -= WeldingGroup.Value.Num() - 1;
+		}
+
+		// Fill up the vertex arrays
+		Positions.SetNumUninitialized(NumWeldedVertices);
+		Normals.SetNumUninitialized(NumWeldedVertices);
+
+		const TConstArrayView<FVector3f> SimRestPosition = GetPatternsSimRestPosition();
+		const TConstArrayView<FVector3f> SimRestNormal = GetPatternsSimRestNormal();
+
+		TArray<uint32> WeldedIndices;
+		WeldedIndices.SetNumUninitialized(NumSimVertices);
+
+		uint32 WeldedIndex = 0;
+		for (int32 VertexIndex = 0; VertexIndex < NumSimVertices; ++VertexIndex)
+		{
+			if (WeldingMap[VertexIndex] == VertexIndex)
 			{
-				End = (End == INDEX_NONE) ? EndArray[PatternIndex] : FMath::Max(End, EndArray[PatternIndex]);
+				Positions[WeldedIndex] = SimRestPosition[VertexIndex];
+				Normals[WeldedIndex] = SimRestNormal[VertexIndex];
+				WeldedIndices[VertexIndex] = WeldedIndex++;
+			}
+			else
+			{
+				WeldedIndices[VertexIndex] = WeldedIndices[WeldingMap[VertexIndex]];
 			}
 		}
-		return TTuple<int32, int32>(Start, End);
-	}
 
-	int32 FClothLodConstAdapter::GetPatternsNumElements(const TManagedArray<int32>& StartArray, const TManagedArray<int32>& EndArray) const
-	{
-		const TTuple<int32, int32> StartEnd = GetPatternsElementsStartEnd<true, true>(StartArray, EndArray);
-		const int32 Start = StartEnd.Get<0>();
-		const int32 End = StartEnd.Get<1>();
-		check(Start != INDEX_NONE || End == INDEX_NONE);  // Best to avoid situations where only one boundary of the range is set to INDEX_NONE
-
-		return Start == INDEX_NONE ? 0 : End - Start + 1;
-	}
+		// Fill up the face array
+		const int32 NumSimFaces = GetPatternsNumSimFaces();
+		Indices.SetNumUninitialized(NumSimFaces * 3);
 
-	template<typename T>
-	TConstArrayView<T> FClothLodConstAdapter::GetPatternsElements(const TManagedArray<T>& ElementArray, const TManagedArray<int32>& StartArray, const TManagedArray<int32>& EndArray) const
-	{
-		const TTuple<int32, int32> StartEnd = GetPatternsElementsStartEnd<true, true>(StartArray, EndArray);
-		const int32 Start = StartEnd.Get<0>();
-		const int32 End = StartEnd.Get<1>();
+		const TConstArrayView<FIntVector3> SimIndices = GetPatternsSimIndices();
 
-		check(Start != INDEX_NONE || End == INDEX_NONE);  // Best to avoid situations where only one boundary of the range is set to INDEX_NONE
-		return Start == INDEX_NONE ? TConstArrayView<T>() : TConstArrayView<T>(ElementArray.GetData() + Start, End - Start + 1);
+		for (int32 FaceIndex = 0; FaceIndex < NumSimFaces; ++FaceIndex)
+		{
+			Indices[FaceIndex * 3 + 0] = WeldedIndices[SimIndices[FaceIndex][0]];
+			Indices[FaceIndex * 3 + 1] = WeldedIndices[SimIndices[FaceIndex][1]];
+			Indices[FaceIndex * 3 + 2] = WeldedIndices[SimIndices[FaceIndex][2]];
+		}
 	}
 
 	FClothLodAdapter::FClothLodAdapter(const TSharedPtr<FClothCollection>& InClothCollection, int32 InLodIndex)
@@ -389,12 +528,12 @@ namespace UE::Chaos::ClothAsset
 		}
 		GetClothCollection()->RemoveElements(FClothCollection::PatternsGroup, NumPatterns, GetClothCollection()->PatternStart[ElementIndex]);
 
-		const int32 NumStitchings = GetNumStitchings();
+		const int32 NumStitchings = GetNumSeams();
 		for (int32 StitchingIndex = 0; StitchingIndex < NumStitchings; ++StitchingIndex)
 		{
 			//GetStitching(StitchingIndex).Reset();  // TODO
 		}
-		GetClothCollection()->RemoveElements(FClothCollection::StitchingsGroup, NumStitchings, GetClothCollection()->StitchingStart[ElementIndex]);
+		GetClothCollection()->RemoveElements(FClothCollection::SeamsGroup, NumStitchings, GetClothCollection()->SeamStart[ElementIndex]);
 
 		const int32 NumTetherBatches = GetNumTetherBatches();
 		for (int32 TetherBatchIndex = 0; TetherBatchIndex < NumTetherBatches; ++TetherBatchIndex)
@@ -412,22 +551,21 @@ namespace UE::Chaos::ClothAsset
 
 		GetClothCollection()->PatternStart[ElementIndex] = INDEX_NONE;
 		GetClothCollection()->PatternEnd[ElementIndex] = INDEX_NONE;
-		GetClothCollection()->StitchingStart[ElementIndex] = INDEX_NONE;
-		GetClothCollection()->StitchingEnd[ElementIndex] = INDEX_NONE;
+		GetClothCollection()->SeamStart[ElementIndex] = INDEX_NONE;
+		GetClothCollection()->SeamEnd[ElementIndex] = INDEX_NONE;
 		GetClothCollection()->TetherBatchStart[ElementIndex] = INDEX_NONE;
 		GetClothCollection()->TetherBatchEnd[ElementIndex] = INDEX_NONE;
 		GetClothCollection()->LodBiasDepth[ElementIndex] = 0;
 	}
 
-	template<typename T>
-	TArrayView<T> FClothLodAdapter::GetPatternsElements(TManagedArray<T>& ElementArray, const TManagedArray<int32>& StartArray, const TManagedArray<int32>& EndArray)
+	int32 FClothLodAdapter::SetNumSeams(int32 NumSeams)
 	{
-		const TTuple<int32, int32> StartEnd = GetPatternsElementsStartEnd<true, true>(StartArray, EndArray);
-		const int32 Start = StartEnd.Get<0>();
-		const int32 End = StartEnd.Get<1>();
-
-		check(Start != INDEX_NONE || End == INDEX_NONE);  // Best to avoid situations where only one boundary of the range is set to INDEX_NONE
-		return Start == INDEX_NONE ? TArrayView<T>() : TArrayView<T>(ElementArray.GetData() + Start, End - Start + 1);
+		return GetClothCollection()->SetNumElements(
+			NumSeams,
+			FClothCollection::SeamsGroup,
+			GetClothCollection()->SeamStart,
+			GetClothCollection()->SeamEnd,
+			GetElementIndex());
 	}
 
 	void FClothLodAdapter::Initialize(const TArray<FVector3f>& Positions, const TArray<uint32>& Indices)
@@ -445,5 +583,18 @@ namespace UE::Chaos::ClothAsset
 				Pattern.Initialize(Island.Positions, Island.RestPositions, Island.Indices);
 			}
 		}
+
+		TArray<FSeam> Seams;
+		BuildSeams(Islands, Seams);  // Build the seam information as to be able to re-weld the mesh for simulation
+
+		SetNumSeams(Seams.Num());
+
+		const TArrayView<TArray<FIntVector2>> SeamStitches = GetSeamStitches();
+		const TArrayView<FIntVector2> SeamPatterns = GetSeamPatterns();
+		for (int32 SeamIndex = 0; SeamIndex < Seams.Num(); ++SeamIndex)
+		{
+			SeamPatterns[SeamIndex] = Seams[SeamIndex].Patterns;
+			SeamStitches[SeamIndex] = Seams[SeamIndex].Stitches.Array();
+		}
 	}
 }  // End namespace UE::Chaos::ClothAsset
diff --git a/Engine/Plugins/Experimental/ChaosClothAsset/Source/ChaosClothAsset/Private/ChaosClothAsset/ClothPatternAdapter.cpp b/Engine/Plugins/Experimental/ChaosClothAsset/Source/ChaosClothAsset/Private/ChaosClothAsset/ClothPatternAdapter.cpp
--- a/Engine/Plugins/Experimental/ChaosClothAsset/Source/ChaosClothAsset/Private/ChaosClothAsset/ClothPatternAdapter.cpp
+++ b/Engine/Plugins/Experimental/ChaosClothAsset/Source/ChaosClothAsset/Private/ChaosClothAsset/ClothPatternAdapter.cpp
@@ -32,27 +32,6 @@ namespace UE::Chaos::ClothAsset
 		return TConstArrayView<float>();
 	}
 
-	int32 FClothPatternConstAdapter::GetNumElements(const TManagedArray<int32>& StartArray, const TManagedArray<int32>& EndArray) const
-	{
-		const int32 ElementIndex = GetElementIndex();
-		const int32 Start = StartArray[ElementIndex];
-		const int32 End = EndArray[ElementIndex];
-		check(Start != INDEX_NONE || End == INDEX_NONE);  // Best to avoid situations where only one boundary of the range is set to INDEX_NONE
-
-		return Start == INDEX_NONE ? 0 : End - Start + 1;
-	}
-
-	template<typename T>
-	TConstArrayView<T> FClothPatternConstAdapter::GetElements(const TManagedArray<T>& ElementArray, const TManagedArray<int32>& StartArray, const TManagedArray<int32>& EndArray) const
-	{
-		const int32 ElementIndex = GetElementIndex();
-		const int32 Start = StartArray[ElementIndex];
-		const int32 End = EndArray[ElementIndex];
-		check(Start != INDEX_NONE || End == INDEX_NONE);  // Best to avoid situations where only one boundary of the range is set to INDEX_NONE
-
-		return Start == INDEX_NONE ? TConstArrayView<T>() : TConstArrayView<T>(ElementArray.GetData() + Start, End - Start + 1);
-	}
-
 	FClothPatternAdapter::FClothPatternAdapter(const TSharedPtr<FClothCollection>& InClothCollection, int32 InLodIndex, int32 InPatternIndex)
 		: FClothPatternConstAdapter(InClothCollection, InLodIndex, InPatternIndex)
 	{
@@ -75,38 +54,42 @@ namespace UE::Chaos::ClothAsset
 
 	int32 FClothPatternAdapter::SetNumSimVertices(int32 InNumSimVertices)
 	{
-		return SetNumElements(
+		return GetClothCollection()->SetNumElements(
 			InNumSimVertices,
 			FClothCollection::SimVerticesGroup,
 			GetClothCollection()->SimVerticesStart,
-			GetClothCollection()->SimVerticesEnd);
+			GetClothCollection()->SimVerticesEnd,
+			GetElementIndex());
 	}
 
 	int32 FClothPatternAdapter::SetNumSimFaces(int32 InNumSimFaces)
 	{
-		return SetNumElements(
+		return GetClothCollection()->SetNumElements(
 			InNumSimFaces,
 			FClothCollection::SimFacesGroup,
 			GetClothCollection()->SimFacesStart,
-			GetClothCollection()->SimFacesEnd);
+			GetClothCollection()->SimFacesEnd,
+			GetElementIndex());
 	}
 
 	int32 FClothPatternAdapter::SetNumRenderVertices(int32 InNumRenderVertices)
 	{
-		return SetNumElements(
+		return GetClothCollection()->SetNumElements(
 			InNumRenderVertices,
 			FClothCollection::RenderVerticesGroup,
 			GetClothCollection()->RenderVerticesStart,
-			GetClothCollection()->RenderVerticesEnd);
+			GetClothCollection()->RenderVerticesEnd,
+			GetElementIndex());
 	}
 
 	int32 FClothPatternAdapter::SetNumRenderFaces(int32 InNumRenderFaces)
 	{
-		return SetNumElements(
+		return GetClothCollection()->SetNumElements(
 			InNumRenderFaces,
 			FClothCollection::RenderFacesGroup,
 			GetClothCollection()->RenderFacesStart,
-			GetClothCollection()->RenderFacesEnd);
+			GetClothCollection()->RenderFacesEnd,
+			GetElementIndex());
 	}
 
 	void FClothPatternAdapter::SetDefaults()
@@ -126,77 +109,6 @@ namespace UE::Chaos::ClothAsset
 		GetClothCollection()->SimMaterialIndex[ElementIndex] = 0;
 	}
 
-	int32 FClothPatternAdapter::SetNumElements(int32 InNumElements, const FName& GroupName, TManagedArray<int32>& StartArray, TManagedArray<int32>& EndArray)
-	{
-		check(InNumElements >= 0);
-
-		const int32 ElementIndex = GetElementIndex();
-
-		int32& Start = StartArray[ElementIndex];
-		int32& End = EndArray[ElementIndex];
-		check(Start != INDEX_NONE || End == INDEX_NONE);  // Best to avoid situations where only one boundary of the range is set to INDEX_NONE
-
-		const int32 NumElements = (Start == INDEX_NONE) ? 0 : End - Start + 1;
-
-		if (const int32 Delta = InNumElements - NumElements)
-		{
-			if (Delta > 0)
-			{
-				// Find a previous valid index range to insert after when the range is empty
-				auto ComputeEnd = [&EndArray](int32 ElementIndex)->int32
-				{
-					for (int32 Index = ElementIndex; Index >= 0; --Index)
-					{
-						if (EndArray[Index] != INDEX_NONE)
-						{
-							return EndArray[Index];
-						}
-					}
-					return INDEX_NONE;
-				};
-
-				// Grow the array
-				const int32 Position = ComputeEnd(ElementIndex) + 1;
-				GetClothCollection()->InsertElements(Delta, Position, GroupName);
-
-				// Update Start/End
-				if (!NumElements)
-				{
-					Start = Position;
-				}
-				End = Start + InNumElements - 1;
-			}
-			else
-			{
-				// Shrink the array
-				const int32 Position = Start + InNumElements;
-				GetClothCollection()->RemoveElements(GroupName, -Delta, Position);
-
-				// Update Start/End
-				if (InNumElements)
-				{
-					End = Position - 1;
-				}
-				else
-				{
-					End = Start = INDEX_NONE;  // It is important to set the start & end to INDEX_NONE so that they never get automatically re-indexed by the managed array collection
-				}
-			}
-		}
-		return Start;
-	}
-
-	template<typename T>
-	TArrayView<T> FClothPatternAdapter::GetElements(TManagedArray<T>& ElementArray, const TManagedArray<int32>& StartArray, const TManagedArray<int32>& EndArray) const
-	{
-		const int32 ElementIndex = GetElementIndex();
-		const int32 Start = StartArray[ElementIndex];
-		const int32 End = EndArray[ElementIndex];
-		check(Start != INDEX_NONE || End == INDEX_NONE);  // Best to avoid situations where only one boundary of the range is set to INDEX_NONE
-
-		return Start == INDEX_NONE ? TArrayView<T>() : TArrayView<T>(ElementArray.GetData() + Start, End - Start + 1);
-	}
-
 	void FClothPatternAdapter::Initialize(const TArray<FVector2f>& Positions, const TArray<FVector3f>& RestPositions, const TArray<uint32>& Indices)
 	{
 		const int32 NumSimVertices = Positions.Num();
diff --git a/Engine/Plugins/Experimental/ChaosClothAsset/Source/ChaosClothAsset/Public/ChaosClothAsset/ClothAdapter.h b/Engine/Plugins/Experimental/ChaosClothAsset/Source/ChaosClothAsset/Public/ChaosClothAsset/ClothAdapter.h
--- a/Engine/Plugins/Experimental/ChaosClothAsset/Source/ChaosClothAsset/Public/ChaosClothAsset/ClothAdapter.h
+++ b/Engine/Plugins/Experimental/ChaosClothAsset/Source/ChaosClothAsset/Public/ChaosClothAsset/ClothAdapter.h
@@ -39,7 +39,7 @@ namespace UE::Chaos::ClothAsset
 	/**
 	 * Cloth adapter object to provide a more convenient object oriented access to the cloth collection.
 	 */
-	class CHAOSCLOTHASSET_API FClothAdapter : public FClothConstAdapter
+	class CHAOSCLOTHASSET_API FClothAdapter final : public FClothConstAdapter
 	{
 	public:
 		FClothAdapter(const TSharedPtr<FClothCollection>& InClothCollection);
diff --git a/Engine/Plugins/Experimental/ChaosClothAsset/Source/ChaosClothAsset/Public/ChaosClothAsset/ClothCollection.h b/Engine/Plugins/Experimental/ChaosClothAsset/Source/ChaosClothAsset/Public/ChaosClothAsset/ClothCollection.h
--- a/Engine/Plugins/Experimental/ChaosClothAsset/Source/ChaosClothAsset/Public/ChaosClothAsset/ClothCollection.h
+++ b/Engine/Plugins/Experimental/ChaosClothAsset/Source/ChaosClothAsset/Public/ChaosClothAsset/ClothCollection.h
@@ -2,14 +2,12 @@
 
 #pragma once
 
-#include "CoreMinimal.h"
 #include "GeometryCollection/ManagedArrayCollection.h"
-#include "Templates/SharedPointer.h"
 
 // TODO:
 // - Transition Sim LOD data
 // - Move wrap deformer into struct similar to FMeshToMeshVertData
-// - Add FVertexBoneData to MANAGED_ARRAY_TYPE
+// - Add FVertexBoneData to MANAGED_ARRAY_TYPE?
 
 namespace UE::Chaos::ClothAsset
 {
@@ -43,6 +41,27 @@ namespace UE::Chaos::ClothAsset
 		using Super::Serialize;
 		void Serialize(FArchive& Ar);
 
+		/** Set the number of elements to one of the groups that have start/end indices while maintaining the correct order of the data. */
+		int32 SetNumElements(int32 InNumElements, const FName& GroupName, TManagedArray<int32>& StartArray, TManagedArray<int32>& EndArray, int32 StartEndIndex);
+
+		/** Get the number of elements to one of the groups that have start/end indices. */
+		int32 GetNumElements(const TManagedArray<int32>& StartArray, const TManagedArray<int32>& EndArray, int32 StartEndIndex) const;
+
+		template<typename T>
+		inline TConstArrayView<T> GetElements(const TManagedArray<T>& ElementArray, const TManagedArray<int32>& StartArray, const TManagedArray<int32>& EndArray, int32 StartEndIndex) const;
+		template<typename T>
+		inline TArrayView<T> GetElements(TManagedArray<T>& ElementArray, const TManagedArray<int32>& StartArray, const TManagedArray<int32>& EndArray, int32 StartEndIndex);
+
+		int32 GetPatternsNumElements(const TManagedArray<int32>& StartArray, const TManagedArray<int32>& EndArray, int32 LodIndex) const;
+
+		template<typename T>
+		inline TConstArrayView<T> GetPatternsElements(const TManagedArray<T>& ElementArray, const TManagedArray<int32>& StartArray, const TManagedArray<int32>& EndArray, int32 LodIndex) const;
+		template<typename T>
+		inline TArrayView<T> GetPatternsElements(TManagedArray<T>& ElementArray, const TManagedArray<int32>& StartArray, const TManagedArray<int32>& EndArray, int32 LodIndex);
+
+		template<bool bStart = true, bool bEnd = true>
+		inline TTuple<int32, int32> GetPatternsElementsStartEnd(const TManagedArray<int32>& StartArray, const TManagedArray<int32>& EndArray, int32 StartEndIndex) const;
+
 		// Attribute groups, predefined data member of the FClothLod object.
 		static const FName SimVerticesGroup;  // Contains patterns' 2D positions, 3D draped position (rest)
 		static const FName SimFacesGroup;  // Contains indices to sim vertex
@@ -50,7 +69,7 @@ namespace UE::Chaos::ClothAsset
 		static const FName RenderFacesGroup;  // Contains indices to render vertex
 		static const FName WrapDeformersGroup;  // Contains the wrap deformers cloth capture information
 		static const FName PatternsGroup;  // Contains pattern relationships to other groups
-		static const FName StitchingsGroup;  // Contains stitched sim vertex indices for quick welding or constraint creations
+		static const FName SeamsGroup;  // Contains pairs of stitched sim vertex indices and matching pattern indices
 		static const FName TethersGroup;  // Tethers information
 		static const FName TetherBatchesGroup;  // Tethers parallel batch processing information
 		static const FName LodsGroup;  // Lod split information
@@ -100,8 +119,9 @@ namespace UE::Chaos::ClothAsset
 		TManagedArray<int32> StatusFlags;  // Whether this pattern should be rendered, simulated, ...etc
 		TManagedArray<int32> SimMaterialIndex;  // Cloth material per pattern
 
-		// Stitching Group
-		TManagedArray<TSet<int32>> StitchedVertices;  // Stitched vertex indices
+		// Seam Group
+		TManagedArray<FIntVector2> SeamPatterns;  // Stitched pattern indices (LOD based indices) for SeamIndicesFirst
+		TManagedArray<TArray<FIntVector2>> SeamStitches;  // Stitched vertex indices pair (LOD based indices)
 
 		// Tethers Group
 		TManagedArray<int32> TetherKinematicIndex;
@@ -115,8 +135,8 @@ namespace UE::Chaos::ClothAsset
 		// LOD Group
 		TManagedArray<int32> PatternStart;
 		TManagedArray<int32> PatternEnd;
-		TManagedArray<int32> StitchingStart;
-		TManagedArray<int32> StitchingEnd;
+		TManagedArray<int32> SeamStart;
+		TManagedArray<int32> SeamEnd;
 		TManagedArray<int32> TetherBatchStart;
 		TManagedArray<int32> TetherBatchEnd;
 		TManagedArray<int32> LodBiasDepth;  // The number of LODs covered by each Sim LOD (for the wrap deformer)
@@ -124,4 +144,68 @@ namespace UE::Chaos::ClothAsset
 	protected:
 		void Construct();
 	};
+
+	template<typename T>
+	inline TConstArrayView<T> FClothCollection::GetElements(const TManagedArray<T>& ElementArray, const TManagedArray<int32>& StartArray, const TManagedArray<int32>& EndArray, int32 StartEndIndex) const
+	{
+		const int32 Start = StartArray[StartEndIndex];
+		const int32 End = EndArray[StartEndIndex];
+		check(Start != INDEX_NONE || End == INDEX_NONE);  // Best to avoid situations where only one boundary of the range is set to INDEX_NONE
+
+		return Start == INDEX_NONE ? TConstArrayView<T>() : TConstArrayView<T>(ElementArray.GetData() + Start, End - Start + 1);
+	}
+
+	template<typename T>
+	inline TArrayView<T> FClothCollection::GetElements(TManagedArray<T>& ElementArray, const TManagedArray<int32>& StartArray, const TManagedArray<int32>& EndArray, int32 StartEndIndex)
+	{
+		const int32 Start = StartArray[StartEndIndex];
+		const int32 End = EndArray[StartEndIndex];
+		check(Start != INDEX_NONE || End == INDEX_NONE);  // Best to avoid situations where only one boundary of the range is set to INDEX_NONE
+
+		return Start == INDEX_NONE ? TArrayView<T>() : TArrayView<T>(ElementArray.GetData() + Start, End - Start + 1);
+	}
+
+	template<typename T>
+	inline TConstArrayView<T> FClothCollection::GetPatternsElements(const TManagedArray<T>& ElementArray, const TManagedArray<int32>& StartArray, const TManagedArray<int32>& EndArray, int32 LodIndex) const
+	{
+		const TTuple<int32, int32> StartEnd = GetPatternsElementsStartEnd(StartArray, EndArray, LodIndex);
+		const int32 Start = StartEnd.Get<0>();
+		const int32 End = StartEnd.Get<1>();
+
+		check(Start != INDEX_NONE || End == INDEX_NONE);  // Best to avoid situations where only one boundary of the range is set to INDEX_NONE
+		return Start == INDEX_NONE ? TConstArrayView<T>() : TConstArrayView<T>(ElementArray.GetData() + Start, End - Start + 1);
+	}
+
+	template<typename T>
+	inline TArrayView<T> FClothCollection::GetPatternsElements(TManagedArray<T>& ElementArray, const TManagedArray<int32>& StartArray, const TManagedArray<int32>& EndArray, int32 LodIndex)
+	{
+		const TTuple<int32, int32> StartEnd = GetPatternsElementsStartEnd(StartArray, EndArray, LodIndex);
+		const int32 Start = StartEnd.Get<0>();
+		const int32 End = StartEnd.Get<1>();
+
+		check(Start != INDEX_NONE || End == INDEX_NONE);  // Best to avoid situations where only one boundary of the range is set to INDEX_NONE
+		return Start == INDEX_NONE ? TArrayView<T>() : TArrayView<T>(ElementArray.GetData() + Start, End - Start + 1);
+	}
+
+	template<bool bStart, bool bEnd>
+	inline TTuple<int32, int32> FClothCollection::GetPatternsElementsStartEnd(const TManagedArray<int32>& StartArray, const TManagedArray<int32>& EndArray, int32 LodIndex) const
+	{
+		const int32 LodPatternStart = PatternStart[LodIndex];
+		const int32 LodPatternEnd = PatternEnd[LodIndex];
+
+		int32 Start = INDEX_NONE;  // Find Start and End indices for the entire LOD minding empty patterns on the way
+		int32 End = INDEX_NONE;
+		for (int32 PatternIndex = LodPatternStart; PatternIndex <= LodPatternEnd; ++PatternIndex)
+		{
+			if (bStart && StartArray[PatternIndex] != INDEX_NONE)
+			{
+				Start = (Start == INDEX_NONE) ? StartArray[PatternIndex] : FMath::Min(Start, StartArray[PatternIndex]);
+			}
+			if (bEnd && EndArray[PatternIndex] != INDEX_NONE)
+			{
+				End = (End == INDEX_NONE) ? EndArray[PatternIndex] : FMath::Max(End, EndArray[PatternIndex]);
+			}
+		}
+		return TTuple<int32, int32>(Start, End);
+	}
 }  // End namespace UE::Chaos::ClothAsset
diff --git a/Engine/Plugins/Experimental/ChaosClothAsset/Source/ChaosClothAsset/Public/ChaosClothAsset/ClothLodAdapter.h b/Engine/Plugins/Experimental/ChaosClothAsset/Source/ChaosClothAsset/Public/ChaosClothAsset/ClothLodAdapter.h
--- a/Engine/Plugins/Experimental/ChaosClothAsset/Source/ChaosClothAsset/Public/ChaosClothAsset/ClothLodAdapter.h
+++ b/Engine/Plugins/Experimental/ChaosClothAsset/Source/ChaosClothAsset/Public/ChaosClothAsset/ClothLodAdapter.h
@@ -25,42 +25,59 @@ namespace UE::Chaos::ClothAsset
 
 		FClothPatternConstAdapter GetPattern(int32  PatternIndex) const;
 
-		int32 GetNumStitchings() const { return GetNumElements(ClothCollection->StitchingStart, ClothCollection->StitchingEnd); }
-		int32 GetNumTetherBatches() const { return GetNumElements(ClothCollection->TetherBatchStart, ClothCollection->TetherBatchEnd); }
-
 		// Patterns Group
-		int32 GetNumPatterns() const { return GetNumElements(ClothCollection->PatternStart, ClothCollection->PatternEnd); }
-		TConstArrayView<int32> GetSimVerticesStart() const { return GetElements(GetClothCollection()->SimVerticesStart, GetClothCollection()->PatternStart, GetClothCollection()->PatternEnd); }
-		TConstArrayView<int32> GetSimVerticesEnd() const { return GetElements(GetClothCollection()->SimVerticesEnd, GetClothCollection()->PatternStart, GetClothCollection()->PatternEnd); }
-		TConstArrayView<int32> GetSimFacesStart() const { return GetElements(GetClothCollection()->SimFacesStart, GetClothCollection()->PatternStart, GetClothCollection()->PatternEnd); }
-		TConstArrayView<int32> GetSimFacesEnd() const { return GetElements(GetClothCollection()->SimFacesEnd, GetClothCollection()->PatternStart, GetClothCollection()->PatternEnd); }
-		TConstArrayView<int32> GetRenderVerticesStart() const { return GetElements(GetClothCollection()->RenderVerticesStart, GetClothCollection()->PatternStart, GetClothCollection()->PatternEnd); }
-		TConstArrayView<int32> GetRenderVerticesEnd() const { return GetElements(GetClothCollection()->RenderVerticesEnd, GetClothCollection()->PatternStart, GetClothCollection()->PatternEnd); }
-		TConstArrayView<int32> GetRenderFacesStart() const { return GetElements(GetClothCollection()->RenderFacesStart, GetClothCollection()->PatternStart, GetClothCollection()->PatternEnd); }
-		TConstArrayView<int32> GetRenderFacesEnd() const { return GetElements(GetClothCollection()->RenderFacesEnd, GetClothCollection()->PatternStart, GetClothCollection()->PatternEnd); }
-
-		// Patterns Sim Vertices Group, use these to access all pattern at once, and when using the SimIndices
-		int32 GetPatternsSimVerticesStart() const { return GetPatternsElementsStartEnd<true, false>(GetClothCollection()->SimVerticesStart, GetClothCollection()->SimVerticesEnd).Get<0>(); }
-		int32 GetPatternsSimVerticesEnd() const { return GetPatternsElementsStartEnd<false, true>(GetClothCollection()->SimVerticesStart, GetClothCollection()->SimVerticesEnd).Get<1>(); }
-		TTuple<int32, int32> GetPatternsSimVerticesRange() const { return GetPatternsElementsStartEnd<true, true>(GetClothCollection()->SimVerticesStart, GetClothCollection()->SimVerticesEnd); }
-		int32 GetPatternsNumSimVertices() const { return GetPatternsNumElements(GetClothCollection()->SimVerticesStart, GetClothCollection()->SimVerticesEnd); }
-
-		TConstArrayView<FVector2f> GetPatternsSimPosition() const { return GetPatternsElements(GetClothCollection()->SimPosition, GetClothCollection()->SimVerticesStart, GetClothCollection()->SimVerticesEnd); }
-		TConstArrayView<FVector3f> GetPatternsSimRestPosition() const { return GetPatternsElements(GetClothCollection()->SimRestPosition, GetClothCollection()->SimVerticesStart, GetClothCollection()->SimVerticesEnd); }
-		TConstArrayView<FVector3f> GetPatternsSimRestNormal() const { return GetPatternsElements(GetClothCollection()->SimRestNormal, GetClothCollection()->SimVerticesStart, GetClothCollection()->SimVerticesEnd); }
-
-		// Patterns Render Vertices Group, use these to access all pattern at once, and when using the RenderIndices
-		int32 GetPatternsRenderVerticesStart() const { return GetPatternsElementsStartEnd<true, false>(GetClothCollection()->RenderVerticesStart, GetClothCollection()->RenderVerticesEnd).Get<0>(); }
-		int32 GetPatternsRenderVerticesEnd() const { return GetPatternsElementsStartEnd<false, true>(GetClothCollection()->RenderVerticesStart, GetClothCollection()->RenderVerticesEnd).Get<1>(); }
-		TTuple<int32, int32> GetPatternsRenderVerticesRange() const { return GetPatternsElementsStartEnd<true, true>(GetClothCollection()->RenderVerticesStart, GetClothCollection()->RenderVerticesEnd); }
-		int32 GetPatternsNumRenderVertices() const { return GetPatternsNumElements(GetClothCollection()->RenderVerticesStart, GetClothCollection()->RenderVerticesEnd); }
-
-		TConstArrayView<FVector3f> GetPatternsRenderPosition() const { return GetPatternsElements(GetClothCollection()->RenderPosition, GetClothCollection()->RenderVerticesStart, GetClothCollection()->RenderVerticesEnd); }
-		TConstArrayView<FVector3f> GetPatternsRenderNormal() const { return GetPatternsElements(GetClothCollection()->RenderNormal, GetClothCollection()->RenderVerticesStart, GetClothCollection()->RenderVerticesEnd); }
-		TConstArrayView<FVector3f> GetPatternsRenderTangentU() const { return GetPatternsElements(GetClothCollection()->RenderTangentU, GetClothCollection()->RenderVerticesStart, GetClothCollection()->RenderVerticesEnd); }
-		TConstArrayView<FVector3f> GetPatternsRenderTangentV() const { return GetPatternsElements(GetClothCollection()->RenderTangentV, GetClothCollection()->RenderVerticesStart, GetClothCollection()->RenderVerticesEnd); }
-		TConstArrayView<TArray<FVector2f>> GePatternstRenderUVs() const { return GetPatternsElements(GetClothCollection()->RenderUVs, GetClothCollection()->RenderVerticesStart, GetClothCollection()->RenderVerticesEnd); }
-		TConstArrayView<FLinearColor> GetPatternsRenderColor() const { return GetPatternsElements(GetClothCollection()->RenderColor, GetClothCollection()->RenderVerticesStart, GetClothCollection()->RenderVerticesEnd); }
+		int32 GetNumPatterns() const { return GetClothCollection()->GetNumElements(GetClothCollection()->PatternStart, GetClothCollection()->PatternEnd, GetElementIndex()); }
+		TConstArrayView<int32> GetSimVerticesStart() const { return GetClothCollection()->GetElements(GetClothCollection()->SimVerticesStart, GetClothCollection()->PatternStart, GetClothCollection()->PatternEnd, GetElementIndex()); }
+		TConstArrayView<int32> GetSimVerticesEnd() const { return GetClothCollection()->GetElements(GetClothCollection()->SimVerticesEnd, GetClothCollection()->PatternStart, GetClothCollection()->PatternEnd, GetElementIndex()); }
+		TConstArrayView<int32> GetSimFacesStart() const { return GetClothCollection()->GetElements(GetClothCollection()->SimFacesStart, GetClothCollection()->PatternStart, GetClothCollection()->PatternEnd, GetElementIndex()); }
+		TConstArrayView<int32> GetSimFacesEnd() const { return GetClothCollection()->GetElements(GetClothCollection()->SimFacesEnd, GetClothCollection()->PatternStart, GetClothCollection()->PatternEnd, GetElementIndex()); }
+		TConstArrayView<int32> GetRenderVerticesStart() const { return GetClothCollection()->GetElements(GetClothCollection()->RenderVerticesStart, GetClothCollection()->PatternStart, GetClothCollection()->PatternEnd, GetElementIndex()); }
+		TConstArrayView<int32> GetRenderVerticesEnd() const { return GetClothCollection()->GetElements(GetClothCollection()->RenderVerticesEnd, GetClothCollection()->PatternStart, GetClothCollection()->PatternEnd, GetElementIndex()); }
+		TConstArrayView<int32> GetRenderFacesStart() const { return GetClothCollection()->GetElements(GetClothCollection()->RenderFacesStart, GetClothCollection()->PatternStart, GetClothCollection()->PatternEnd, GetElementIndex()); }
+		TConstArrayView<int32> GetRenderFacesEnd() const { return GetClothCollection()->GetElements(GetClothCollection()->RenderFacesEnd, GetClothCollection()->PatternStart, GetClothCollection()->PatternEnd, GetElementIndex()); }
+
+		// Seams Group
+		int32 GetNumSeams() const { return GetClothCollection()->GetNumElements(GetClothCollection()->SeamStart, GetClothCollection()->SeamEnd, GetElementIndex()); }
+		TConstArrayView<FIntVector2> GetSeamPatterns() const { return GetClothCollection()->GetElements(GetClothCollection()->SeamPatterns, GetClothCollection()->SeamStart, GetClothCollection()->SeamEnd, GetElementIndex()); }
+		TConstArrayView<TArray<FIntVector2>> GetSeamStitches() const { return GetClothCollection()->GetElements(GetClothCollection()->SeamStitches, GetClothCollection()->SeamStart, GetClothCollection()->SeamEnd, GetElementIndex()); }
+
+		// Tethers Group
+		int32 GetNumTetherBatches() const { return GetClothCollection()->GetNumElements(GetClothCollection()->TetherBatchStart, GetClothCollection()->TetherBatchEnd, GetElementIndex()); }
+		// TODO: Tether API 
+
+		// All patterns Sim Vertices Group, use these to access all pattern at once, and when using the SimIndices
+		int32 GetPatternsSimVerticesStart() const { return GetClothCollection()->GetPatternsElementsStartEnd<true, false>(GetClothCollection()->SimVerticesStart, GetClothCollection()->SimVerticesEnd, GetElementIndex()).Get<0>(); }
+		int32 GetPatternsSimVerticesEnd() const { return GetClothCollection()->GetPatternsElementsStartEnd<false, true>(GetClothCollection()->SimVerticesStart, GetClothCollection()->SimVerticesEnd, GetElementIndex()).Get<1>(); }
+		TTuple<int32, int32> GetPatternsSimVerticesRange() const { return GetClothCollection()->GetPatternsElementsStartEnd<true, true>(GetClothCollection()->SimVerticesStart, GetClothCollection()->SimVerticesEnd, GetElementIndex()); }
+		int32 GetPatternsNumSimVertices() const { return GetClothCollection()->GetPatternsNumElements(GetClothCollection()->SimVerticesStart, GetClothCollection()->SimVerticesEnd, GetElementIndex()); }
+
+		TConstArrayView<FVector2f> GetPatternsSimPosition() const { return GetClothCollection()->GetPatternsElements(GetClothCollection()->SimPosition, GetClothCollection()->SimVerticesStart, GetClothCollection()->SimVerticesEnd, GetElementIndex()); }
+		TConstArrayView<FVector3f> GetPatternsSimRestPosition() const { return GetClothCollection()->GetPatternsElements(GetClothCollection()->SimRestPosition, GetClothCollection()->SimVerticesStart, GetClothCollection()->SimVerticesEnd, GetElementIndex()); }
+		TConstArrayView<FVector3f> GetPatternsSimRestNormal() const { return GetClothCollection()->GetPatternsElements(GetClothCollection()->SimRestNormal, GetClothCollection()->SimVerticesStart, GetClothCollection()->SimVerticesEnd, GetElementIndex()); }
+
+		// All patterns Sim Faces Group, use these to access all pattern at once
+		TTuple<int32, int32> GetPatternsSimFacesRange() const { return GetClothCollection()->GetPatternsElementsStartEnd<true, true>(GetClothCollection()->SimFacesStart, GetClothCollection()->SimFacesEnd, GetElementIndex()); }
+		int32 GetPatternsNumSimFaces() const { return GetClothCollection()->GetPatternsNumElements(GetClothCollection()->SimFacesStart, GetClothCollection()->SimFacesEnd, GetElementIndex()); }
+		TConstArrayView<FIntVector3> GetPatternsSimIndices() const { return GetClothCollection()->GetPatternsElements(GetClothCollection()->SimIndices, GetClothCollection()->SimFacesStart, GetClothCollection()->SimFacesEnd, GetElementIndex()); }
+
+		// All patterns Render Vertices Group, use these to access all pattern at once, and when using the RenderIndices
+		int32 GetPatternsRenderVerticesStart() const { return GetClothCollection()->GetPatternsElementsStartEnd<true, false>(GetClothCollection()->RenderVerticesStart, GetClothCollection()->RenderVerticesEnd, GetElementIndex()).Get<0>(); }
+		int32 GetPatternsRenderVerticesEnd() const { return GetClothCollection()->GetPatternsElementsStartEnd<false, true>(GetClothCollection()->RenderVerticesStart, GetClothCollection()->RenderVerticesEnd, GetElementIndex()).Get<1>(); }
+		TTuple<int32, int32> GetPatternsRenderVerticesRange() const { return GetClothCollection()->GetPatternsElementsStartEnd(GetClothCollection()->RenderVerticesStart, GetClothCollection()->RenderVerticesEnd, GetElementIndex()); }
+		int32 GetPatternsNumRenderVertices() const { return GetClothCollection()->GetPatternsNumElements(GetClothCollection()->RenderVerticesStart, GetClothCollection()->RenderVerticesEnd, GetElementIndex()); }
+
+		TConstArrayView<FVector3f> GetPatternsRenderPosition() const { return GetClothCollection()->GetPatternsElements(GetClothCollection()->RenderPosition, GetClothCollection()->RenderVerticesStart, GetClothCollection()->RenderVerticesEnd, GetElementIndex()); }
+		TConstArrayView<FVector3f> GetPatternsRenderNormal() const { return GetClothCollection()->GetPatternsElements(GetClothCollection()->RenderNormal, GetClothCollection()->RenderVerticesStart, GetClothCollection()->RenderVerticesEnd, GetElementIndex()); }
+		TConstArrayView<FVector3f> GetPatternsRenderTangentU() const { return GetClothCollection()->GetPatternsElements(GetClothCollection()->RenderTangentU, GetClothCollection()->RenderVerticesStart, GetClothCollection()->RenderVerticesEnd, GetElementIndex()); }
+		TConstArrayView<FVector3f> GetPatternsRenderTangentV() const { return GetClothCollection()->GetPatternsElements(GetClothCollection()->RenderTangentV, GetClothCollection()->RenderVerticesStart, GetClothCollection()->RenderVerticesEnd, GetElementIndex()); }
+		TConstArrayView<TArray<FVector2f>> GePatternstRenderUVs() const { return GetClothCollection()->GetPatternsElements(GetClothCollection()->RenderUVs, GetClothCollection()->RenderVerticesStart, GetClothCollection()->RenderVerticesEnd, GetElementIndex()); }
+		TConstArrayView<FLinearColor> GetPatternsRenderColor() const { return GetClothCollection()->GetPatternsElements(GetClothCollection()->RenderColor, GetClothCollection()->RenderVerticesStart, GetClothCollection()->RenderVerticesEnd, GetElementIndex()); }
+
+		// All patterns Sim Faces Group, use these to access all pattern at once
+		TTuple<int32, int32> GetPatternsRenderFacesRange() const { return GetClothCollection()->GetPatternsElementsStartEnd<true, true>(GetClothCollection()->RenderFacesStart, GetClothCollection()->RenderFacesEnd, GetElementIndex()); }
+		int32 GetPatternsNumRenderFaces() const { return GetClothCollection()->GetPatternsNumElements(GetClothCollection()->RenderFacesStart, GetClothCollection()->RenderFacesEnd, GetElementIndex()); }
+		TConstArrayView<FIntVector3> GetPatternsRenderIndices() const { return GetClothCollection()->GetPatternsElements(GetClothCollection()->RenderIndices, GetClothCollection()->RenderFacesStart, GetClothCollection()->RenderFacesEnd, GetElementIndex()); }
+		TConstArrayView<int32> GetPatternsRenderMaterialIndex() const { return GetClothCollection()->GetPatternsElements(GetClothCollection()->RenderMaterialIndex, GetClothCollection()->RenderFacesStart, GetClothCollection()->RenderFacesEnd, GetElementIndex()); }
 
 		/** Return the element index within the cloth collection. */
 		int32 GetElementIndex() const { return LodIndex; }
@@ -71,19 +88,10 @@ namespace UE::Chaos::ClothAsset
 		/** Return the underlaying cloth collection this adapter has been created with. */
 		const TSharedPtr<const FClothCollection>& GetClothCollection() const { return ClothCollection; }
 
-	protected:
-		template<bool bStart, bool bEnd>
-		TTuple<int32, int32> GetPatternsElementsStartEnd(const TManagedArray<int32>& StartArray, const TManagedArray<int32>& EndArray) const;
+		/** Return the welded simulation mesh for this LOD. */
+		void BuildSimulationMesh(TArray<FVector3f>& Positions, TArray<FVector3f>& Normals, TArray<uint32>& Indices) const;
 
 	private:
-		int32 GetNumElements(const TManagedArray<int32>& StartArray, const TManagedArray<int32>& EndArray) const;
-		template<typename T>
-		TConstArrayView<T> GetElements(const TManagedArray<T>& ElementArray, const TManagedArray<int32>& StartArray, const TManagedArray<int32>& EndArray) const;
-
-		int32 GetPatternsNumElements(const TManagedArray<int32>& StartArray, const TManagedArray<int32>& EndArray) const;
-		template<typename T>
-		TConstArrayView<T> GetPatternsElements(const TManagedArray<T>& ElementArray, const TManagedArray<int32>& StartArray, const TManagedArray<int32>& EndArray) const;
-
 		TSharedPtr<const FClothCollection> ClothCollection;
 		int32 LodIndex;
 	};
@@ -91,7 +99,7 @@ namespace UE::Chaos::ClothAsset
 	/**
 	 * Cloth LOD adapter object to provide a more convenient object oriented access to the cloth collection.
 	 */
-	class CHAOSCLOTHASSET_API FClothLodAdapter : public FClothLodConstAdapter
+	class CHAOSCLOTHASSET_API FClothLodAdapter final : public FClothLodConstAdapter
 	{
 	public:
 		FClothLodAdapter(const TSharedPtr<FClothCollection>& InClothCollection, int32 InLodIndex);
@@ -120,25 +128,44 @@ namespace UE::Chaos::ClothAsset
 		/** Return the underlaying Cloth Collection this adapter has been created with. */
 		TSharedPtr<FClothCollection> GetClothCollection() { return ConstCastSharedPtr<FClothCollection>(FClothLodConstAdapter::GetClothCollection()); }
 
-		// Patterns Sim Vertices Group, use these to access all pattern at once, and when using the SimIndices
-		TArrayView<FVector2f> GetPatternsSimPosition() { return GetPatternsElements(GetClothCollection()->SimPosition, GetClothCollection()->SimVerticesStart, GetClothCollection()->SimVerticesEnd); }
-		TArrayView<FVector3f> GetPatternsSimRestPosition() { return GetPatternsElements(GetClothCollection()->SimRestPosition, GetClothCollection()->SimVerticesStart, GetClothCollection()->SimVerticesEnd); }
-		TArrayView<FVector3f> GetPatternsSimRestNormal() { return GetPatternsElements(GetClothCollection()->SimRestNormal, GetClothCollection()->SimVerticesStart, GetClothCollection()->SimVerticesEnd); }
-
-		// Patterns Render Vertices Group, use these to access all pattern at once, and when using the RenderIndices
-		TArrayView<FVector3f> GetPatternsRenderPosition() { return GetPatternsElements(GetClothCollection()->RenderPosition, GetClothCollection()->RenderVerticesStart, GetClothCollection()->RenderVerticesEnd); }
-		TArrayView<FVector3f> GetPatternsRenderNormal() { return GetPatternsElements(GetClothCollection()->RenderNormal, GetClothCollection()->RenderVerticesStart, GetClothCollection()->RenderVerticesEnd); }
-		TArrayView<FVector3f> GetPatternsRenderTangentU() { return GetPatternsElements(GetClothCollection()->RenderTangentU, GetClothCollection()->RenderVerticesStart, GetClothCollection()->RenderVerticesEnd); }
-		TArrayView<FVector3f> GetPatternsRenderTangentV() { return GetPatternsElements(GetClothCollection()->RenderTangentV, GetClothCollection()->RenderVerticesStart, GetClothCollection()->RenderVerticesEnd); }
-		TArrayView<TArray<FVector2f>> GePatternstRenderUVs() { return GetPatternsElements(GetClothCollection()->RenderUVs, GetClothCollection()->RenderVerticesStart, GetClothCollection()->RenderVerticesEnd); }
-		TArrayView<FLinearColor> GetPatternsRenderColor() { return GetPatternsElements(GetClothCollection()->RenderColor, GetClothCollection()->RenderVerticesStart, GetClothCollection()->RenderVerticesEnd); }
+		// Patterns Group
+		TArrayView<int32> GetSimVerticesStart() { return GetClothCollection()->GetElements(GetClothCollection()->SimVerticesStart, GetClothCollection()->PatternStart, GetClothCollection()->PatternEnd, GetElementIndex()); }
+		TArrayView<int32> GetSimVerticesEnd() { return GetClothCollection()->GetElements(GetClothCollection()->SimVerticesEnd, GetClothCollection()->PatternStart, GetClothCollection()->PatternEnd, GetElementIndex()); }
+		TArrayView<int32> GetSimFacesStart() { return GetClothCollection()->GetElements(GetClothCollection()->SimFacesStart, GetClothCollection()->PatternStart, GetClothCollection()->PatternEnd, GetElementIndex()); }
+		TArrayView<int32> GetSimFacesEnd() { return GetClothCollection()->GetElements(GetClothCollection()->SimFacesEnd, GetClothCollection()->PatternStart, GetClothCollection()->PatternEnd, GetElementIndex()); }
+		TArrayView<int32> GetRenderVerticesStart() { return GetClothCollection()->GetElements(GetClothCollection()->RenderVerticesStart, GetClothCollection()->PatternStart, GetClothCollection()->PatternEnd, GetElementIndex()); }
+		TArrayView<int32> GetRenderVerticesEnd() { return GetClothCollection()->GetElements(GetClothCollection()->RenderVerticesEnd, GetClothCollection()->PatternStart, GetClothCollection()->PatternEnd, GetElementIndex()); }
+		TArrayView<int32> GetRenderFacesStart() { return GetClothCollection()->GetElements(GetClothCollection()->RenderFacesStart, GetClothCollection()->PatternStart, GetClothCollection()->PatternEnd, GetElementIndex()); }
+		TArrayView<int32> GetRenderFacesEnd() { return GetClothCollection()->GetElements(GetClothCollection()->RenderFacesEnd, GetClothCollection()->PatternStart, GetClothCollection()->PatternEnd, GetElementIndex()); }
+
+		// Seams Group
+		int32 SetNumSeams(int32 NumSeams);
+		TArrayView<FIntVector2> GetSeamPatterns() { return GetClothCollection()->GetElements(GetClothCollection()->SeamPatterns, GetClothCollection()->SeamStart, GetClothCollection()->SeamEnd, GetElementIndex()); }
+		TArrayView<TArray<FIntVector2>> GetSeamStitches() { return GetClothCollection()->GetElements(GetClothCollection()->SeamStitches, GetClothCollection()->SeamStart, GetClothCollection()->SeamEnd, GetElementIndex()); }
+
+		// All patterns Sim Vertices Group, use these to access all pattern at once, and when using the SimIndices
+		TArrayView<FVector2f> GetPatternsSimPosition() { return GetClothCollection()->GetPatternsElements(GetClothCollection()->SimPosition, GetClothCollection()->SimVerticesStart, GetClothCollection()->SimVerticesEnd, GetElementIndex()); }
+		TArrayView<FVector3f> GetPatternsSimRestPosition() { return GetClothCollection()->GetPatternsElements(GetClothCollection()->SimRestPosition, GetClothCollection()->SimVerticesStart, GetClothCollection()->SimVerticesEnd, GetElementIndex()); }
+		TArrayView<FVector3f> GetPatternsSimRestNormal() { return GetClothCollection()->GetPatternsElements(GetClothCollection()->SimRestNormal, GetClothCollection()->SimVerticesStart, GetClothCollection()->SimVerticesEnd, GetElementIndex()); }
+
+		// All patterns Sim Faces Group, use these to access all pattern at once
+		TArrayView<FIntVector3> GetPatternsSimIndices() { return GetClothCollection()->GetPatternsElements(GetClothCollection()->SimIndices, GetClothCollection()->SimFacesStart, GetClothCollection()->SimFacesEnd, GetElementIndex()); }
+
+		// All patterns Render Vertices Group, use these to access all pattern at once, and when using the RenderIndices
+		TArrayView<FVector3f> GetPatternsRenderPosition() { return GetClothCollection()->GetPatternsElements(GetClothCollection()->RenderPosition, GetClothCollection()->RenderVerticesStart, GetClothCollection()->RenderVerticesEnd, GetElementIndex()); }
+		TArrayView<FVector3f> GetPatternsRenderNormal() { return GetClothCollection()->GetPatternsElements(GetClothCollection()->RenderNormal, GetClothCollection()->RenderVerticesStart, GetClothCollection()->RenderVerticesEnd, GetElementIndex()); }
+		TArrayView<FVector3f> GetPatternsRenderTangentU() { return GetClothCollection()->GetPatternsElements(GetClothCollection()->RenderTangentU, GetClothCollection()->RenderVerticesStart, GetClothCollection()->RenderVerticesEnd, GetElementIndex()); }
+		TArrayView<FVector3f> GetPatternsRenderTangentV() { return GetClothCollection()->GetPatternsElements(GetClothCollection()->RenderTangentV, GetClothCollection()->RenderVerticesStart, GetClothCollection()->RenderVerticesEnd, GetElementIndex()); }
+		TArrayView<TArray<FVector2f>> GePatternstRenderUVs() { return GetClothCollection()->GetPatternsElements(GetClothCollection()->RenderUVs, GetClothCollection()->RenderVerticesStart, GetClothCollection()->RenderVerticesEnd, GetElementIndex()); }
+		TArrayView<FLinearColor> GetPatternsRenderColor() { return GetClothCollection()->GetPatternsElements(GetClothCollection()->RenderColor, GetClothCollection()->RenderVerticesStart, GetClothCollection()->RenderVerticesEnd, GetElementIndex()); }
+
+		// All patterns Render Faces Group, use these to access all pattern at once
+		TArrayView<FIntVector3> GetPatternsRenderIndices() { return GetClothCollection()->GetPatternsElements(GetClothCollection()->RenderIndices, GetClothCollection()->RenderFacesStart, GetClothCollection()->RenderFacesEnd, GetElementIndex()); }
+		TArrayView<int32> GetPatternsRenderMaterialIndex() { return GetClothCollection()->GetPatternsElements(GetClothCollection()->RenderMaterialIndex, GetClothCollection()->RenderFacesStart, GetClothCollection()->RenderFacesEnd, GetElementIndex()); }
 
 	private:
 		friend class FClothAdapter;
 
 		void SetDefaults();
-
-		template<typename T>
-		TArrayView<T> GetPatternsElements(TManagedArray<T>& ElementArray, const TManagedArray<int32>& StartArray, const TManagedArray<int32>& EndArray);
 	};
 }  // End namespace UE::Chaos::ClothAsset
diff --git a/Engine/Plugins/Experimental/ChaosClothAsset/Source/ChaosClothAsset/Public/ChaosClothAsset/ClothPatternAdapter.h b/Engine/Plugins/Experimental/ChaosClothAsset/Source/ChaosClothAsset/Public/ChaosClothAsset/ClothPatternAdapter.h
--- a/Engine/Plugins/Experimental/ChaosClothAsset/Source/ChaosClothAsset/Public/ChaosClothAsset/ClothPatternAdapter.h
+++ b/Engine/Plugins/Experimental/ChaosClothAsset/Source/ChaosClothAsset/Public/ChaosClothAsset/ClothPatternAdapter.h
@@ -22,32 +22,32 @@ namespace UE::Chaos::ClothAsset
 
 		// Sim Vertices Group
 		// Note: Use the FClothLodConstAdapter accessors instead of these for the array indices to match the SimIndices values
-		int32 GetNumSimVertices() const { return GetNumElements(GetClothCollection()->SimVerticesStart, GetClothCollection()->SimVerticesEnd); }
-		TConstArrayView<FVector2f> GetSimPosition() const { return GetElements(GetClothCollection()->SimPosition, GetClothCollection()->SimVerticesStart, GetClothCollection()->SimVerticesEnd); }
-		TConstArrayView<FVector3f> GetSimRestPosition() const { return GetElements(GetClothCollection()->SimRestPosition, GetClothCollection()->SimVerticesStart, GetClothCollection()->SimVerticesEnd); }
-		TConstArrayView<FVector3f> GetSimRestNormal() const { return GetElements(GetClothCollection()->SimRestNormal, GetClothCollection()->SimVerticesStart, GetClothCollection()->SimVerticesEnd); }
+		int32 GetNumSimVertices() const { return GetClothCollection()->GetNumElements(GetClothCollection()->SimVerticesStart, GetClothCollection()->SimVerticesEnd, GetElementIndex()); }
+		TConstArrayView<FVector2f> GetSimPosition() const { return GetClothCollection()->GetElements(GetClothCollection()->SimPosition, GetClothCollection()->SimVerticesStart, GetClothCollection()->SimVerticesEnd, GetElementIndex()); }
+		TConstArrayView<FVector3f> GetSimRestPosition() const { return GetClothCollection()->GetElements(GetClothCollection()->SimRestPosition, GetClothCollection()->SimVerticesStart, GetClothCollection()->SimVerticesEnd, GetElementIndex()); }
+		TConstArrayView<FVector3f> GetSimRestNormal() const { return GetClothCollection()->GetElements(GetClothCollection()->SimRestNormal, GetClothCollection()->SimVerticesStart, GetClothCollection()->SimVerticesEnd, GetElementIndex()); }
 
 		// Sim Faces Group, note: SimIndices points to the LOD arrays, not the pattern arrays
-		int32 GetNumSimFaces() const { return GetNumElements(GetClothCollection()->SimFacesStart, GetClothCollection()->SimFacesEnd); }
-		TConstArrayView<FIntVector3> GetSimIndices() const { return GetElements(GetClothCollection()->SimIndices, GetClothCollection()->SimFacesStart, GetClothCollection()->SimFacesEnd); }
+		int32 GetNumSimFaces() const { return GetClothCollection()->GetNumElements(GetClothCollection()->SimFacesStart, GetClothCollection()->SimFacesEnd, GetElementIndex()); }
+		TConstArrayView<FIntVector3> GetSimIndices() const { return GetClothCollection()->GetElements(GetClothCollection()->SimIndices, GetClothCollection()->SimFacesStart, GetClothCollection()->SimFacesEnd, GetElementIndex()); }
 
 		// Render Vertices Group
 		// Note: Use the FClothLodConstAdapter accessors instead of these for the array indices to match the RenderIndices values
-		int32 GetNumRenderVertices() const { return GetNumElements(GetClothCollection()->RenderVerticesStart, GetClothCollection()->RenderVerticesEnd); }
-		TConstArrayView<FVector3f> GetRenderPosition() const { return GetElements(GetClothCollection()->RenderPosition, GetClothCollection()->RenderVerticesStart, GetClothCollection()->RenderVerticesEnd); }
-		TConstArrayView<FVector3f> GetRenderNormal() const { return GetElements(GetClothCollection()->RenderNormal, GetClothCollection()->RenderVerticesStart, GetClothCollection()->RenderVerticesEnd); }
-		TConstArrayView<FVector3f> GetRenderTangentU() const { return GetElements(GetClothCollection()->RenderTangentU, GetClothCollection()->RenderVerticesStart, GetClothCollection()->RenderVerticesEnd); }
-		TConstArrayView<FVector3f> GetRenderTangentV() const { return GetElements(GetClothCollection()->RenderTangentV, GetClothCollection()->RenderVerticesStart, GetClothCollection()->RenderVerticesEnd); }
-		TConstArrayView<TArray<FVector2f>> GetRenderUVs() const { return GetElements(GetClothCollection()->RenderUVs, GetClothCollection()->RenderVerticesStart, GetClothCollection()->RenderVerticesEnd); }
-		TConstArrayView<FLinearColor> GetRenderColor() const { return GetElements(GetClothCollection()->RenderColor, GetClothCollection()->RenderVerticesStart, GetClothCollection()->RenderVerticesEnd); }
+		int32 GetNumRenderVertices() const { return ClothCollection->GetNumElements(GetClothCollection()->RenderVerticesStart, GetClothCollection()->RenderVerticesEnd, GetElementIndex()); }
+		TConstArrayView<FVector3f> GetRenderPosition() const { return GetClothCollection()->GetElements(GetClothCollection()->RenderPosition, GetClothCollection()->RenderVerticesStart, GetClothCollection()->RenderVerticesEnd, GetElementIndex()); }
+		TConstArrayView<FVector3f> GetRenderNormal() const { return GetClothCollection()->GetElements(GetClothCollection()->RenderNormal, GetClothCollection()->RenderVerticesStart, GetClothCollection()->RenderVerticesEnd, GetElementIndex()); }
+		TConstArrayView<FVector3f> GetRenderTangentU() const { return GetClothCollection()->GetElements(GetClothCollection()->RenderTangentU, GetClothCollection()->RenderVerticesStart, GetClothCollection()->RenderVerticesEnd, GetElementIndex()); }
+		TConstArrayView<FVector3f> GetRenderTangentV() const { return GetClothCollection()->GetElements(GetClothCollection()->RenderTangentV, GetClothCollection()->RenderVerticesStart, GetClothCollection()->RenderVerticesEnd, GetElementIndex()); }
+		TConstArrayView<TArray<FVector2f>> GetRenderUVs() const { return GetClothCollection()->GetElements(GetClothCollection()->RenderUVs, GetClothCollection()->RenderVerticesStart, GetClothCollection()->RenderVerticesEnd, GetElementIndex()); }
+		TConstArrayView<FLinearColor> GetRenderColor() const { return GetClothCollection()->GetElements(GetClothCollection()->RenderColor, GetClothCollection()->RenderVerticesStart, GetClothCollection()->RenderVerticesEnd, GetElementIndex()); }
 
 		// Render Faces Group, note: RenderIndices points to the LOD arrays, not the pattern arrays
-		int32 GetNumRenderFaces() const { return GetNumElements(GetClothCollection()->RenderFacesStart, GetClothCollection()->RenderFacesEnd); }
-		TConstArrayView<FIntVector3> GetRenderIndices() const { return GetElements(GetClothCollection()->RenderIndices, GetClothCollection()->RenderFacesStart, GetClothCollection()->RenderFacesEnd); }
-		TConstArrayView<int32> GetRenderMaterialIndex() const { return GetElements(GetClothCollection()->RenderMaterialIndex, GetClothCollection()->RenderFacesStart, GetClothCollection()->RenderFacesEnd); }
+		int32 GetNumRenderFaces() const { return GetClothCollection()->GetNumElements(GetClothCollection()->RenderFacesStart, GetClothCollection()->RenderFacesEnd, GetElementIndex()); }
+		TConstArrayView<FIntVector3> GetRenderIndices() const { return GetClothCollection()->GetElements(GetClothCollection()->RenderIndices, GetClothCollection()->RenderFacesStart, GetClothCollection()->RenderFacesEnd, GetElementIndex()); }
+		TConstArrayView<int32> GetRenderMaterialIndex() const { return GetClothCollection()->GetElements(GetClothCollection()->RenderMaterialIndex, GetClothCollection()->RenderFacesStart, GetClothCollection()->RenderFacesEnd, GetElementIndex()); }
 
 		// Wrap Deformers Group
-		int32 GetNumWrapDeformers() const { return GetNumElements(GetClothCollection()->WrapDeformerStart, GetClothCollection()->WrapDeformerEnd); }
+		int32 GetNumWrapDeformers() const { return GetClothCollection()->GetNumElements(GetClothCollection()->WrapDeformerStart, GetClothCollection()->WrapDeformerEnd, GetElementIndex()); }
 
 		/** Return the element index within the cloth collection. */
 		int32 GetElementIndex() const { return ClothCollection->PatternStart[LodIndex] + PatternIndex; }
@@ -69,10 +69,6 @@ namespace UE::Chaos::ClothAsset
 		const TSharedPtr<const FClothCollection>& GetClothCollection() const { return ClothCollection; }
 
 	private:
-		int32 GetNumElements(const TManagedArray<int32>& StartArray, const TManagedArray<int32>& EndArray) const;
-		template<typename T>
-		TConstArrayView<T> GetElements(const TManagedArray<T>& ElementArray, const TManagedArray<int32>& StartArray, const TManagedArray<int32>& EndArray) const;
-
 		TSharedPtr<const FClothCollection> ClothCollection;
 		int32 LodIndex;
 		int32 PatternIndex;
@@ -81,7 +77,7 @@ namespace UE::Chaos::ClothAsset
 	/**
 	 * Cloth pattern adapter object to provide a more convenient object oriented access to the cloth collection.
 	 */
-	class CHAOSCLOTHASSET_API FClothPatternAdapter: public FClothPatternConstAdapter
+	class CHAOSCLOTHASSET_API FClothPatternAdapter final : public FClothPatternConstAdapter
 	{
 	public:
 		FClothPatternAdapter(const TSharedPtr<FClothCollection>& InClothCollection, int32 InLodIndex, int32 InPatternIndex);
@@ -112,25 +108,25 @@ namespace UE::Chaos::ClothAsset
 
 		// Sim Vertices Group
 		// Note: Use the FClothLodAdapter accessors instead of these for the array indices to match the SimIndices values
-		TArrayView<FVector2f> GetSimPosition() { return GetElements(GetClothCollection()->SimPosition, GetClothCollection()->SimVerticesStart, GetClothCollection()->SimVerticesEnd); }
-		TArrayView<FVector3f> GetSimRestPosition() { return GetElements(GetClothCollection()->SimRestPosition, GetClothCollection()->SimVerticesStart, GetClothCollection()->SimVerticesEnd); }
-		TArrayView<FVector3f> GetSimRestNormal() { return GetElements(GetClothCollection()->SimRestNormal, GetClothCollection()->SimVerticesStart, GetClothCollection()->SimVerticesEnd); }
+		TArrayView<FVector2f> GetSimPosition() { return GetClothCollection()->GetElements(GetClothCollection()->SimPosition, GetClothCollection()->SimVerticesStart, GetClothCollection()->SimVerticesEnd, GetElementIndex()); }
+		TArrayView<FVector3f> GetSimRestPosition() { return GetClothCollection()->GetElements(GetClothCollection()->SimRestPosition, GetClothCollection()->SimVerticesStart, GetClothCollection()->SimVerticesEnd, GetElementIndex()); }
+		TArrayView<FVector3f> GetSimRestNormal() { return GetClothCollection()->GetElements(GetClothCollection()->SimRestNormal, GetClothCollection()->SimVerticesStart, GetClothCollection()->SimVerticesEnd, GetElementIndex()); }
 
 		// Sim Faces Group, note: SimIndices points to the LOD arrays, not the pattern arrays
-		TArrayView<FIntVector3> GetSimIndices() { return GetElements(GetClothCollection()->SimIndices, GetClothCollection()->SimFacesStart, GetClothCollection()->SimFacesEnd); }
+		TArrayView<FIntVector3> GetSimIndices() { return GetClothCollection()->GetElements(GetClothCollection()->SimIndices, GetClothCollection()->SimFacesStart, GetClothCollection()->SimFacesEnd, GetElementIndex()); }
 
 		// Render Vertices Group
 		// Note: Use the FClothLodAdapter accessors instead of these for the array indices to match the RenderIndices values
-		TArrayView<FVector3f> GetRenderPosition() { return GetElements(GetClothCollection()->RenderPosition, GetClothCollection()->RenderVerticesStart, GetClothCollection()->RenderVerticesEnd); }
-		TArrayView<FVector3f> GetRenderNormal() { return GetElements(GetClothCollection()->RenderNormal, GetClothCollection()->RenderVerticesStart, GetClothCollection()->RenderVerticesEnd); }
-		TArrayView<FVector3f> GetRenderTangentU() { return GetElements(GetClothCollection()->RenderTangentU, GetClothCollection()->RenderVerticesStart, GetClothCollection()->RenderVerticesEnd); }
-		TArrayView<FVector3f> GetRenderTangentV() { return GetElements(GetClothCollection()->RenderTangentV, GetClothCollection()->RenderVerticesStart, GetClothCollection()->RenderVerticesEnd); }
-		TArrayView<TArray<FVector2f>> GetRenderUVs() { return GetElements(GetClothCollection()->RenderUVs, GetClothCollection()->RenderVerticesStart, GetClothCollection()->RenderVerticesEnd); }
-		TArrayView<FLinearColor> GetRenderColor() { return GetElements(GetClothCollection()->RenderColor, GetClothCollection()->RenderVerticesStart, GetClothCollection()->RenderVerticesEnd); }
+		TArrayView<FVector3f> GetRenderPosition() { return GetClothCollection()->GetElements(GetClothCollection()->RenderPosition, GetClothCollection()->RenderVerticesStart, GetClothCollection()->RenderVerticesEnd, GetElementIndex()); }
+		TArrayView<FVector3f> GetRenderNormal() { return GetClothCollection()->GetElements(GetClothCollection()->RenderNormal, GetClothCollection()->RenderVerticesStart, GetClothCollection()->RenderVerticesEnd, GetElementIndex()); }
+		TArrayView<FVector3f> GetRenderTangentU() { return GetClothCollection()->GetElements(GetClothCollection()->RenderTangentU, GetClothCollection()->RenderVerticesStart, GetClothCollection()->RenderVerticesEnd, GetElementIndex()); }
+		TArrayView<FVector3f> GetRenderTangentV() { return GetClothCollection()->GetElements(GetClothCollection()->RenderTangentV, GetClothCollection()->RenderVerticesStart, GetClothCollection()->RenderVerticesEnd, GetElementIndex()); }
+		TArrayView<TArray<FVector2f>> GetRenderUVs() { return GetClothCollection()->GetElements(GetClothCollection()->RenderUVs, GetClothCollection()->RenderVerticesStart, GetClothCollection()->RenderVerticesEnd, GetElementIndex()); }
+		TArrayView<FLinearColor> GetRenderColor() { return GetClothCollection()->GetElements(GetClothCollection()->RenderColor, GetClothCollection()->RenderVerticesStart, GetClothCollection()->RenderVerticesEnd, GetElementIndex()); }
 
 		// Render Faces Group, note: RenderIndices points to the LOD arrays, not the pattern arrays
-		TArrayView<FIntVector3> GetRenderIndices() { return GetElements(GetClothCollection()->RenderIndices, GetClothCollection()->RenderFacesStart, GetClothCollection()->RenderFacesEnd); }
-		TArrayView<int32> GetRenderMaterialIndex() { return GetElements(GetClothCollection()->RenderMaterialIndex, GetClothCollection()->RenderFacesStart, GetClothCollection()->RenderFacesEnd); }
+		TArrayView<FIntVector3> GetRenderIndices() { return GetClothCollection()->GetElements(GetClothCollection()->RenderIndices, GetClothCollection()->RenderFacesStart, GetClothCollection()->RenderFacesEnd, GetElementIndex()); }
+		TArrayView<int32> GetRenderMaterialIndex() { return GetClothCollection()->GetElements(GetClothCollection()->RenderMaterialIndex, GetClothCollection()->RenderFacesStart, GetClothCollection()->RenderFacesEnd, GetElementIndex()); }
 
 		/**
 		 * Return a view on the specified vertex weightmap.
@@ -146,10 +142,5 @@ namespace UE::Chaos::ClothAsset
 		friend class FClothLodAdapter;
 
 		void SetDefaults();
-
-		int32 SetNumElements(int32 InNumElements, const FName& GroupName, TManagedArray<int32>& StartArray, TManagedArray<int32>& EndArray);
-
-		template<typename T>
-		TArrayView<T> GetElements(TManagedArray<T>& ElementArray, const TManagedArray<int32>& StartArray, const TManagedArray<int32>& EndArray) const;
 	};
 }  // End namespace UE::Chaos::ClothAsset
diff --git a/Engine/Plugins/Experimental/ChaosClothAsset/Source/ChaosClothAsset/Public/ChaosClothAsset/ClothPresetCollection.h b/Engine/Plugins/Experimental/ChaosClothAsset/Source/ChaosClothAsset/Public/ChaosClothAsset/ClothPresetCollection.h
--- a/Engine/Plugins/Experimental/ChaosClothAsset/Source/ChaosClothAsset/Public/ChaosClothAsset/ClothPresetCollection.h
+++ b/Engine/Plugins/Experimental/ChaosClothAsset/Source/ChaosClothAsset/Public/ChaosClothAsset/ClothPresetCollection.h
@@ -18,8 +18,8 @@ namespace UE::Chaos::ClothAsset
 		FClothPresetCollection();
 		FClothPresetCollection(FClothPresetCollection&) = delete;
 		FClothPresetCollection& operator=(const FClothPresetCollection&) = delete;
-		FClothPresetCollection(FClothPresetCollection&&) = default;
-		FClothPresetCollection& operator=(FClothPresetCollection&&) = default;
+		FClothPresetCollection(FClothPresetCollection&&) = delete;
+		FClothPresetCollection& operator=(FClothPresetCollection&&) = delete;
 
 		using Super::Serialize;
 		void Serialize(FArchive& Ar);
