ChaosCloth - Add simulation to the Cloth Component.

diff --git a/Engine/Plugins/Experimental/ChaosCloth/Source/Chaos/Private/ChaosCloth/ChaosClothingSimulationMesh.cpp b/Engine/Plugins/Experimental/ChaosCloth/Source/Chaos/Private/ChaosCloth/ChaosClothingSimulationMesh.cpp
--- a/Engine/Plugins/Experimental/ChaosCloth/Source/Chaos/Private/ChaosCloth/ChaosClothingSimulationMesh.cpp
+++ b/Engine/Plugins/Experimental/ChaosCloth/Source/Chaos/Private/ChaosCloth/ChaosClothingSimulationMesh.cpp
@@ -38,25 +38,20 @@ namespace Chaos
 
 PRAGMA_DISABLE_DEPRECATION_WARNINGS  // TODO: CHAOS_IS_CLOTHINGSIMULATIONMESH_ABSTRACT
 
-FClothingSimulationMesh::FClothingSimulationMesh(const USkinnedMeshComponent* InSkinnedMeshComponent)
+FClothingSimulationMesh::FClothingSimulationMesh()
 	: Asset(nullptr)
-	, SkinnedMeshComponent(InSkinnedMeshComponent)
+	, SkeletalMeshComponent(nullptr)
 {
 }
 
 FClothingSimulationMesh::FClothingSimulationMesh(const UClothingAssetCommon* InAsset, const USkeletalMeshComponent* InSkeletalMeshComponent)
 	: Asset(InAsset)
-	, SkinnedMeshComponent(InSkeletalMeshComponent)
+	, SkeletalMeshComponent(InSkeletalMeshComponent)
 {
 }
 
 FClothingSimulationMesh::~FClothingSimulationMesh() = default;
 
-const USkeletalMeshComponent* FClothingSimulationMesh::GetSkeletalMeshComponent() const
-{
-	return Cast<USkeletalMeshComponent>(SkinnedMeshComponent);
-}
-
 int32 FClothingSimulationMesh::GetNumLODs() const
 {
 	return Asset ? Asset->LodData.Num() : 0;
@@ -66,9 +61,9 @@ int32 FClothingSimulationMesh::GetLODIndex() const
 {
 	int32 LODIndex = INDEX_NONE;
 
-	if (Asset && SkinnedMeshComponent)
+	if (Asset && SkeletalMeshComponent)
 	{
-		const int32 MeshLODIndex = SkinnedMeshComponent->GetPredictedLODLevel();
+		const int32 MeshLODIndex = SkeletalMeshComponent->GetPredictedLODLevel();
 		if (Asset->LodMap.IsValidIndex(MeshLODIndex))
 		{
 			const int32 ClothLODIndex = Asset->LodMap[MeshLODIndex];
@@ -218,16 +213,14 @@ TConstArrayView<FClothVertBoneData> FClothingSimulationMesh::GetBoneData(int32 L
 	return EmptyArray;
 }
 
-const TArray<FMeshToMeshVertData>& FClothingSimulationMesh::GetTransitionUpSkinData(int32 LODIndex) const
+TConstArrayView<FMeshToMeshVertData> FClothingSimulationMesh::GetTransitionUpSkinData(int32 LODIndex) const
 {
-	static TArray<FMeshToMeshVertData> EmptyArray;
-	return IsValidLODIndex(LODIndex) ? Asset->LodData[LODIndex].TransitionUpSkinData : EmptyArray;
+	return IsValidLODIndex(LODIndex) ? TConstArrayView<FMeshToMeshVertData>(Asset->LodData[LODIndex].TransitionUpSkinData) : TConstArrayView<FMeshToMeshVertData>();
 }
 
-const TArray<FMeshToMeshVertData>& FClothingSimulationMesh::GetTransitionDownSkinData(int32 LODIndex) const
+TConstArrayView<FMeshToMeshVertData> FClothingSimulationMesh::GetTransitionDownSkinData(int32 LODIndex) const
 {
-	static TArray<FMeshToMeshVertData> EmptyArray;
-	return IsValidLODIndex(LODIndex) ? Asset->LodData[LODIndex].TransitionDownSkinData : EmptyArray;
+	return IsValidLODIndex(LODIndex) ? TConstArrayView<FMeshToMeshVertData>(Asset->LodData[LODIndex].TransitionDownSkinData) : TConstArrayView<FMeshToMeshVertData>();
 }
 
 const FClothingSimulationContextCommon* FClothingSimulationMesh::GetContext() const
@@ -240,10 +233,7 @@ const FClothingSimulationContextCommon* FClothingSimulationMesh::GetContext() co
 PRAGMA_ENABLE_DEPRECATION_WARNINGS
 
 #else
-FClothingSimulationMesh::FClothingSimulationMesh(const USkinnedMeshComponent* InSkinnedMeshComponent)
-	: SkinnedMeshComponent(InSkinnedMeshComponent)
-{
-}
+FClothingSimulationMesh::FClothingSimulationMesh() = default;
 
 FClothingSimulationMesh::~FClothingSimulationMesh() = default;
 #endif
@@ -270,7 +260,7 @@ bool FClothingSimulationMesh::WrapDeformLOD(
 	}
 
 	const int32 NumPoints = GetNumPoints(LODIndex);
-	const TArray<FMeshToMeshVertData>& SkinData = (PrevLODIndex < LODIndex) ?
+	const TConstArrayView<FMeshToMeshVertData> SkinData = (PrevLODIndex < LODIndex) ?
 		GetTransitionUpSkinData(LODIndex) :
 		GetTransitionDownSkinData(LODIndex);
 
@@ -311,7 +301,7 @@ bool FClothingSimulationMesh::WrapDeformLOD(
 	}
 
 	const int32 NumPoints = GetNumPoints(LODIndex);
-	const TArray<FMeshToMeshVertData>& SkinData = (PrevLODIndex < LODIndex) ?
+	const TConstArrayView<FMeshToMeshVertData> SkinData = (PrevLODIndex < LODIndex) ?
 		GetTransitionUpSkinData(LODIndex) :
 		GetTransitionDownSkinData(LODIndex);
 
@@ -431,7 +421,7 @@ void FClothingSimulationMesh::Update(
 	check(Solver);
 
 	// Exit if any inputs are missing or not ready, and if the LOD is invalid
-	if (!IsValidLODIndex(LODIndex) || !SkinnedMeshComponent)
+	if (!IsValidLODIndex(LODIndex))
 	{
 		return;
 	}
diff --git a/Engine/Plugins/Experimental/ChaosCloth/Source/Chaos/Private/ChaosCloth/ChaosClothingSimulationSkeletalMesh.cpp b/Engine/Plugins/Experimental/ChaosCloth/Source/Chaos/Private/ChaosCloth/ChaosClothingSimulationSkeletalMesh.cpp
--- a/Engine/Plugins/Experimental/ChaosCloth/Source/Chaos/Private/ChaosCloth/ChaosClothingSimulationSkeletalMesh.cpp
+++ b/Engine/Plugins/Experimental/ChaosCloth/Source/Chaos/Private/ChaosCloth/ChaosClothingSimulationSkeletalMesh.cpp
@@ -10,9 +10,10 @@ namespace Chaos
 {
 	FClothingSimulationSkeletalMesh::FClothingSimulationSkeletalMesh(const UClothingAssetCommon* InAsset, const USkeletalMeshComponent* InSkeletalMeshComponent)
 PRAGMA_DISABLE_DEPRECATION_WARNINGS  // TODO: CHAOS_IS_CLOTHINGSIMULATIONMESH_ABSTRACT
-		: FClothingSimulationMesh(InSkeletalMeshComponent)
+		: FClothingSimulationMesh()
 PRAGMA_ENABLE_DEPRECATION_WARNINGS
 		, Asset(InAsset)
+		, SkeletalMeshComponent(InSkeletalMeshComponent)
 	{
 	}
 
@@ -25,11 +26,9 @@ PRAGMA_ENABLE_DEPRECATION_WARNINGS
 	{
 		int32 LODIndex = INDEX_NONE;
 
-PRAGMA_DISABLE_DEPRECATION_WARNINGS  // TODO: CHAOS_IS_CLOTHINGSIMULATIONMESH_ABSTRACT
-		if (Asset && GetSkinnedMeshComponent())
+		if (Asset && SkeletalMeshComponent)
 		{
-			const int32 OwnerLODIndex = GetSkinnedMeshComponent()->GetPredictedLODLevel();
-PRAGMA_ENABLE_DEPRECATION_WARNINGS
+			const int32 OwnerLODIndex = SkeletalMeshComponent->GetPredictedLODLevel();
 			if (Asset->LodMap.IsValidIndex(OwnerLODIndex))
 			{
 				const int32 ClothLODIndex = Asset->LodMap[OwnerLODIndex];
@@ -178,16 +177,14 @@ PRAGMA_ENABLE_DEPRECATION_WARNINGS
 		return TConstArrayView<FClothVertBoneData>();
 	}
 
-	const TArray<FMeshToMeshVertData>& FClothingSimulationSkeletalMesh::GetTransitionUpSkinData(int32 LODIndex) const
+	TConstArrayView<FMeshToMeshVertData> FClothingSimulationSkeletalMesh::GetTransitionUpSkinData(int32 LODIndex) const
 	{
-		static TArray<FMeshToMeshVertData> EmptyArray;
-		return IsValidLODIndex(LODIndex) ? Asset->LodData[LODIndex].TransitionUpSkinData : EmptyArray;
+		return IsValidLODIndex(LODIndex) ? TConstArrayView<FMeshToMeshVertData>(Asset->LodData[LODIndex].TransitionUpSkinData) : TConstArrayView<FMeshToMeshVertData>();
 	}
 
-	const TArray<FMeshToMeshVertData>& FClothingSimulationSkeletalMesh::GetTransitionDownSkinData(int32 LODIndex) const
+	TConstArrayView<FMeshToMeshVertData> FClothingSimulationSkeletalMesh::GetTransitionDownSkinData(int32 LODIndex) const
 	{
-		static TArray<FMeshToMeshVertData> EmptyArray;
-		return IsValidLODIndex(LODIndex) ? Asset->LodData[LODIndex].TransitionDownSkinData : EmptyArray;
+		return IsValidLODIndex(LODIndex) ? TConstArrayView<FMeshToMeshVertData>(Asset->LodData[LODIndex].TransitionDownSkinData) : TConstArrayView<FMeshToMeshVertData>();
 	}
 
 	const FClothingSimulationContextCommon* FClothingSimulationSkeletalMesh::GetContext() const
diff --git a/Engine/Plugins/Experimental/ChaosCloth/Source/Chaos/Public/ChaosCloth/ChaosClothingSimulationMesh.h b/Engine/Plugins/Experimental/ChaosCloth/Source/Chaos/Public/ChaosCloth/ChaosClothingSimulationMesh.h
--- a/Engine/Plugins/Experimental/ChaosCloth/Source/Chaos/Public/ChaosCloth/ChaosClothingSimulationMesh.h
+++ b/Engine/Plugins/Experimental/ChaosCloth/Source/Chaos/Public/ChaosCloth/ChaosClothingSimulationMesh.h
@@ -8,7 +8,6 @@
 
 //#define CHAOS_IS_CLOTHINGSIMULATIONMESH_ABSTRACT 1  // For testing only. TODO: Once FClothingSimulationMesh is abstract from 5.4, remove all lines where this is referenced
 
-class USkinnedMeshComponent;
 struct FMeshToMeshVertData;
 struct FClothVertBoneData;
 
@@ -31,7 +30,7 @@ namespace Chaos
 #endif
 	{
 	public:
-		FClothingSimulationMesh(const USkinnedMeshComponent* InSkinnedMeshComponent);
+		FClothingSimulationMesh();
 		virtual ~FClothingSimulationMesh();
 
 		FClothingSimulationMesh(const FClothingSimulationMesh&) = delete;
@@ -39,16 +38,13 @@ namespace Chaos
 		FClothingSimulationMesh& operator=(const FClothingSimulationMesh&) = delete;
 		FClothingSimulationMesh& operator=(FClothingSimulationMesh&&) = delete;
 
-		// ---- Node property getters/setters
-		const USkinnedMeshComponent* GetSkinnedMeshComponent() const { return SkinnedMeshComponent; }
-		// ---- End of node property getters/setters
 #if !defined(CHAOS_IS_CLOTHINGSIMULATIONMESH_ABSTRACT) || !CHAOS_IS_CLOTHINGSIMULATIONMESH_ABSTRACT
 		UE_DEPRECATED(5.2, "Use the FClothingSimulationSkeletalMesh constructor instead.")
 		FClothingSimulationMesh(const UClothingAssetCommon* InAsset, const USkeletalMeshComponent* InSkeletalMeshComponent);
 		UE_DEPRECATED(5.2, "Use the FClothingSimulationSkeletalMesh::GetAsset() instead.")
 		virtual const UClothingAssetCommon* GetAsset() const { return Asset; }
 		UE_DEPRECATED(5.2, "Use the FClothingSimulationSkeletalMesh::GetSkeletalMeshComponent() instead.")
-		virtual const USkeletalMeshComponent* GetSkeletalMeshComponent() const;
+		virtual const USkeletalMeshComponent* GetSkeletalMeshComponent() const { return SkeletalMeshComponent; }
 
 		virtual int32 GetNumLODs() const;
 		virtual int32 GetLODIndex() const;
@@ -66,8 +62,8 @@ namespace Chaos
 		virtual const TArray<FMatrix44f>& GetRefToLocalMatrices() const;
 		virtual TConstArrayView<int32> GetBoneMap() const;
 		virtual TConstArrayView<FClothVertBoneData> GetBoneData(int32 LODIndex) const;
-		virtual const TArray<FMeshToMeshVertData>& GetTransitionUpSkinData(int32 LODIndex) const;
-		virtual const TArray<FMeshToMeshVertData>& GetTransitionDownSkinData(int32 LODIndex) const;
+		virtual TConstArrayView<FMeshToMeshVertData> GetTransitionUpSkinData(int32 LODIndex) const;
+		virtual TConstArrayView<FMeshToMeshVertData> GetTransitionDownSkinData(int32 LODIndex) const;
 #else
 		// ---- Cloth interface ----
 		/* Return the number of LODs on this mesh. */
@@ -119,10 +115,10 @@ namespace Chaos
 		virtual TConstArrayView<FClothVertBoneData> GetBoneData(int32 LODIndex) const = 0;
 
 		/* Return the transition up data (PrevLODIndex < LODIndex), for matching shapes during LOD changes. */
-		virtual const TArray<FMeshToMeshVertData>& GetTransitionUpSkinData(int32 LODIndex) const = 0;
+		virtual TConstArrayView<FMeshToMeshVertData> GetTransitionUpSkinData(int32 LODIndex) const = 0;
 
 		/* Return the transition down data (PrevLODIndex > LODIndex), for matching shapes during LOD changes. */
-		virtual const TArray<FMeshToMeshVertData>& GetTransitionDownSkinData(int32 LODIndex) const = 0;
+		virtual TConstArrayView<FMeshToMeshVertData> GetTransitionDownSkinData(int32 LODIndex) const = 0;
 #endif
 		/** Return this mesh uniform scale as the maximum of the three axis scale value. */
 		virtual Softs::FSolverReal GetScale() const;
@@ -162,13 +158,12 @@ namespace Chaos
 			Softs::FSolverVec3* OutPositions,
 			Softs::FSolverVec3* OutNormals) const;
 
-	private:
 #if !defined(CHAOS_IS_CLOTHINGSIMULATIONMESH_ABSTRACT) || !CHAOS_IS_CLOTHINGSIMULATIONMESH_ABSTRACT
 		const FClothingSimulationContextCommon* GetContext() const;
 
 		const UClothingAssetCommon* Asset;
+		const USkeletalMeshComponent* SkeletalMeshComponent;
 #endif
-		const USkinnedMeshComponent* SkinnedMeshComponent;
 	};
 } // namespace Chaos
 
diff --git a/Engine/Plugins/Experimental/ChaosCloth/Source/Chaos/Public/ChaosCloth/ChaosClothingSimulationSkeletalMesh.h b/Engine/Plugins/Experimental/ChaosCloth/Source/Chaos/Public/ChaosCloth/ChaosClothingSimulationSkeletalMesh.h
--- a/Engine/Plugins/Experimental/ChaosCloth/Source/Chaos/Public/ChaosCloth/ChaosClothingSimulationSkeletalMesh.h
+++ b/Engine/Plugins/Experimental/ChaosCloth/Source/Chaos/Public/ChaosCloth/ChaosClothingSimulationSkeletalMesh.h
@@ -26,7 +26,7 @@ PRAGMA_ENABLE_DEPRECATION_WARNINGS
 #if !defined(CHAOS_IS_CLOTHINGSIMULATIONMESH_ABSTRACT) || !CHAOS_IS_CLOTHINGSIMULATIONMESH_ABSTRACT
 		virtual const UClothingAssetCommon* GetAsset() const override { return Asset; }
 PRAGMA_DISABLE_DEPRECATION_WARNINGS  // TODO: CHAOS_IS_CLOTHINGSIMULATIONMESH_ABSTRACT
-		virtual const USkeletalMeshComponent* GetSkeletalMeshComponent() const override { return Cast<USkeletalMeshComponent>(GetSkinnedMeshComponent()); }
+		virtual const USkeletalMeshComponent* GetSkeletalMeshComponent() const override { return SkeletalMeshComponent; }
 PRAGMA_ENABLE_DEPRECATION_WARNINGS
 #else
 		const UClothingAssetCommon* GetAsset() const { return Asset; }
@@ -49,13 +49,14 @@ PRAGMA_ENABLE_DEPRECATION_WARNINGS
 		virtual const TArray<FMatrix44f>& GetRefToLocalMatrices() const override;
 		virtual TConstArrayView<int32> GetBoneMap() const override;
 		virtual TConstArrayView<FClothVertBoneData> GetBoneData(int32 LODIndex) const override;
-		virtual const TArray<FMeshToMeshVertData>& GetTransitionUpSkinData(int32 LODIndex) const override;
-		virtual const TArray<FMeshToMeshVertData>& GetTransitionDownSkinData(int32 LODIndex) const override;
+		virtual TConstArrayView<FMeshToMeshVertData> GetTransitionUpSkinData(int32 LODIndex) const override;
+		virtual TConstArrayView<FMeshToMeshVertData> GetTransitionDownSkinData(int32 LODIndex) const override;
 		//~ End FClothingSimulationMesh Interface
 
 	private:
 		const FClothingSimulationContextCommon* GetContext() const;
 
 		const UClothingAssetCommon* Asset;
+		const USkeletalMeshComponent* SkeletalMeshComponent;
 	};
 }
diff --git a/Engine/Plugins/Experimental/ChaosClothAsset/ChaosClothAsset.uplugin b/Engine/Plugins/Experimental/ChaosClothAsset/ChaosClothAsset.uplugin
--- a/Engine/Plugins/Experimental/ChaosClothAsset/ChaosClothAsset.uplugin
+++ b/Engine/Plugins/Experimental/ChaosClothAsset/ChaosClothAsset.uplugin
@@ -25,5 +25,11 @@
 			"Type": "Runtime",
 			"LoadingPhase": "Default"
 		}
+	],
+	"Plugins": [
+		{
+			"Name": "ChaosCloth",
+			"Enabled": true
+		}
 	]
 }
diff --git a/Engine/Plugins/Experimental/ChaosClothAsset/Source/ChaosClothAssetEngine/ChaosClothAssetEngine.Build.cs b/Engine/Plugins/Experimental/ChaosClothAsset/Source/ChaosClothAssetEngine/ChaosClothAssetEngine.Build.cs
--- a/Engine/Plugins/Experimental/ChaosClothAsset/Source/ChaosClothAssetEngine/ChaosClothAssetEngine.Build.cs
+++ b/Engine/Plugins/Experimental/ChaosClothAsset/Source/ChaosClothAssetEngine/ChaosClothAssetEngine.Build.cs
@@ -17,7 +17,10 @@ public class ChaosClothAssetEngine : ModuleRules
 				"RenderCore",
 				"RHI",
 				"Chaos",
-				"ChaosClothAsset"
+				"ChaosClothAsset",
+				"ChaosCloth",
+				"TaskGraph",
+				"ClothingSystemRuntimeInterface"
 			}
 		);
 
diff --git a/Engine/Plugins/Experimental/ChaosClothAsset/Source/ChaosClothAssetEngine/Private/ChaosClothAsset/ChaosClothAssetEngineModule.cpp b/Engine/Plugins/Experimental/ChaosClothAsset/Source/ChaosClothAssetEngine/Private/ChaosClothAsset/ChaosClothAssetEngineModule.cpp
--- a/Engine/Plugins/Experimental/ChaosClothAsset/Source/ChaosClothAssetEngine/Private/ChaosClothAsset/ChaosClothAssetEngineModule.cpp
+++ b/Engine/Plugins/Experimental/ChaosClothAsset/Source/ChaosClothAssetEngine/Private/ChaosClothAsset/ChaosClothAssetEngineModule.cpp
@@ -1,9 +1,11 @@
 // Copyright Epic Games, Inc. All Rights Reserved.
 
-#include "CoreMinimal.h"
+#include "ChaosClothAsset/ClothAssetPrivate.h"
 #include "Modules/ModuleInterface.h"
 #include "Modules/ModuleManager.h"
 
+DEFINE_LOG_CATEGORY(LogChaosClothAsset);
+
 namespace UE::Chaos::ClothAsset
 {
 	class FClothAssetEngineModule : public IModuleInterface
diff --git a/Engine/Plugins/Experimental/ChaosClothAsset/Source/ChaosClothAssetEngine/Private/ChaosClothAsset/ClothAsset.cpp b/Engine/Plugins/Experimental/ChaosClothAsset/Source/ChaosClothAssetEngine/Private/ChaosClothAsset/ClothAsset.cpp
--- a/Engine/Plugins/Experimental/ChaosClothAsset/Source/ChaosClothAssetEngine/Private/ChaosClothAsset/ClothAsset.cpp
+++ b/Engine/Plugins/Experimental/ChaosClothAsset/Source/ChaosClothAssetEngine/Private/ChaosClothAsset/ClothAsset.cpp
@@ -4,6 +4,8 @@
 #include "ChaosClothAsset/ClothAssetBuilder.h"
 #include "ChaosClothAsset/ClothAdapter.h"
 #include "ChaosClothAsset/ClothGeometryTools.h"
+#include "ChaosClothAsset/ClothAssetPrivate.h"
+#include "ChaosClothAsset/ClothSimulationModel.h"
 #include "Animation/Skeleton.h"
 #include "Rendering/SkeletalMeshModel.h"
 #include "EngineUtils.h"
@@ -18,8 +20,6 @@
 // In case of merge conflicts with DDC versions, you MUST generate a new GUID and set this new GUID as the version.
 #define CHAOS_CLOTH_ASSET_DERIVED_DATA_VERSION TEXT("19823D996CA54F279B9A9FA8ED7A8EB6")
 
-DEFINE_LOG_CATEGORY_STATIC(LogChaosClothAsset, Log, All);
-
 UChaosClothAsset::UChaosClothAsset(const FObjectInitializer& ObjectInitializer)
 	: Super(ObjectInitializer)
 	, DisableBelowMinLodStripping(FPerPlatformBool(false))
@@ -30,6 +30,13 @@ UChaosClothAsset::UChaosClothAsset(const FObjectInitializer& ObjectInitializer)
 {
 }
 
+UChaosClothAsset::UChaosClothAsset(FVTableHelper& Helper)
+	: Super(Helper)
+{
+}
+
+UChaosClothAsset::~UChaosClothAsset() = default;
+
 FMatrix UChaosClothAsset::GetComposedRefPoseMatrix(FName InBoneName) const
 {
 	FMatrix LocalPose(FMatrix::Identity);
@@ -92,7 +99,9 @@ void UChaosClothAsset::BeginPostLoadInternal(FSkinnedAssetPostLoadContext& Conte
 	// This scope allows us to use any locked properties without causing stalls
 	FSkinnedAssetAsyncBuildScope AsyncBuildScope(this);
 
-	BuildModel();
+	BuildClothSimulationModel();  // TODO: Cache ClothSimulationModel?
+
+	BuildMeshModel();
 #endif // #if WITH_EDITOR
 }
 
@@ -273,6 +282,23 @@ void UChaosClothAsset::CalculateBounds()
 	Bounds = FBoxSphereBounds(BoundingBox);
 }
 
+void UChaosClothAsset::UpdateSkeleton(bool bRebuildClothSimulationModel)
+{
+	// Rebuild Skeleton
+	RefSkeleton = Skeleton->GetReferenceSkeleton();  // TODO: Does a copy work? Try to see what's exactly needed here...
+
+	constexpr bool bRebuildNameMap = false;
+	RefSkeleton.RebuildRefSkeleton(Skeleton, bRebuildNameMap);
+
+	CalculateInvRefMatrices();
+
+	if (bRebuildClothSimulationModel)
+	{
+		// Rebuild simulation model  // TODO: How does this work with skinning
+		BuildClothSimulationModel();
+	}
+}
+
 void UChaosClothAsset::Build()
 {
 	using namespace UE::Chaos::ClothAsset;
@@ -284,12 +310,10 @@ void UChaosClothAsset::Build()
 	AssetGuid = FGuid::NewGuid();
 
 	// Rebuild Skeleton
-	RefSkeleton = Skeleton->GetReferenceSkeleton();  // TODO: Does a copy work? Try to see what's exactly needed here...
-
-	constexpr bool bRebuildNameMap = false;
-	RefSkeleton.RebuildRefSkeleton(Skeleton, bRebuildNameMap);
+	constexpr bool bRebuildClothSimulationModel = false;
+	UpdateSkeleton(bRebuildClothSimulationModel);
 
-	CalculateInvRefMatrices();
+	// Update bounds
 	CalculateBounds();
 
 	// Add LODs to the render data
@@ -300,9 +324,12 @@ void UChaosClothAsset::Build()
 	LODInfo.Reset(NumLods);
 	LODInfo.AddDefaulted(NumLods);  // TODO: Expose some properties to fill up the LOD infos
 
+	// Build simulation model
+	BuildClothSimulationModel();
+
 	// Rebuild LOD Model
 #if WITH_EDITORONLY_DATA
-	BuildModel();
+	BuildMeshModel();
 #endif
 
 	// Load/save render data from/to DDC
@@ -317,7 +344,7 @@ void UChaosClothAsset::Build()
 }
 
 #if WITH_EDITORONLY_DATA
-void UChaosClothAsset::BuildModel()
+void UChaosClothAsset::BuildMeshModel()
 {
 	const TArray<IClothAssetBuilderClassProvider*> ClassProviders = IModularFeatures::Get().GetModularFeatureImplementations<IClothAssetBuilderClassProvider>(IClothAssetBuilderClassProvider::FeatureName);
 	if (const TSubclassOf<UClothAssetBuilder> ClothAssetBuilderClass = ClassProviders.Num() ? ClassProviders[0]->GetClothAssetBuilderClass() : nullptr)
@@ -343,6 +370,11 @@ void UChaosClothAsset::BuildModel()
 }
 #endif  // #if WITH_EDITORONLY_DATA
 
+void UChaosClothAsset::BuildClothSimulationModel()
+{
+	ClothSimulationModel = MakeUnique<FChaosClothSimulationModel>(GetClothCollection(), GetRefSkeleton());
+}
+
 const FMeshUVChannelInfo* UChaosClothAsset::GetUVChannelData(int32 MaterialIndex) const
 {
 	if (GetMaterials().IsValidIndex(MaterialIndex))
diff --git a/Engine/Plugins/Experimental/ChaosClothAsset/Source/ChaosClothAssetEngine/Private/ChaosClothAsset/ClothAssetBuilder.cpp b/Engine/Plugins/Experimental/ChaosClothAsset/Source/ChaosClothAssetEngine/Private/ChaosClothAsset/ClothAssetBuilder.cpp
new file mode 100644
--- /dev/null
+++ b/Engine/Plugins/Experimental/ChaosClothAsset/Source/ChaosClothAssetEngine/Private/ChaosClothAsset/ClothAssetBuilder.cpp
@@ -0,0 +1,30 @@
+// Copyright Epic Games, Inc. All Rights Reserved.
+
+#include "ChaosClothAsset/ClothAssetBuilder.h"
+#include "ChaosClothAsset/ClothAsset.h"
+#include "ChaosClothAsset/ClothSimulationModel.h"
+
+int32 UClothAssetBuilder::GetNumVertices(const UChaosClothAsset& ClothAsset, int32 LodIndex)
+{
+	const FChaosClothSimulationModel* const ClothSimulationModel = ClothAsset.GetClothSimulationModel();
+	return ClothSimulationModel && ClothSimulationModel->IsValidLodIndex(LodIndex) ?
+		ClothSimulationModel->GetNumVertices(LodIndex) :
+		0;
+}
+
+TConstArrayView<FVector3f> UClothAssetBuilder::GetSimPositions(const UChaosClothAsset& ClothAsset, int32 LodIndex)
+{
+	const FChaosClothSimulationModel* const ClothSimulationModel = ClothAsset.GetClothSimulationModel();
+	return ClothSimulationModel && ClothSimulationModel->IsValidLodIndex(LodIndex) ?
+		ClothSimulationModel->GetPositions(LodIndex) :
+		TConstArrayView<FVector3f>();
+}
+
+TConstArrayView<uint32> UClothAssetBuilder::GetSimIndices(const UChaosClothAsset& ClothAsset, int32 LodIndex)
+{
+	const FChaosClothSimulationModel* const ClothSimulationModel = ClothAsset.GetClothSimulationModel();
+	return ClothSimulationModel && ClothSimulationModel->IsValidLodIndex(LodIndex) ?
+		ClothSimulationModel->GetIndices(LodIndex) :
+		TConstArrayView<uint32>();
+
+}
diff --git a/Engine/Plugins/Experimental/ChaosClothAsset/Source/ChaosClothAssetEngine/Private/ChaosClothAsset/ClothAssetPrivate.h b/Engine/Plugins/Experimental/ChaosClothAsset/Source/ChaosClothAssetEngine/Private/ChaosClothAsset/ClothAssetPrivate.h
new file mode 100644
--- /dev/null
+++ b/Engine/Plugins/Experimental/ChaosClothAsset/Source/ChaosClothAssetEngine/Private/ChaosClothAsset/ClothAssetPrivate.h
@@ -0,0 +1,9 @@
+// Copyright Epic Games, Inc. All Rights Reserved.
+
+#pragma once
+
+#include "Logging/LogMacros.h"
+#include "Stats/Stats.h"
+
+DECLARE_LOG_CATEGORY_EXTERN(LogChaosClothAsset, Log, All);
+DECLARE_STATS_GROUP(TEXT("ChaosClothAsset"), STATGROUP_ChaosClothAsset, STATCAT_Advanced);
diff --git a/Engine/Plugins/Experimental/ChaosClothAsset/Source/ChaosClothAssetEngine/Private/ChaosClothAsset/ClothComponent.cpp b/Engine/Plugins/Experimental/ChaosClothAsset/Source/ChaosClothAssetEngine/Private/ChaosClothAsset/ClothComponent.cpp
--- a/Engine/Plugins/Experimental/ChaosClothAsset/Source/ChaosClothAssetEngine/Private/ChaosClothAsset/ClothComponent.cpp
+++ b/Engine/Plugins/Experimental/ChaosClothAsset/Source/ChaosClothAssetEngine/Private/ChaosClothAsset/ClothComponent.cpp
@@ -2,24 +2,25 @@
 
 #include "ChaosClothAsset/ClothComponent.h"
 #include "ChaosClothAsset/ClothAsset.h"
+#include "ChaosClothAsset/ClothAssetPrivate.h"
+#include "ChaosClothAsset/ClothSimulationModel.h"
+#include "ChaosClothAsset/ClothSimulationProxy.h"
+#include "HAL/IConsoleManager.h"
 #include "SkeletalRenderPublic.h"
-#include "PhysicsEngine/PhysicsAsset.h"
 
-DEFINE_LOG_CATEGORY_STATIC(LogChaosClothComponent, Log, All);
-
-#include UE_INLINE_GENERATED_CPP_BY_NAME(ClothComponent)
-
-DECLARE_STATS_GROUP(TEXT("ChaosClothComponent"), STATGROUP_ChaosClothComponent, STATCAT_Advanced);
-DECLARE_CYCLE_STAT(TEXT("SendRenderDynamicData_Concurrent"), STAT_ChaosClothComponent_SendRenderDynamicData_Concurrent, STATGROUP_ChaosClothComponent);
-DECLARE_CYCLE_STAT(TEXT("MeshObjectUpdate"), STAT_ChaosClothComponent_MeshObjectUpdate, STATGROUP_ChaosClothComponent);
-DECLARE_CYCLE_STAT(TEXT("CalcBounds"), STAT_ChaosClothComponent_CalcBounds, STATGROUP_ChaosClothComponent);
+CSV_DECLARE_CATEGORY_MODULE_EXTERN(ENGINE_API, Animation);
 
 UChaosClothComponent::UChaosClothComponent(const FObjectInitializer& ObjectInitializer)
 	: Super(ObjectInitializer)
+	, bDisableClothSimulation(0)
+	, bSuspendSimulation(0)
+	, bWaitForParallelClothTask(0)
+	, bBindClothToLeaderComponent(0)
 {
 	bAutoActivate = true;
 	PrimaryComponentTick.bCanEverTick = true;
 	PrimaryComponentTick.TickGroup = TG_PrePhysics;
+	PrimaryComponentTick.EndTickGroup = TG_PostPhysics;
 
 	//VisibilityBasedAnimTickOption = EVisibilityBasedAnimTickOption::AlwaysTickPoseAndRefreshBones;
 
@@ -72,26 +73,124 @@ UChaosClothComponent::UChaosClothComponent(const FObjectInitializer& ObjectIniti
 	//CurrentSkinWeightProfileName = NAME_None;
 }
 
+UChaosClothComponent::UChaosClothComponent(FVTableHelper& Helper)
+	: Super(Helper)
+{
+}
+
+UChaosClothComponent::~UChaosClothComponent() = default;
+
 void UChaosClothComponent::OnRegister()
 {
+	using namespace UE::Chaos::ClothAsset;
+
 	LLM_SCOPE(ELLMTag::Chaos);
 
 	Super::OnRegister();
 
 	if (GetClothAsset())
 	{
-		FSkeletalMeshLODRenderData& LODData = GetClothAsset()->GetResourceForRendering()->LODRenderData[GetPredictedLODLevel()];
-		GetClothAsset()->FillComponentSpaceTransforms(GetClothAsset()->GetRefSkeleton().GetRefBonePose(), LODData.RequiredBones, GetEditableComponentSpaceTransforms());
+		const FChaosClothSimulationModel* const ClothSimulationModel = GetClothAsset()->GetClothSimulationModel();
+		if (ensure(ClothSimulationModel) && ClothSimulationModel->GetNumLods())
+		{
+			FSkeletalMeshLODRenderData& LODData = GetClothAsset()->GetResourceForRendering()->LODRenderData[GetPredictedLODLevel()];
+			GetClothAsset()->FillComponentSpaceTransforms(GetClothAsset()->GetRefSkeleton().GetRefBonePose(), LODData.RequiredBones, GetEditableComponentSpaceTransforms());
 
-		bNeedToFlipSpaceBaseBuffers = true; // Have updated space bases so need to flip
-		FlipEditableSpaceBases();
-		bHasValidBoneTransform = true;
+			bNeedToFlipSpaceBaseBuffers = true; // Have updated space bases so need to flip
+			FlipEditableSpaceBases();
+			bHasValidBoneTransform = true;
+
+			// Create simulation proxy
+			ClothSimulationProxy = MakeUnique<FClothSimulationProxy>(*this);
+		}
 	}
 }
 
+void UChaosClothComponent::OnUnregister()
+{
+	Super::OnUnregister();
+
+	// Release cloth simulation
+	ClothSimulationProxy.Reset(nullptr);
+}
+
+void UChaosClothComponent::TickComponent(float DeltaTime, enum ELevelTick TickType, FActorComponentTickFunction* ThisTickFunction)
+{
+	CSV_SCOPED_TIMING_STAT_EXCLUSIVE(Physics);
+
+	// TODO: Fields
+	//if (ClothingSimulation)
+	//{
+	//	ClothingSimulation->UpdateWorldForces(this);
+	//}
+
+	// TODO: Suspended update
+	//// If we are suspended, we will not simulate clothing, but as clothing is simulated in local space
+	//// relative to a root bone we need to extract simulation positions as this bone could be animated.
+	//if ((!CVarEnableClothPhysics.GetValueOnGameThread() || bClothingSimulationSuspended) && ClothingSimulation)
+	//{
+	//	// First update the simulation context, since the simulation isn't ticking
+	//	// and it is still required to get the correct simulation data and bounds.
+	//	constexpr bool bIsInitialization = false;
+	//	ClothingSimulation->FillContext(this, DeltaTime, ClothingSimulationContext, bIsInitialization);
+	//	ClothingSimulation->GetSimulationData(CurrentSimulationData, this, Cast<USkeletalMeshComponent>(MasterPoseComponent.Get()));
+	//}
+
+	// Make sure that the previous frame simulation has completed
+	HandleExistingParallelSimulation();
+
+	// < This would be the right place to update the preset/use an interactor, ...etc.
+
+	// Update the proxy and start the simulation parallel task
+	StartNewParallelSimulation(DeltaTime);
+
+	// TODO: Wait in tick function
+	//if (ShouldWaitForClothInTickFunction())
+	//{
+	//	FGraphEventArray Prerequisites;
+	//	Prerequisites.Add(ParallelClothTask);
+	//	FGraphEventRef ClothCompletionEvent = TGraphTask<FParallelClothCompletionTask>::CreateTask(&Prerequisites, ENamedThreads::GameThread).ConstructAndDispatchWhenReady(this);
+	//	ThisTickFunction.GetCompletionHandle()->DontCompleteUntil(ClothCompletionEvent);
+	//}
+
+	Super::TickComponent(DeltaTime, TickType, ThisTickFunction);
+}
+
 void UChaosClothComponent::RefreshBoneTransforms(FActorComponentTickFunction* TickFunction)
 {
-	// TODO: RefreshBoneTransforms
+	MarkRenderDynamicDataDirty();
+
+	bNeedToFlipSpaceBaseBuffers = true;
+	bHasValidBoneTransform = false;
+	FlipEditableSpaceBases();
+	bHasValidBoneTransform = true;
+}
+
+void UChaosClothComponent::GetUpdateClothSimulationData_AnyThread(TMap<int32, FClothSimulData>& OutClothSimulData, FMatrix& OutLocalToWorld, float& OutClothBlendWeight)
+{
+	OutLocalToWorld = GetComponentToWorld().ToMatrixWithScale();
+
+	const UChaosClothComponent* const LeaderPoseClothComponent = Cast<UChaosClothComponent>(LeaderPoseComponent.Get());
+	if (LeaderPoseClothComponent && LeaderPoseClothComponent->ClothSimulationProxy && bBindClothToLeaderComponent)
+	{
+		OutClothBlendWeight = ClothBlendWeight;
+		OutClothSimulData = LeaderPoseClothComponent->ClothSimulationProxy->GetCurrentSimulationData_AnyThread();
+	}
+	else if (!bDisableClothSimulation && !bBindClothToLeaderComponent)
+	{
+		OutClothBlendWeight = ClothBlendWeight;
+		OutClothSimulData = ClothSimulationProxy->GetCurrentSimulationData_AnyThread();
+	}
+	else
+	{
+		OutClothSimulData.Reset();
+	}
+
+	// Blend cloth out whenever the simulation data is invalid
+	if (!OutClothSimulData.Num())
+	{
+		OutClothBlendWeight = 0.0f;
+	}
 }
 
 void UChaosClothComponent::PostLoad()
@@ -138,4 +237,86 @@ void UChaosClothComponent::PostEditChangeProperty(FPropertyChangedEvent& Propert
 }
 #endif // WITH_EDITOR
 
+bool UChaosClothComponent::RequiresPreEndOfFrameSync() const
+{
+	if (!IsSimulationSuspended() && !ShouldWaitForClothInTickFunction())
+	{
+		// By default we await the cloth task in TickComponent, but...
+		// If we have cloth and have no game-thread dependencies on the cloth output, 
+		// then we will wait for the cloth task in SendAllEndOfFrameUpdates.
+		return true;
+	}
+	return Super::RequiresPreEndOfFrameSync();
+}
+
+void UChaosClothComponent::OnPreEndOfFrameSync()
+{
+	Super::OnPreEndOfFrameSync();
+
+	HandleExistingParallelSimulation();
+}
+
+FBoxSphereBounds UChaosClothComponent::CalcBounds(const FTransform& LocalToWorld) const
+{
+	const IConsoleVariable* const CVarCacheLocalSpaceBounds = IConsoleManager::Get().FindConsoleVariable(TEXT("a.CacheLocalSpaceBounds"));
+	const bool bCacheLocalSpaceBounds = CVarCacheLocalSpaceBounds ? CVarCacheLocalSpaceBounds->GetBool() : true;
+
+	const FTransform CachedBoundsTransform = bCacheLocalSpaceBounds ? FTransform::Identity : LocalToWorld;
+
+	FBoxSphereBounds NewBounds;
+	if (ClothSimulationProxy)
+	{
+		NewBounds = ClothSimulationProxy->CalculateBounds_AnyThread().TransformBy(CachedBoundsTransform);
+	}
+
+	CachedWorldOrLocalSpaceBounds = NewBounds;
+	bCachedLocalBoundsUpToDate = bCacheLocalSpaceBounds;
+	bCachedWorldSpaceBoundsUpToDate = !bCacheLocalSpaceBounds;
+
+	if (bCacheLocalSpaceBounds)
+	{
+		CachedWorldToLocalTransform.SetIdentity();
+		return NewBounds.TransformBy(LocalToWorld);
+	}
+	CachedWorldToLocalTransform = LocalToWorld.ToInverseMatrixWithScale();
+	return NewBounds;
+}
 
+void UChaosClothComponent::StartNewParallelSimulation(float DeltaTime)
+{
+	if (ClothSimulationProxy.IsValid())
+	{
+		CSV_SCOPED_TIMING_STAT(Animation, Cloth);
+		ClothSimulationProxy->Tick_GameThread(DeltaTime);
+	}
+}
+
+void UChaosClothComponent::HandleExistingParallelSimulation()
+{
+	if (bBindClothToLeaderComponent)
+	{
+		if (UChaosClothComponent* const LeaderComponent = Cast<UChaosClothComponent>(LeaderPoseComponent.Get()))
+		{
+			LeaderComponent->HandleExistingParallelSimulation();
+		}
+	}
+
+	if (ClothSimulationProxy.IsValid())
+	{
+		ClothSimulationProxy->CompleteParallelSimulation_GameThread();
+	}
+}
+
+bool UChaosClothComponent::ShouldWaitForClothInTickFunction() const
+{
+	static IConsoleVariable* const CVarClothPhysicsWaitForParallelClothTask = IConsoleManager::Get().FindConsoleVariable(TEXT("p.ClothPhysics.WaitForParallelClothTask"));
+
+	return bWaitForParallelClothTask || (CVarClothPhysicsWaitForParallelClothTask && CVarClothPhysicsWaitForParallelClothTask->GetBool());
+}
+
+bool UChaosClothComponent::IsSimulationSuspended() const
+{
+	static IConsoleVariable* const CVarClothPhysics = IConsoleManager::Get().FindConsoleVariable(TEXT("p.ClothPhysics"));
+
+	return bSuspendSimulation || !ClothSimulationProxy.IsValid() || (CVarClothPhysics && !CVarClothPhysics->GetBool());
+}
diff --git a/Engine/Plugins/Experimental/ChaosClothAsset/Source/ChaosClothAssetEngine/Private/ChaosClothAsset/ClothSimulationContext.cpp b/Engine/Plugins/Experimental/ChaosClothAsset/Source/ChaosClothAssetEngine/Private/ChaosClothAsset/ClothSimulationContext.cpp
new file mode 100644
--- /dev/null
+++ b/Engine/Plugins/Experimental/ChaosClothAsset/Source/ChaosClothAssetEngine/Private/ChaosClothAsset/ClothSimulationContext.cpp
@@ -0,0 +1,130 @@
+// Copyright Epic Games, Inc. All Rights Reserved.
+
+#include "ChaosClothAsset/ClothSimulationContext.h"
+#include "ChaosClothAsset/ClothAsset.h"
+#include "ChaosClothAsset/ClothComponent.h"
+#include "HAL/IConsoleManager.h"
+
+namespace UE::Chaos::ClothAsset
+{
+	void FClothSimulationContext::Fill(const UChaosClothComponent& ClothComponent, float InDeltaTime, float MaxDeltaTime, bool bIsInitialization)
+	{
+		// Set the time
+		DeltaTime = FMath::Min(InDeltaTime, MaxDeltaTime);
+
+		// Set the current LOD index
+		LodIndex = ClothComponent.GetPredictedLODLevel();
+
+		// Set the teleport mode
+		const IConsoleVariable* const CVarMaxDeltaTimeTeleportMultiplier = IConsoleManager::Get().FindConsoleVariable(TEXT("p.Cloth.MaxDeltaTimeTeleportMultiplier"));
+		constexpr float MaxDeltaTimeTeleportMultiplierDefault = 1.5f;
+		const float MaxDeltaTimeTeleportMultiplier = CVarMaxDeltaTimeTeleportMultiplier ? CVarMaxDeltaTimeTeleportMultiplier->GetFloat() : MaxDeltaTimeTeleportMultiplierDefault;
+
+		TeleportMode = (DeltaTime > MaxDeltaTime * MaxDeltaTimeTeleportMultiplier) ?
+			EClothingTeleportMode::Teleport :
+			ClothComponent.GetClothTeleportMode();
+
+		VelocityScale = (TeleportMode == EClothingTeleportMode::None) ?
+			FMath::Min(InDeltaTime, MaxDeltaTime) / InDeltaTime : 1.f;
+
+		// Copy component transform
+		ComponentTransform = ClothComponent.GetComponentTransform();
+
+		// Update bone transforms
+		const UChaosClothAsset* const ClothAsset = ClothComponent.GetClothAsset();
+		const FReferenceSkeleton* const ReferenceSkeleton = ClothAsset ? &ClothAsset->GetRefSkeleton() : nullptr;
+		int32 NumBones = ReferenceSkeleton ? ReferenceSkeleton->GetNum() : 0;
+
+		if (USkinnedMeshComponent* const LeaderComponent = ClothComponent.LeaderPoseComponent.Get())
+		{
+			const TArray<int32>& LeaderBoneMap = ClothComponent.GetLeaderBoneMap();
+			if (!LeaderBoneMap.Num())
+			{
+				// This case indicates an invalid leader pose component (e.g. no skeletal mesh)
+				BoneTransforms.Empty(NumBones);
+				BoneTransforms.AddDefaulted(NumBones);
+			}
+			else
+			{
+				NumBones = LeaderBoneMap.Num();
+				BoneTransforms.Reset(NumBones);
+				BoneTransforms.AddDefaulted(NumBones);
+
+				if (!bIsInitialization)  // Initializations must be done in bind pose
+				{
+					const TArray<FTransform>& LeaderTransforms = LeaderComponent->GetComponentSpaceTransforms();
+					for (int32 BoneIndex = 0; BoneIndex < NumBones; ++BoneIndex)
+					{
+						bool bFoundLeader = false;
+						if (LeaderBoneMap.IsValidIndex(BoneIndex))
+						{
+							const int32 LeaderIndex = LeaderBoneMap[BoneIndex];
+							if (LeaderIndex != INDEX_NONE && LeaderIndex < LeaderTransforms.Num())
+							{
+								BoneTransforms[BoneIndex] = LeaderTransforms[LeaderIndex];
+								bFoundLeader = true;
+							}
+						}
+
+						if (!bFoundLeader && ReferenceSkeleton)
+						{
+							const int32 ParentIndex = ReferenceSkeleton->GetParentIndex(BoneIndex);
+
+							BoneTransforms[BoneIndex] =
+								BoneTransforms.IsValidIndex(ParentIndex) && ParentIndex < BoneIndex ?
+								BoneTransforms[ParentIndex] * ReferenceSkeleton->GetRefBonePose()[BoneIndex] :
+								ReferenceSkeleton->GetRefBonePose()[BoneIndex];
+						}
+					}
+				}
+			}
+		}
+		else if (!bIsInitialization)  // Initializations must be done in bind pose
+		{
+			BoneTransforms = ClothComponent.GetComponentSpaceTransforms();
+		}
+		else
+		{
+			BoneTransforms.Reset(NumBones);
+			BoneTransforms.AddDefaulted(NumBones);
+		}
+
+		// Update bone matrices
+		if (!bIsInitialization)
+		{
+			ClothComponent.GetCurrentRefToLocalMatrices(RefToLocalMatrices, LodIndex);
+		}
+		else
+		{
+			RefToLocalMatrices.Reset(NumBones);
+			RefToLocalMatrices.AddUninitialized(NumBones);
+
+			for (int32 BoneIndex = 0; BoneIndex < NumBones; ++BoneIndex)
+			{
+				RefToLocalMatrices[BoneIndex] = FMatrix44f::Identity;
+			}
+		}
+
+		// Update gravity
+		const UWorld* const World = ClothComponent.GetWorld();
+		constexpr float EarthGravity = -981.f;
+		WorldGravity = FVector(0.f, 0.f, World ? World->GetGravityZ() : EarthGravity);
+
+		// Update wind velocity
+		if (World && World->Scene && World->IsGameWorld())
+		{
+			const FVector Position = ClothComponent.GetComponentTransform().GetTranslation();
+
+			float WindSpeed;
+			float WindMinGust;
+			float WindMaxGust;
+			World->Scene->GetWindParameters_GameThread(Position, WindVelocity, WindSpeed, WindMinGust, WindMaxGust);
+
+			WindVelocity *= WindSpeed;
+		}
+		else
+		{
+			WindVelocity = FVector::ZeroVector;
+		}
+	}
+}  // End namespace UE::Chaos::ClothAsset
diff --git a/Engine/Plugins/Experimental/ChaosClothAsset/Source/ChaosClothAssetEngine/Private/ChaosClothAsset/ClothSimulationContext.h b/Engine/Plugins/Experimental/ChaosClothAsset/Source/ChaosClothAssetEngine/Private/ChaosClothAsset/ClothSimulationContext.h
new file mode 100644
--- /dev/null
+++ b/Engine/Plugins/Experimental/ChaosClothAsset/Source/ChaosClothAssetEngine/Private/ChaosClothAsset/ClothSimulationContext.h
@@ -0,0 +1,51 @@
+// Copyright Epic Games, Inc. All Rights Reserved.
+
+#pragma once
+
+#include "Containers/Array.h"
+#include "Math/Transform.h"
+#include "ClothingSystemRuntimeTypes.h"
+
+class UChaosClothComponent;
+
+namespace UE::Chaos::ClothAsset
+{
+	/**
+	 * Cloth simulation context.
+	 * Class used to pass data between the cloth component and its simulation proxy.
+	 * The context is set in the game thread before the simulation is started and should not
+	 * be modified until the simulation is ended.
+	 */
+	struct FClothSimulationContext
+	{
+		/** Time step in seconds. */
+		float DeltaTime;
+
+		/** Current LOD index at the start of the simulation. */
+		int32 LodIndex;
+
+		/** Component to world transform. */
+		FTransform ComponentTransform;
+
+		/** Component space bone transforms of the owning component. */
+		TArray<FTransform> BoneTransforms;
+
+		/** Bone matrices used for skinning. */
+		TArray<FMatrix44f> RefToLocalMatrices;
+
+		/** Teleport mode. */
+		EClothingTeleportMode TeleportMode = EClothingTeleportMode::None;
+
+		/** Velocity scale to compensate for time differences when the MaxPhysicsDeltaTime kicks in. */
+		float VelocityScale;
+
+		/** Gravity extracted from the world. */
+		FVector WorldGravity;
+
+		/** Wind velocity at the component location. */
+		FVector WindVelocity;
+
+		/** Fill the context from data collected from the specified component. */
+		void Fill(const UChaosClothComponent& ClothComponent, float InDeltaTime, float MaxDeltaTime, bool bIsInitialization = false);
+	};
+}
diff --git a/Engine/Plugins/Experimental/ChaosClothAsset/Source/ChaosClothAssetEngine/Private/ChaosClothAsset/ClothSimulationMesh.cpp b/Engine/Plugins/Experimental/ChaosClothAsset/Source/ChaosClothAssetEngine/Private/ChaosClothAsset/ClothSimulationMesh.cpp
new file mode 100644
--- /dev/null
+++ b/Engine/Plugins/Experimental/ChaosClothAsset/Source/ChaosClothAssetEngine/Private/ChaosClothAsset/ClothSimulationMesh.cpp
@@ -0,0 +1,120 @@
+// Copyright Epic Games, Inc. All Rights Reserved.
+
+#include "ChaosClothAsset/ClothSimulationMesh.h"
+#include "ChaosClothAsset/ClothAsset.h"
+#include "ChaosClothAsset/ClothSimulationModel.h"
+#include "ChaosClothAsset/ClothSimulationContext.h"
+#include "ClothSimulationMesh.h"
+
+namespace UE::Chaos::ClothAsset
+{
+	FClothSimulationMesh::FClothSimulationMesh(const FChaosClothSimulationModel& InClothSimulationModel, const FClothSimulationContext& InClothSimulationContext)
+PRAGMA_DISABLE_DEPRECATION_WARNINGS  // TODO: CHAOS_IS_CLOTHINGSIMULATIONMESH_ABSTRACT
+		: ::Chaos::FClothingSimulationMesh()
+PRAGMA_ENABLE_DEPRECATION_WARNINGS
+		, ClothSimulationModel(InClothSimulationModel)
+		, ClothSimulationContext(InClothSimulationContext)
+	{
+	}
+
+	int32 FClothSimulationMesh::GetNumLODs() const
+	{
+		return ClothSimulationModel.GetNumLods();
+	}
+
+	int32 FClothSimulationMesh::GetLODIndex() const
+	{
+		return ClothSimulationContext.LodIndex;
+	}
+
+	int32 FClothSimulationMesh::GetOwnerLODIndex(int32 LODIndex) const
+	{
+		return LODIndex;  // The component LOD currently matches the asset LOD
+	}
+
+	bool FClothSimulationMesh::IsValidLODIndex(int32 LODIndex) const
+	{
+		return LODIndex >= 0 && LODIndex < GetNumLODs();
+	}
+
+	int32 FClothSimulationMesh::GetNumPoints(int32 LODIndex) const
+	{
+		return ClothSimulationModel.GetPositions(LODIndex).Num();
+	}
+
+	TConstArrayView<FVector3f> FClothSimulationMesh::GetPositions(int32 LODIndex) const
+	{
+		return TConstArrayView<FVector3f>(ClothSimulationModel.GetPositions(LODIndex));
+	}
+
+	TConstArrayView<FVector3f> FClothSimulationMesh::GetNormals(int32 LODIndex) const
+	{
+		return TConstArrayView<FVector3f>(ClothSimulationModel.GetNormals(LODIndex));
+	}
+
+	TConstArrayView<uint32> FClothSimulationMesh::GetIndices(int32 LODIndex) const
+	{
+		return TConstArrayView<uint32>(ClothSimulationModel.GetIndices(LODIndex));
+	}
+
+	TArray<TConstArrayView<::Chaos::FRealSingle>> FClothSimulationMesh::GetWeightMaps(int32 LODIndex) const
+	{
+		constexpr int32 MaxNumWeightMaps = 15; // TODO: Refactor how weight maps are used in base simulation classes
+		static const TConstArrayView<::Chaos::FRealSingle> WeightMaps[MaxNumWeightMaps];
+		return TArray<TConstArrayView<::Chaos::FRealSingle>>(&WeightMaps[0], MaxNumWeightMaps);
+	}
+
+	TArray<TConstArrayView<TTuple<int32, int32, float>>> FClothSimulationMesh::GetTethers(int32 LODIndex, bool bUseGeodesicTethers) const
+	{
+		return TArray<TConstArrayView<TTuple<int32, int32, float>>>();  // TODO: Tethers
+	}
+
+	int32 FClothSimulationMesh::GetReferenceBoneIndex() const
+	{
+		return ClothSimulationModel.ReferenceBoneIndex;
+	}
+
+	::Chaos::FRigidTransform3 FClothSimulationMesh::GetReferenceBoneTransform() const
+	{
+		// TODO: Leader pose component (see FClothingSimulationContextCommon::FillBoneTransforms)
+		const TArray<FTransform>& BoneTransforms = ClothSimulationContext.BoneTransforms;
+
+		if (BoneTransforms.IsValidIndex(ClothSimulationModel.ReferenceBoneIndex))
+		{
+			return BoneTransforms[ClothSimulationModel.ReferenceBoneIndex] * GetComponentToWorldTransform();
+		}
+		return GetComponentToWorldTransform();
+	}
+
+	const FTransform& FClothSimulationMesh::GetComponentToWorldTransform() const
+	{
+		return ClothSimulationContext.ComponentTransform;
+	}
+
+	const TArray<FMatrix44f>& FClothSimulationMesh::GetRefToLocalMatrices() const
+	{
+		return ClothSimulationContext.RefToLocalMatrices;
+	}
+
+	TConstArrayView<int32> FClothSimulationMesh::GetBoneMap() const
+	{
+		return ClothSimulationModel.UsedBoneIndices;
+	}
+
+	TConstArrayView<FClothVertBoneData> FClothSimulationMesh::GetBoneData(int32 LODIndex) const
+	{
+		return ClothSimulationModel.GetBoneData(LODIndex);
+	}
+
+	TConstArrayView<FMeshToMeshVertData> FClothSimulationMesh::GetTransitionUpSkinData(int32 LODIndex) const
+	{
+		// TODO: UpSkinData
+		return TConstArrayView<FMeshToMeshVertData>();
+	}
+
+	TConstArrayView<FMeshToMeshVertData> FClothSimulationMesh::GetTransitionDownSkinData(int32 LODIndex) const
+	{
+		// TODO: DownSkinData
+		return TConstArrayView<FMeshToMeshVertData>();
+	}
+}
diff --git a/Engine/Plugins/Experimental/ChaosClothAsset/Source/ChaosClothAssetEngine/Private/ChaosClothAsset/ClothSimulationMesh.h b/Engine/Plugins/Experimental/ChaosClothAsset/Source/ChaosClothAssetEngine/Private/ChaosClothAsset/ClothSimulationMesh.h
new file mode 100644
--- /dev/null
+++ b/Engine/Plugins/Experimental/ChaosClothAsset/Source/ChaosClothAssetEngine/Private/ChaosClothAsset/ClothSimulationMesh.h
@@ -0,0 +1,53 @@
+// Copyright Epic Games, Inc. All Rights Reserved.
+#pragma once
+
+#include "ChaosCloth/ChaosClothingSimulationMesh.h"
+
+class UChaosClothAsset;
+struct FChaosClothSimulationModel;
+
+namespace UE::Chaos::ClothAsset
+{
+	struct FClothSimulationContext;
+
+PRAGMA_DISABLE_DEPRECATION_WARNINGS  // TODO: CHAOS_IS_CLOTHINGSIMULATIONMESH_ABSTRACT
+	class FClothSimulationMesh : public ::Chaos::FClothingSimulationMesh
+PRAGMA_ENABLE_DEPRECATION_WARNINGS
+	{
+	public:
+		FClothSimulationMesh(const FChaosClothSimulationModel& InClothSimulationModel, const FClothSimulationContext& InClothSimulationContext);
+PRAGMA_DISABLE_DEPRECATION_WARNINGS  // TODO: CHAOS_IS_CLOTHINGSIMULATIONMESH_ABSTRACT
+		virtual ~FClothSimulationMesh() override = default;
+PRAGMA_ENABLE_DEPRECATION_WARNINGS
+
+		FClothSimulationMesh(const FClothSimulationMesh&) = delete;
+		FClothSimulationMesh(FClothSimulationMesh&&) = delete;
+		FClothSimulationMesh& operator=(const FClothSimulationMesh&) = delete;
+		FClothSimulationMesh& operator=(FClothSimulationMesh&&) = delete;
+
+		//~ Begin FClothingSimulationMesh Interface
+		virtual int32 GetNumLODs() const override;
+		virtual int32 GetLODIndex() const override;
+		virtual int32 GetOwnerLODIndex(int32 LODIndex) const override;
+		virtual bool IsValidLODIndex(int32 LODIndex) const override;
+		virtual int32 GetNumPoints(int32 LODIndex) const override;
+		virtual TConstArrayView<FVector3f> GetPositions(int32 LODIndex) const override;
+		virtual TConstArrayView<FVector3f> GetNormals(int32 LODIndex) const override;
+		virtual TConstArrayView<uint32> GetIndices(int32 LODIndex) const override;
+		virtual TArray<TConstArrayView<::Chaos::FRealSingle>> GetWeightMaps(int32 LODIndex) const override;
+		virtual TArray<TConstArrayView<TTuple<int32, int32, float>>> GetTethers(int32 LODIndex, bool bUseGeodesicTethers) const override;
+		virtual int32 GetReferenceBoneIndex() const override;
+		virtual ::Chaos::FRigidTransform3 GetReferenceBoneTransform() const override;
+		virtual const FTransform& GetComponentToWorldTransform() const override;
+		virtual const TArray<FMatrix44f>& GetRefToLocalMatrices() const override;
+		virtual TConstArrayView<int32> GetBoneMap() const override;
+		virtual TConstArrayView<FClothVertBoneData> GetBoneData(int32 LODIndex) const override;
+		virtual TConstArrayView<FMeshToMeshVertData> GetTransitionUpSkinData(int32 LODIndex) const override;
+		virtual TConstArrayView<FMeshToMeshVertData> GetTransitionDownSkinData(int32 LODIndex) const override;
+		//~ End FClothingSimulationMesh Interface
+
+	private:
+		const FChaosClothSimulationModel& ClothSimulationModel;
+		const FClothSimulationContext& ClothSimulationContext;
+	};
+}
diff --git a/Engine/Plugins/Experimental/ChaosClothAsset/Source/ChaosClothAssetEngine/Private/ChaosClothAsset/ClothSimulationModel.cpp b/Engine/Plugins/Experimental/ChaosClothAsset/Source/ChaosClothAssetEngine/Private/ChaosClothAsset/ClothSimulationModel.cpp
new file mode 100644
--- /dev/null
+++ b/Engine/Plugins/Experimental/ChaosClothAsset/Source/ChaosClothAssetEngine/Private/ChaosClothAsset/ClothSimulationModel.cpp
@@ -0,0 +1,143 @@
+// Copyright Epic Games, Inc. All Rights Reserved.
+
+#include "ChaosClothAsset/ClothSimulationModel.h"
+#include "ChaosClothAsset/ClothCollection.h"
+#include "ChaosClothAsset/ClothAdapter.h"
+#include "ReferenceSkeleton.h"
+
+namespace UE::Chaos::ClothAsset
+{
+	// Find the root bone for this cloth asset (common bone for all used bones)
+	static int32 CalculateReferenceBoneIndex(const TArray<FChaosClothSimulationLodModel> ClothSimulationLodModels, const FReferenceSkeleton& ReferenceSkeleton, const TArray<int32>& UsedBoneIndices)
+	{
+		// Starts at root
+		int32 ReferenceBoneIndex = 0;
+
+		// List of valid paths to the root bone from each weighted bone
+		TArray<TArray<int32>> PathsToRoot;
+
+		// First build a list per used bone for it's path to root
+		TArray<int32> WeightedBones;  // List of actually weighted (not just used) bones
+
+		for (const FChaosClothSimulationLodModel& ClothSimulationLodModel : ClothSimulationLodModels)
+		{
+			for (const FClothVertBoneData& VertBoneData : ClothSimulationLodModel.BoneData)
+			{
+				for (int32 InfluenceIndex = 0; InfluenceIndex < VertBoneData.NumInfluences; ++InfluenceIndex)
+				{
+					if (VertBoneData.BoneWeights[InfluenceIndex] > SMALL_NUMBER)
+					{
+						const int32 UnmappedBoneIndex = VertBoneData.BoneIndices[InfluenceIndex];
+						check(UsedBoneIndices.IsValidIndex(UnmappedBoneIndex));
+						WeightedBones.AddUnique(UsedBoneIndices[UnmappedBoneIndex]);
+					}
+					else
+					{
+						// Hit the last weight (they're sorted)
+						break;
+					}
+				}
+			}
+		}
+
+		const int32 NumWeightedBones = WeightedBones.Num();
+		PathsToRoot.Reserve(NumWeightedBones);
+
+		// Compute paths to the root bone
+		for (int32 WeightedBoneIndex = 0; WeightedBoneIndex < NumWeightedBones; ++WeightedBoneIndex)
+		{
+			PathsToRoot.AddDefaulted();
+			TArray<int32>& Path = PathsToRoot.Last();
+
+			int32 CurrentBone = WeightedBones[WeightedBoneIndex];
+			Path.Add(CurrentBone);
+
+			while (CurrentBone != 0 && CurrentBone != INDEX_NONE)
+			{
+				CurrentBone = ReferenceSkeleton.GetParentIndex(CurrentBone);
+				Path.Add(CurrentBone);
+			}
+		}
+
+		// Paths are from leaf->root, we want the other way
+		for (TArray<int32>& Path : PathsToRoot)
+		{
+			Algo::Reverse(Path);
+		}
+
+		// Verify the last common bone in all paths as the root of the sim space
+		const int32 NumPaths = PathsToRoot.Num();
+		if (NumPaths > 0)
+		{
+			TArray<int32>& FirstPath = PathsToRoot[0];
+
+			const int32 FirstPathSize = FirstPath.Num();
+			for (int32 PathEntryIndex = 0; PathEntryIndex < FirstPathSize; ++PathEntryIndex)
+			{
+				const int32 CurrentQueryIndex = FirstPath[PathEntryIndex];
+				bool bValidRoot = true;
+
+				for (int32 PathIndex = 1; PathIndex < NumPaths; ++PathIndex)
+				{
+					if (!PathsToRoot[PathIndex].Contains(CurrentQueryIndex))
+					{
+						bValidRoot = false;
+						break;
+					}
+				}
+
+				if (bValidRoot)
+				{
+					ReferenceBoneIndex = CurrentQueryIndex;
+				}
+				else
+				{
+					// Once we fail to find a valid root we're done.
+					break;
+				}
+			}
+		}
+		else
+		{
+			// Just use the root
+			ReferenceBoneIndex = 0;
+		}
+		return ReferenceBoneIndex;
+	}
+}  // End namespace UE::Chaos::ClothAsset
+
+FChaosClothSimulationModel::FChaosClothSimulationModel(const TSharedPtr<const UE::Chaos::ClothAsset::FClothCollection>& ClothCollection, const FReferenceSkeleton& ReferenceSkeleton)
+{
+	using namespace UE::Chaos::ClothAsset;
+
+	const FClothConstAdapter Cloth(ClothCollection);
+
+	// Initialize LOD models
+	const int32 NumLods = Cloth.GetNumLods();
+	ClothSimulationLodModels.SetNum(NumLods);
+
+	for (int32 LodIndex = 0; LodIndex < NumLods; ++LodIndex)
+	{
+		FChaosClothSimulationLodModel& LodModel = ClothSimulationLodModels[LodIndex];
+		const FClothLodConstAdapter ClothLod = Cloth.GetLod(LodIndex);
+		ClothLod.BuildSimulationMesh(LodModel.Positions, LodModel.Normals, LodModel.Indices);
+
+		const int32 NumSimVertices = LodModel.Positions.Num();
+		LodModel.BoneData.SetNum(NumSimVertices);
+		for (int32 VertexIndex = 0; VertexIndex < NumSimVertices; ++VertexIndex)
+		{
+			// TODO: Skinning of the cloth asset needs to be done in order to be able to author theses values
+			LodModel.BoneData[VertexIndex].NumInfluences = 1;
+			LodModel.BoneData[VertexIndex].BoneIndices[0] = 0;
+			LodModel.BoneData[VertexIndex].BoneWeights[0] = 1.f;
+		}
+	}
+
+	// Populate used bone names and indices  
+	// TODO: Skinning of the cloth asset needs to be done in order to be able to author theses values
+	UsedBoneNames.Add(TEXT("Root"));
+	UsedBoneIndices.Add(0);
+
+	// Initialize Reference bone index
+	ReferenceBoneIndex = CalculateReferenceBoneIndex(ClothSimulationLodModels, ReferenceSkeleton, UsedBoneIndices);
+}
diff --git a/Engine/Plugins/Experimental/ChaosClothAsset/Source/ChaosClothAssetEngine/Private/ChaosClothAsset/ClothSimulationModel.h b/Engine/Plugins/Experimental/ChaosClothAsset/Source/ChaosClothAssetEngine/Private/ChaosClothAsset/ClothSimulationModel.h
new file mode 100644
--- /dev/null
+++ b/Engine/Plugins/Experimental/ChaosClothAsset/Source/ChaosClothAssetEngine/Private/ChaosClothAsset/ClothSimulationModel.h
@@ -0,0 +1,82 @@
+// Copyright Epic Games, Inc. All Rights Reserved.
+
+#pragma once
+
+#include "Containers/Array.h"
+#include "Math/MathFwd.h"
+#include "Templates/SharedPointer.h"
+#include "ClothVertBoneData.h"
+
+#include "ClothSimulationModel.generated.h"
+
+struct FReferenceSkeleton;
+namespace UE::Chaos::ClothAsset
+{
+	class FClothCollection;
+}
+
+/**
+ * Cloth simulation LOD model.
+ * Contains the LOD data as fed to the solver for constraints creation.
+ */
+USTRUCT()
+struct FChaosClothSimulationLodModel
+{
+	GENERATED_BODY()
+
+	/** Vertex positions. */
+	UPROPERTY()
+	TArray<FVector3f> Positions;
+	
+	/** Vertex Normals. */
+	UPROPERTY()
+	TArray<FVector3f> Normals;
+
+	/** Triangles indices. */
+	UPROPERTY()
+	TArray<uint32> Indices;
+
+	/** Skinning weights. */
+	UPROPERTY()
+	TArray<FClothVertBoneData> BoneData;
+};
+
+/**
+ * Cloth simulation model.
+ * Contains the LOD models.
+ */
+USTRUCT()
+struct FChaosClothSimulationModel
+{
+	GENERATED_BODY()
+
+	/** LOD data. */
+	UPROPERTY()
+	TArray<FChaosClothSimulationLodModel> ClothSimulationLodModels;
+
+	/** List of bones this asset uses from the reference skeleton. */
+	UPROPERTY()
+	TArray<FName> UsedBoneNames;
+
+	/** List of the indices for the bones in UsedBoneNames, used for remapping. */
+	UPROPERTY()
+	TArray<int32> UsedBoneIndices;
+
+	/** Bone to treat as the root of the simulation space. */
+	UPROPERTY()
+	int32 ReferenceBoneIndex;
+
+	FChaosClothSimulationModel() = default;
+	FChaosClothSimulationModel(const TSharedPtr<const UE::Chaos::ClothAsset::FClothCollection>& ClothCollection, const FReferenceSkeleton& ReferenceSkeleton);
+
+	int32 GetNumLods() const { return ClothSimulationLodModels.Num(); }
+
+	bool IsValidLodIndex(int32 LodIndex) const { return ClothSimulationLodModels.IsValidIndex(LodIndex); }
+
+	int32 GetNumVertices(int32 LodIndex) const { return IsValidLodIndex(LodIndex) ? ClothSimulationLodModels[LodIndex].Positions.Num() : 0; }
+
+	TConstArrayView<FVector3f> GetPositions(int32 LodIndex) const { return IsValidLodIndex(LodIndex) ? ClothSimulationLodModels[LodIndex].Positions : TConstArrayView<FVector3f>(); }
+	TConstArrayView<FVector3f> GetNormals(int32 LodIndex) const { return IsValidLodIndex(LodIndex) ? ClothSimulationLodModels[LodIndex].Normals : TConstArrayView<FVector3f>(); }
+	TConstArrayView<uint32> GetIndices(int32 LodIndex) const { return IsValidLodIndex(LodIndex) ? ClothSimulationLodModels[LodIndex].Indices : TConstArrayView<uint32>(); }
+	TConstArrayView<FClothVertBoneData> GetBoneData(int32 LodIndex) const { return IsValidLodIndex(LodIndex) ? ClothSimulationLodModels[LodIndex].BoneData : TConstArrayView<FClothVertBoneData>(); }
+};
diff --git a/Engine/Plugins/Experimental/ChaosClothAsset/Source/ChaosClothAssetEngine/Private/ChaosClothAsset/ClothSimulationProxy.cpp b/Engine/Plugins/Experimental/ChaosClothAsset/Source/ChaosClothAssetEngine/Private/ChaosClothAsset/ClothSimulationProxy.cpp
new file mode 100644
--- /dev/null
+++ b/Engine/Plugins/Experimental/ChaosClothAsset/Source/ChaosClothAssetEngine/Private/ChaosClothAsset/ClothSimulationProxy.cpp
@@ -0,0 +1,424 @@
+// Copyright Epic Games, Inc. All Rights Reserved.
+
+#include "ChaosClothAsset/ClothSimulationProxy.h"
+#include "ChaosClothAsset/ClothAsset.h"
+#include "ChaosClothAsset/ClothAssetPrivate.h"
+#include "ChaosClothAsset/ClothComponent.h"
+#include "ChaosClothAsset/ClothSimulationMesh.h"
+#include "ChaosClothAsset/ClothSimulationModel.h"
+#include "ChaosCloth/ChaosClothingSimulationCloth.h"
+#include "ChaosCloth/ChaosClothingSimulationCollider.h"
+#include "ChaosCloth/ChaosClothingSimulationSolver.h"
+#include "PhysicsEngine/PhysicsSettings.h"
+
+#if INTEL_ISPC
+#include "ClothSimulationProxy.ispc.generated.h"
+#endif
+
+DECLARE_CYCLE_STAT(TEXT("ClothSimulationProxy Tick Game"), STAT_ClothSimulationProxy_TickGame, STATGROUP_ChaosClothAsset);
+DECLARE_CYCLE_STAT(TEXT("ClothSimulationProxy Tick Physics"), STAT_ClothSimulationProxy_TickPhysics, STATGROUP_ChaosClothAsset);
+DECLARE_CYCLE_STAT(TEXT("ClothSimulationProxy Write Simulation Data"), STAT_ClothSimulationProxy_WriteSimulationData, STATGROUP_ChaosClothAsset);
+DECLARE_CYCLE_STAT(TEXT("ClothSimulationProxy Calculate Bounds"), STAT_ClothSimulationProxy_CalculateBounds, STATGROUP_ChaosClothAsset);
+DECLARE_CYCLE_STAT(TEXT("ClothSimulationProxy End Parallel Cloth Task"), STAT_ClothSimulationProxy_EndParallelClothTask, STATGROUP_ChaosClothAsset);
+
+CSV_DECLARE_CATEGORY_MODULE_EXTERN(ENGINE_API, Animation);
+
+namespace UE::Chaos::ClothAsset
+{
+#if INTEL_ISPC && !UE_BUILD_SHIPPING
+	static_assert(sizeof(ispc::FVector3f) == sizeof(FVector3f), "sizeof(ispc::FVector3f) != sizeof(FVector3f)");
+	static_assert(sizeof(ispc::FTransform) == sizeof(::Chaos::FRigidTransform3), "sizeof(ispc::FTransform) != sizeof(::Chaos::FRigidTransform3)");
+
+	bool bTransformClothSimulData_ISPC_Enabled = true;
+	FAutoConsoleVariableRef CVarTransformClothSimukDataISPCEnabled(TEXT("p.ChaosClothAsset.TransformClothSimulData.ISPC"), bTransformClothSimulData_ISPC_Enabled, TEXT("Whether to use ISPC optimizations when transforming simulation data back to reference bone space."));
+#endif
+
+	static FAutoConsoleTaskPriority CPrio_FParallelClothTask(
+		TEXT("TaskGraph.TaskPriorities.ParallelClothTask"),
+		TEXT("Task and thread priority for parallel cloth."),
+		ENamedThreads::HighThreadPriority, // If we have high priority task threads, then use them...
+		ENamedThreads::NormalTaskPriority, // .. at normal task priority
+		ENamedThreads::HighTaskPriority);  // If we don't have high priority threads, then use normal priority threads at high task priority instead
+
+	class FClothSimulationProxy::FParallelClothTask
+	{
+	public:
+		FParallelClothTask(FClothSimulationProxy& InClothSimulationProxy)
+			: ClothSimulationProxy(InClothSimulationProxy)
+		{
+		}
+
+		TStatId GetStatId() const
+		{
+			RETURN_QUICK_DECLARE_CYCLE_STAT(FParallelClothTask, STATGROUP_TaskGraphTasks);
+		}
+
+		static ENamedThreads::Type GetDesiredThread()
+		{
+			static IConsoleVariable* const CVarClothPhysicsUseTaskThread = IConsoleManager::Get().FindConsoleVariable(TEXT("p.ClothPhysics.UseTaskThread"));
+
+			if (CVarClothPhysicsUseTaskThread && CVarClothPhysicsUseTaskThread->GetBool())
+			{
+				return CPrio_FParallelClothTask.Get();
+			}
+			return ENamedThreads::GameThread;
+		}
+
+		static ESubsequentsMode::Type GetSubsequentsMode()
+		{
+			return ESubsequentsMode::TrackSubsequents;
+		}
+
+		void DoTask(ENamedThreads::Type CurrentThread, const FGraphEventRef& MyCompletionGraphEvent)
+		{
+			FScopeCycleCounterUObject ContextScope(ClothSimulationProxy.ClothComponent.GetSkinnedAsset());
+			CSV_SCOPED_TIMING_STAT(Animation, Cloth);
+
+			ClothSimulationProxy.Tick_PhysicsThread();
+		}
+
+	private:
+		FClothSimulationProxy& ClothSimulationProxy;
+	};
+
+	FClothSimulationProxy::FClothSimulationProxy(const UChaosClothComponent& InClothComponent)
+		: ClothComponent(InClothComponent)
+		, Solver(MakeUnique<::Chaos::FClothingSimulationSolver>())
+		, MaxDeltaTime(UPhysicsSettings::Get()->MaxPhysicsDeltaTime)
+	{
+		using namespace ::Chaos;
+
+		check(IsInGameThread());
+
+		// Need a valid context to initialize the mesh
+		constexpr bool bIsInitialization = true;
+		ClothSimulationContext.Fill(ClothComponent, 0.f, MaxDeltaTime, bIsInitialization);
+
+		// Create mesh node
+		const UChaosClothAsset* const ClothAsset = ClothComponent.GetClothAsset();
+		const FChaosClothSimulationModel* const ClothSimulationModel = ClothAsset ? ClothAsset->GetClothSimulationModel() : nullptr;
+		check(ClothSimulationModel);
+		const int32 MeshIndex = Meshes.Emplace(MakeUnique<FClothSimulationMesh>(*ClothSimulationModel, ClothSimulationContext));
+
+		// Create collider node
+		const int32 ColliderIndex = Colliders.Emplace(MakeUnique<FClothingSimulationCollider>(
+			nullptr,
+			nullptr,  // TODO: Cloth Asset Mesh 
+			/*bInUseLODIndexOverride =*/ false,
+			/*InLODIndexOverride =*/ INDEX_NONE));
+
+		//// Set the external collision data to get updated at every frame
+		//Colliders[ColliderIndex]->SetCollisionData(&ExternalCollisionData);
+
+		// Create cloth node
+		const int32 ClothIndex = Cloths.Emplace(MakeUnique<FClothingSimulationCloth>(
+PRAGMA_DISABLE_DEPRECATION_WARNINGS  // TODO: CHAOS_IS_CLOTHINGSIMULATIONMESH_ABSTRACT
+			Meshes[MeshIndex].Get(),
+PRAGMA_ENABLE_DEPRECATION_WARNINGS
+			TArray<FClothingSimulationCollider*>({ Colliders[ColliderIndex].Get() }),
+			0,
+			FClothingSimulationCloth::EMassMode::Density,  //	(FClothingSimulationCloth::EMassMode)ClothConfig->MassMode,
+			0.35f,  //	ClothConfig->GetMassValue(),
+			0.0001f,  //	ClothConfig->MinPerParticleMass,
+			TVec2<FRealSingle>(1.f),  //	TVec2<FRealSingle>(ClothConfig->EdgeStiffnessWeighted.Low, ClothConfig->EdgeStiffnessWeighted.High),
+			TVec2<FRealSingle>(1.f),  //	TVec2<FRealSingle>(ClothConfig->BendingStiffnessWeighted.Low, ClothConfig->BendingStiffnessWeighted.High),
+			0.f,  //	ClothConfig->BucklingRatio,
+			TVec2<FRealSingle>(1.f),  //	TVec2<FRealSingle>(ClothConfig->BucklingStiffnessWeighted.Low, ClothConfig->BucklingStiffnessWeighted.High),
+			false,  //	ClothConfig->bUseBendingElements,
+			TVec2<FRealSingle>(1.f),  //	TVec2<FRealSingle>(ClothConfig->AreaStiffnessWeighted.Low, ClothConfig->AreaStiffnessWeighted.High),
+			0.f,  //	ClothConfig->VolumeStiffness,
+			false,  //	ClothConfig->bUseThinShellVolumeConstraints,
+			TVec2<FRealSingle>(1.f),   //	TVec2<FRealSingle>(ClothConfig->TetherStiffness.Low, ClothConfig->TetherStiffness.High),  // Animatable
+			TVec2<FRealSingle>(1.f),   //	TVec2<FRealSingle>(ClothConfig->TetherScale.Low, ClothConfig->TetherScale.High),  // Animatable
+			FClothingSimulationCloth::ETetherMode::Geodesic,  //	ClothConfig->bUseGeodesicDistance ? FClothingSimulationCloth::ETetherMode::Geodesic : FClothingSimulationCloth::ETetherMode::Euclidean,
+			1.f,  //	/*MaxDistancesMultiplier =*/ 1.f,  // Animatable
+			TVec2<FRealSingle>(0.f),  //	TVec2<FRealSingle>(ClothConfig->AnimDriveStiffness.Low, ClothConfig->AnimDriveStiffness.High),  // Animatable
+			TVec2<FRealSingle>(0.f),  //	TVec2<FRealSingle>(ClothConfig->AnimDriveDamping.Low, ClothConfig->AnimDriveDamping.High),  // Animatable
+			0.f,  //	ClothConfig->ShapeTargetStiffness,  // TODO: This is now deprecated
+			false,  //	false, // bUseXPBDEdgeSprings
+			false,  //	false, // bUseXPBDBendingElements
+			false,  //	false, // bUseXPBDAreaSprings
+			.5f,   //	ClothConfig->GravityScale,
+			false,  //	ClothConfig->bUseGravityOverride,
+			TVec3<FRealSingle>(0.f, 0.f, -980.665f),  //	ClothConfig->Gravity,
+			TVec3<FRealSingle>(0.75f),  //	ClothConfig->LinearVelocityScale,
+			0.75f,  //	ClothConfig->AngularVelocityScale,
+			1.f,  //	ClothConfig->FictitiousAngularScale,
+			TVec2<FRealSingle>(0.035f, 1.f),  //	TVec2<FRealSingle>(ClothConfig->Drag.Low, ClothConfig->Drag.High),  // Animatable
+			TVec2<FRealSingle>(0.035f, 1.f),  //	TVec2<FRealSingle>(ClothConfig->Lift.Low, ClothConfig->Lift.High),  // Animatable
+			false,  //	ClothConfig->bUsePointBasedWindModel,
+			TVec2<FRealSingle>(0.f),  //	TVec2<FRealSingle>(ClothConfig->Pressure.Low, ClothConfig->Pressure.High),  // Animatable
+			0.01f,  //	ClothConfig->DampingCoefficient,
+			0.f,  //	ClothConfig->LocalDampingCoefficient,
+			1.f,  //	ClothConfig->CollisionThickness,
+			0.8f,  //	ClothConfig->FrictionCoefficient,
+			false,  //	ClothConfig->bUseCCD,
+			false,  //	ClothConfig->bUseSelfCollisions,
+			2.f,  //	ClothConfig->SelfCollisionThickness,
+			0.f,  //	ClothConfig->SelfCollisionFriction,
+			false,  //	ClothConfig->bUseSelfIntersections,
+			false,  //	ClothConfig->bUseLegacyBackstop,
+			false,  //	/*bUseLODIndexOverride =*/ false,
+			INDEX_NONE));  //	/*LODIndexOverride =*/ INDEX_NONE));
+		
+		Cloths[ClothIndex]->SetAerodynamicsProperties(TVec2<FRealSingle>(0.5f, 1.f), TVec2<FRealSingle>(0.35f, 1.f), 1.225e-6f, TVec3<FRealSingle>(-30.f, -30.f, 200.f));
+
+		// Add cloth to solver
+		Solver->AddCloth(Cloths[ClothIndex].Get());
+	}
+
+	FClothSimulationProxy::~FClothSimulationProxy()
+	{
+		CompleteParallelSimulation_GameThread();
+	}
+
+	void FClothSimulationProxy::Tick_GameThread(float DeltaTime)
+	{
+		SCOPE_CYCLE_COUNTER(STAT_ClothSimulationProxy_TickGame);
+
+		if (DeltaTime > 0.f)
+		{
+			// Fill a new context
+			ClothSimulationContext.Fill(ClothComponent, DeltaTime, MaxDeltaTime);
+
+			// Start the the cloth simulation thread
+			ParallelClothTask = TGraphTask<FParallelClothTask>::CreateTask(nullptr, ENamedThreads::GameThread).ConstructAndDispatchWhenReady(*this);
+		}
+	}
+
+	void FClothSimulationProxy::Tick_PhysicsThread()
+	{
+		using namespace ::Chaos;
+
+		TRACE_CPUPROFILER_EVENT_SCOPE(FClothSimulationProxy_TickPhysics);
+		SCOPE_CYCLE_COUNTER(STAT_ClothSimulationProxy_TickPhysics);
+		if (ClothSimulationContext.DeltaTime == 0.f)
+		{
+			return;
+		}
+
+		const double StartTime = FPlatformTime::Seconds();
+		const float PrevSimulationTime = SimulationTime;  // Copy the atomic to prevent a re-read
+
+		const bool bNeedsReset = (ClothSimulationContext.TeleportMode == EClothingTeleportMode::TeleportAndReset || PrevSimulationTime == 0.f);
+		const bool bNeedsTeleport = (ClothSimulationContext.TeleportMode > EClothingTeleportMode::None);
+		bIsTeleported = bNeedsTeleport;
+
+		// Update Solver animatable parameters
+		Solver->SetLocalSpaceLocation((FVec3)ClothSimulationContext.ComponentTransform.GetLocation(), bNeedsReset);
+		Solver->SetLocalSpaceRotation((FQuat)ClothSimulationContext.ComponentTransform.GetRotation());
+		Solver->SetWindVelocity(ClothSimulationContext.WindVelocity);
+		Solver->SetGravity(ClothSimulationContext.WorldGravity);
+		Solver->SetNumIterations(8);  // TODO: Solver parameters
+		Solver->SetMaxNumIterations(10);  // TODO: Solver parameters
+		Solver->SetNumSubsteps(1);  // TODO: Solver parameters
+
+		// Check teleport modes
+		for (const TUniquePtr<FClothingSimulationCloth>& Cloth : Cloths)
+		{
+			// Update Cloth animatable parameters while in the cloth loop
+			//Cloth->SetMaxDistancesMultiplier(Context->MaxDistanceScale);
+
+			if (bNeedsReset)
+			{
+				Cloth->Reset();
+			}
+			if (bNeedsTeleport)
+			{
+				Cloth->Teleport();
+			}
+		}
+
+		// Step the simulation
+		Solver->Update((Softs::FSolverReal)ClothSimulationContext.DeltaTime);
+
+		// Keep the actual used number of iterations for the stats
+		NumIterations = Solver->GetNumUsedIterations();
+
+		// Update simulation time in ms (and provide an instant average instead of the value in real-time)
+		const float CurrSimulationTime = (float)((FPlatformTime::Seconds() - StartTime) * 1000.);
+		static const float SimulationTimeDecay = 0.03f; // 0.03 seems to provide a good rate of update for the instant average
+		SimulationTime = PrevSimulationTime ? PrevSimulationTime + (CurrSimulationTime - PrevSimulationTime) * SimulationTimeDecay : CurrSimulationTime;
+	}
+
+	void FClothSimulationProxy::CompleteParallelSimulation_GameThread()
+	{
+		check(IsInGameThread());
+
+		if (IsValidRef(ParallelClothTask))
+		{
+			SCOPE_CYCLE_COUNTER(STAT_ClothSimulationProxy_EndParallelClothTask);
+			CSV_SCOPED_SET_WAIT_STAT(Cloth);
+
+			// There's a simulation in flight
+			FTaskGraphInterface::Get().WaitUntilTaskCompletes(ParallelClothTask, ENamedThreads::GameThread);
+
+			// No longer need this task, it has completed
+			ParallelClothTask.SafeRelease();
+
+			// Write back to the GT cache
+			WriteSimulationData_GameThread();
+		}
+	}
+
+	void FClothSimulationProxy::WriteSimulationData_GameThread()
+	{
+		using namespace ::Chaos;
+		using namespace ::Chaos::Softs;
+
+		check(IsInGameThread());
+
+		CSV_SCOPED_TIMING_STAT(Animation, Cloth);
+		TRACE_CPUPROFILER_EVENT_SCOPE(FClothSimulationProxy_WriteSimulationData);
+		SCOPE_CYCLE_COUNTER(STAT_ClothSimulationProxy_WriteSimulationData);
+
+		USkinnedMeshComponent* LeaderPoseComponent = nullptr;
+		if (ClothComponent.LeaderPoseComponent.IsValid())
+		{
+			LeaderPoseComponent = ClothComponent.LeaderPoseComponent.Get();
+
+			// Check if our bone map is actually valid, if not there is no clothing data to build
+			if (!ClothComponent.GetLeaderBoneMap().Num())
+			{
+				CurrentSimulationData.Reset();
+				return;
+			}
+		}
+
+		if (!Cloths.Num())
+		{
+			CurrentSimulationData.Reset();
+			return;
+		}
+
+		// Reset map when new cloths have appeared
+		if (CurrentSimulationData.Num() != Cloths.Num())
+		{
+			CurrentSimulationData.Reset();
+		}
+
+		// Get the solver's local space
+		const FVec3& LocalSpaceLocation = Solver->GetLocalSpaceLocation(); // Note: Since the ReferenceSpaceTransform can be suspended with the simulation, it is important that the suspended local space location is used too in order to get the simulation data back into reference space
+
+		// Retrieve the component's bones transforms
+		const TArray<FTransform>& ComponentSpaceTransforms = LeaderPoseComponent ? LeaderPoseComponent->GetComponentSpaceTransforms() : ClothComponent.GetComponentSpaceTransforms();
+
+		// Set the simulation data for each of the cloths
+		for (const TUniquePtr<FClothingSimulationCloth>& Cloth : Cloths)
+		{
+			const int32 AssetIndex = Cloth->GetGroupId();
+
+			if (!Cloth->GetMesh())
+			{
+				CurrentSimulationData.Remove(AssetIndex);  // Ensures that the cloth vertex factory won't run unnecessarily
+				continue;  // Invalid or empty cloth
+			}
+
+			// If the LOD has changed while the simulation is suspended, the cloth still needs to be updated with the correct LOD data
+PRAGMA_DISABLE_DEPRECATION_WARNINGS  // TODO: CHAOS_IS_CLOTHINGSIMULATIONMESH_ABSTRACT
+			const int32 LODIndex = Cloth->GetMesh()->GetLODIndex();
+PRAGMA_ENABLE_DEPRECATION_WARNINGS
+			if (LODIndex != Cloth->GetLODIndex(Solver.Get()))
+			{
+				Solver->Update(FSolverReal(0.));  // Update for LOD switching, but do not simulate
+			}
+
+			if (Cloth->GetOffset(Solver.Get()) == INDEX_NONE || Cloth->GetLODIndex(Solver.Get()) == INDEX_NONE)
+			{
+				CurrentSimulationData.Remove(AssetIndex);  // Ensures that the cloth vertex factory won't run unnecessarily
+				continue;  // No valid LOD, there's nothing to write out
+			}
+
+			// Get the reference bone index for this cloth
+			const int32 ReferenceBoneIndex = LeaderPoseComponent ? ClothComponent.GetLeaderBoneMap()[Cloth->GetReferenceBoneIndex()] : Cloth->GetReferenceBoneIndex();
+			if (!ComponentSpaceTransforms.IsValidIndex(ReferenceBoneIndex))
+			{
+				UE_CLOG(!bHasInvalidReferenceBoneTransforms, LogChaosClothAsset, Warning, TEXT("Failed to write back clothing simulation data for component %s as bone transforms are invalid."), *ClothComponent.GetName());
+				bHasInvalidReferenceBoneTransforms = true;
+				CurrentSimulationData.Reset();
+				return;
+			}
+
+			// Get the reference transform used in the current animation pose
+			FTransform ReferenceBoneTransform = ComponentSpaceTransforms[ReferenceBoneIndex];
+			ReferenceBoneTransform *= ClothSimulationContext.ComponentTransform;
+			ReferenceBoneTransform.SetScale3D(FVector(1.0f));  // Scale is already baked in the cloth mesh
+
+			// Set the world space transform to be this cloth's reference bone
+			FClothSimulData& Data = CurrentSimulationData.FindOrAdd(AssetIndex);
+			Data.Transform = ReferenceBoneTransform;
+			Data.ComponentRelativeTransform = ReferenceBoneTransform.GetRelativeTransform(ClothSimulationContext.ComponentTransform);
+
+			// Retrieve the last reference space transform used for this cloth
+			// Note: This won't necessary match the current bone reference transform when the simulation is paused,
+			//       and still allows for the correct positioning of the sim data while the component is animated.
+			FRigidTransform3 ReferenceSpaceTransform = Cloth->GetReferenceSpaceTransform();
+			ReferenceSpaceTransform.AddToTranslation(-LocalSpaceLocation);
+
+			// Copy positions and normals
+			Data.Positions = Cloth->GetParticlePositions(Solver.Get());
+			Data.Normals = Cloth->GetParticleNormals(Solver.Get());
+
+			// Transform into the cloth reference simulation space used at the time of simulation
+			check(Data.Positions.Num() == Data.Normals.Num())
+#if INTEL_ISPC
+			if (bTransformClothSimulData_ISPC_Enabled)
+			{
+				// ISPC is assuming float input here
+				check(sizeof(ispc::FVector3f) == Data.Positions.GetTypeSize());
+				check(sizeof(ispc::FVector3f) == Data.Normals.GetTypeSize());
+
+				ispc::TransformClothSimulData(
+					(ispc::FVector3f*)Data.Positions.GetData(),
+					(ispc::FVector3f*)Data.Normals.GetData(),
+					(ispc::FTransform&)ReferenceSpaceTransform,
+					Data.Positions.Num());
+			}
+			else
+#endif
+			{
+				for (int32 Index = 0; Index < Data.Positions.Num(); ++Index)
+				{
+					Data.Positions[Index] = FVec3f(ReferenceSpaceTransform.InverseTransformPosition(FVec3(Data.Positions[Index])));
+					Data.Normals[Index] = FVec3f(ReferenceSpaceTransform.InverseTransformVector(FVec3(-Data.Normals[Index])));
+				}
+			}
+
+			// Set the current LOD these data apply to, so that the correct deformer mappings can be applied
+PRAGMA_DISABLE_DEPRECATION_WARNINGS  // TODO: CHAOS_IS_CLOTHINGSIMULATIONMESH_ABSTRACT
+			Data.LODIndex = Cloth->GetMesh()->GetOwnerLODIndex(LODIndex);  // The owner component LOD index can be different to the cloth mesh LOD index
+PRAGMA_ENABLE_DEPRECATION_WARNINGS
+		}
+	}
+
+	const TMap<int32, FClothSimulData>& FClothSimulationProxy::GetCurrentSimulationData_AnyThread() const
+	{
+		// This is called during EndOfFrameUpdates, usually in a parallel-for loop. We need to be sure that
+		// the cloth task (if there is one) is complete, but it cannot be waited for here. See OnPreEndOfFrameUpdateSync
+		// which is called just before EOF updates and is where we would have waited for the cloth task.
+		if (!IsValidRef(ParallelClothTask) || ParallelClothTask->IsComplete())
+		{
+			return CurrentSimulationData;
+		}
+		static const TMap<int32, FClothSimulData> EmptyClothSimulationData;
+		return EmptyClothSimulationData;
+	}
+
+	FBoxSphereBounds FClothSimulationProxy::CalculateBounds_AnyThread() const
+	{
+		SCOPE_CYCLE_COUNTER(STAT_ClothSimulationProxy_CalculateBounds);
+
+		check(Solver);
+		if (!IsValidRef(ParallelClothTask) || ParallelClothTask->IsComplete())
+		{
+			FBoxSphereBounds Bounds = Solver->CalculateBounds();
+
+			// The component could be moving while the simulation is suspended so getting the bounds
+			// in world space isn't good enough and the bounds origin needs to be continuously updated
+			Bounds = Bounds.TransformBy(FTransform((FQuat)Solver->GetLocalSpaceRotation(), (FVector)Solver->GetLocalSpaceLocation()).Inverse());
+
+			return Bounds;
+		}
+		return FBoxSphereBounds(ForceInit);
+	}
+}
diff --git a/Engine/Plugins/Experimental/ChaosClothAsset/Source/ChaosClothAssetEngine/Private/ChaosClothAsset/ClothSimulationProxy.h b/Engine/Plugins/Experimental/ChaosClothAsset/Source/ChaosClothAssetEngine/Private/ChaosClothAsset/ClothSimulationProxy.h
new file mode 100644
--- /dev/null
+++ b/Engine/Plugins/Experimental/ChaosClothAsset/Source/ChaosClothAssetEngine/Private/ChaosClothAsset/ClothSimulationProxy.h
@@ -0,0 +1,108 @@
+// Copyright Epic Games, Inc. All Rights Reserved.
+
+#pragma once
+
+#include "ChaosClothAsset/ClothSimulationContext.h"
+#include "Async/TaskGraphInterfaces.h"
+#include "Containers/Array.h"
+#include "Math/Transform.h"
+#include "Templates/UniquePtr.h"
+#include "ClothingSystemRuntimeTypes.h"
+
+namespace Chaos
+{
+	class FClothingSimulationSolver;
+	class FClothingSimulationMesh;
+	class FClothingSimulationCloth;
+	class FClothingSimulationCollider;
+}
+
+class UChaosClothComponent;
+class UChaosClothAsset;
+struct FReferenceSkeleton;
+
+namespace UE::Chaos::ClothAsset
+{
+	/**
+	 * Cloth simulation proxy.
+	 * Class used to share data between the cloth simulation and the cloth component.
+	 */
+	class FClothSimulationProxy final
+	{
+	public:
+		explicit FClothSimulationProxy(const UChaosClothComponent& InClothComponent);
+		~FClothSimulationProxy();
+
+		FClothSimulationProxy() = delete;  // This object cannot be created without a valid reference to a parent UChaosClothComponent 
+		FClothSimulationProxy(const FClothSimulationProxy&) = delete;  // Disable the copy as there must be a single unique proxy per component
+		FClothSimulationProxy(FClothSimulationProxy&&) = delete;  // Disable the move to force it to be associated with a valid component reference
+		FClothSimulationProxy& operator=(const FClothSimulationProxy&) = delete;
+		FClothSimulationProxy& operator=(FClothSimulationProxy&&) = delete;
+
+		void Tick_GameThread(float DeltaTime);
+
+		void CompleteParallelSimulation_GameThread();
+
+		/**
+		 * Return a map of all simulation data as used by the skeletal rendering code.
+		 * The map key is the rendering section's cloth index as set in FSkelMeshRenderSection::CorrespondClothAssetIndex,
+		 * which is 0 for the entire cloth component since all of its sections share the same simulation data.
+		 */
+		const TMap<int32, FClothSimulData>& GetCurrentSimulationData_AnyThread() const;
+
+		FBoxSphereBounds CalculateBounds_AnyThread() const;
+
+	private:
+		void Tick_PhysicsThread();
+
+		void WriteSimulationData_GameThread();
+
+		// Internal physics thread object
+		class FParallelClothTask;
+
+		// Reference for the clothing parallel task, to detect whether or not a simulation is running
+		FGraphEventRef ParallelClothTask;
+
+		// Simulation data written back to the component after the simulation has taken place
+		TMap<int32, FClothSimulData> CurrentSimulationData;
+
+		// Owner component
+		const UChaosClothComponent& ClothComponent;
+
+		// Simulation context used to store the required component data for the duration of the simulation
+		FClothSimulationContext ClothSimulationContext;
+
+		// Simulation objects
+		TUniquePtr<::Chaos::FClothingSimulationSolver> Solver;
+PRAGMA_DISABLE_DEPRECATION_WARNINGS  // TODO: CHAOS_IS_CLOTHINGSIMULATIONMESH_ABSTRACT
+		TArray<TUniquePtr<::Chaos::FClothingSimulationMesh>> Meshes;
+PRAGMA_ENABLE_DEPRECATION_WARNINGS
+		TArray<TUniquePtr<::Chaos::FClothingSimulationCloth>> Cloths;
+		TArray<TUniquePtr<::Chaos::FClothingSimulationCollider>> Colliders;
+
+		// Properties that must be readable from all threads
+		std::atomic<int32> NumCloths;
+		std::atomic<int32> NumKinematicParticles;
+		std::atomic<int32> NumDynamicParticles;
+		std::atomic<int32> NumIterations;
+		std::atomic<int32> NumSubsteps;
+		std::atomic<float> SimulationTime;
+		std::atomic<bool> bIsTeleported;
+
+		mutable bool bHasInvalidReferenceBoneTransforms = false;
+
+		// Cached value of the MaxPhysicsDeltaTime setting for the life of this proxy
+		const float MaxDeltaTime;
+	};
+
+#if !defined(CHAOS_TRANSFORM_CLOTH_SIMUL_DATA_ISPC_ENABLED_DEFAULT)
+#define CHAOS_TRANSFORM_CLOTH_SIMUL_DATA_ISPC_ENABLED_DEFAULT 1
+#endif
+
+	// Support run-time toggling on supported platforms in non-shipping configurations
+#if !INTEL_ISPC || UE_BUILD_SHIPPING
+	static constexpr bool bTransformClothSimulData_ISPC_Enabled = INTEL_ISPC && CHAOS_TRANSFORM_CLOTH_SIMUL_DATA_ISPC_ENABLED_DEFAULT;
+#else
+	extern bool bTransformClothSimulData_ISPC_Enabled;
+#endif
+}
diff --git a/Engine/Plugins/Experimental/ChaosClothAsset/Source/ChaosClothAssetEngine/Private/ChaosClothAsset/ClothSimulationProxy.ispc b/Engine/Plugins/Experimental/ChaosClothAsset/Source/ChaosClothAssetEngine/Private/ChaosClothAsset/ClothSimulationProxy.ispc
new file mode 100644
--- /dev/null
+++ b/Engine/Plugins/Experimental/ChaosClothAsset/Source/ChaosClothAssetEngine/Private/ChaosClothAsset/ClothSimulationProxy.ispc
@@ -0,0 +1,43 @@
+// Copyright Epic Games, Inc. All Rights Reserved.
+
+#define EXPLICIT_VECTOR4 1
+
+#include "Math/Vector.isph"
+#include "Math/Quat.isph"
+#include "Math/Transform.isph"
+
+export void TransformClothSimulData(uniform FVector3f Positions[],
+									uniform FVector3f Normals[],
+									const uniform FTransform &ReferenceSpaceTransform,
+									const uniform int NumPositions)
+{
+	const uniform FVector4 RotationInverse = QuatInverse(ReferenceSpaceTransform.Rotation);
+	const uniform FVector Translation = SetVector(ReferenceSpaceTransform.Translation);
+	const uniform FVector SafeReciprocal = TransformGetSafeScaleReciprocal(SetVector(ReferenceSpaceTransform.Scale3D), SMALL_NUMBER);
+
+	foreach(i = 0 ... NumPositions)
+	{
+		const uniform int Index = extract(i, 0);
+		const FVector3f PositionFloat = VectorLoad(&Positions[Index]);
+		const FVector3f NormalFloat = VectorLoad(&Normals[Index]);
+
+		const FVector InputPositionVector = ConvertVector3fTo3Native(PositionFloat);
+
+		const FVector3f InputNormalVectorFloat = NormalFloat * -FLOAT_ONE;
+		const FVector InputNormalVector = ConvertVector3fTo3Native(InputNormalVectorFloat);
+
+		const FVector TranslatedPosition = InputPositionVector - Translation;
+
+		const FVector VRPosition = VectorQuaternionRotateVector(RotationInverse, TranslatedPosition);
+		const FVector VRNormal = VectorQuaternionRotateVector(RotationInverse, InputNormalVector);
+
+		const FVector OutPosition = VRPosition * SafeReciprocal;
+		const FVector OutNormal = VRNormal * SafeReciprocal;
+
+		const FVector3f OutPositionFloat = ConvertVector3NativeTo3f(OutPosition);
+		const FVector3f OutNormalFloat = ConvertVector3NativeTo3f(OutNormal);
+
+		VectorStore(&Positions[Index], OutPositionFloat);
+		VectorStore(&Normals[Index], OutNormalFloat);
+	}
+}
diff --git a/Engine/Plugins/Experimental/ChaosClothAsset/Source/ChaosClothAssetEngine/Public/ChaosClothAsset/ClothAsset.h b/Engine/Plugins/Experimental/ChaosClothAsset/Source/ChaosClothAssetEngine/Public/ChaosClothAsset/ClothAsset.h
--- a/Engine/Plugins/Experimental/ChaosClothAsset/Source/ChaosClothAssetEngine/Public/ChaosClothAsset/ClothAsset.h
+++ b/Engine/Plugins/Experimental/ChaosClothAsset/Source/ChaosClothAssetEngine/Public/ChaosClothAsset/ClothAsset.h
@@ -2,7 +2,6 @@
 
 #pragma once
 
-#include "CoreMinimal.h"
 #include "ChaosClothAsset/ClothPreset.h"
 #include "Engine/SkeletalMesh.h"  // For FSkeletalMeshLODInfo
 #include "Engine/SkinnedAsset.h"
@@ -13,7 +12,7 @@
 namespace UE::Chaos::ClothAsset { class FClothCollection; }
 class FSkeletalMeshRenderData;
 class FSkeletalMeshModel;
-
+struct FChaosClothSimulationModel;
 
 UENUM()
 enum class EClothAssetAsyncProperties : uint64
@@ -24,9 +23,8 @@ enum class EClothAssetAsyncProperties : uint64
 };
 ENUM_CLASS_FLAGS(EClothAssetAsyncProperties);
 
-
 /**
- * Tailored cloth simulation asset.
+ * Cloth asset for pattern based simulation.
  */
 UCLASS(hidecategories = Object, BlueprintType)
 class CHAOSCLOTHASSETENGINE_API UChaosClothAsset : public USkinnedAsset
@@ -34,8 +32,8 @@ class CHAOSCLOTHASSETENGINE_API UChaosClothAsset : public USkinnedAsset
 	GENERATED_BODY()
 public:
 	UChaosClothAsset(const FObjectInitializer& ObjectInitializer);
-	//UChaosClothAsset(FVTableHelper& Helper);  // This is declared so we can use TUniquePtr<FSkeletalMeshRenderData> with just a forward declare of that class
-	~UChaosClothAsset() = default;
+	UChaosClothAsset(FVTableHelper& Helper);  // This is declared so we can use TUniquePtr<FClothSimulationModel> with just a forward declare of that class
+	virtual ~UChaosClothAsset() override;
 
 	//~ Begin UObject interface
 	virtual void Serialize(FArchive& Ar) override;
@@ -73,7 +71,7 @@ public:
 	virtual USkeletalMeshSocket* FindSocketInfo(FName InSocketName, FTransform& OutTransform, int32& OutBoneIndex, int32& OutIndex) const override { return nullptr; }
 	virtual USkeleton* GetSkeleton() override									{ return Skeleton; }
 	virtual const USkeleton* GetSkeleton() const override						{ return Skeleton; }
-	virtual void SetSkeleton(USkeleton* InSkeleton) override					{ Skeleton = InSkeleton; }
+	virtual void SetSkeleton(USkeleton* InSkeleton) override					{ Skeleton = InSkeleton; UpdateSkeleton(); }
 	virtual UMeshDeformer* GetDefaultMeshDeformer() const override				{ return nullptr; }
 	virtual bool IsValidLODIndex(int32 Index) const override					{ return LODInfo.IsValidIndex(Index); }
 	virtual int32 GetMinLodIdx(bool bForceLowestLODIdx = false) const override	{ return 0; }
@@ -100,6 +98,9 @@ public:
 	/** Return the enclosed Cloth Collection object, const version. */
 	const TSharedPtr<UE::Chaos::ClothAsset::FClothCollection> GetClothCollection() const { return ClothCollection; }
 
+	/** Return the cloth simulation ready LOD model data. */
+	const FChaosClothSimulationModel* GetClothSimulationModel() const { return ClothSimulationModel.Get(); }
+
 	/** Build this asset static render and simulation data. This needs to be done every time the asset has changed. */
 	void Build();
 
@@ -134,11 +135,20 @@ private:
 	/** Re-calculate the bounds for this asset. */
 	void CalculateBounds();
 
+	/** Re-calculate the reference bone index, from which the cloth gets transformed when the simulation is suspended. */
+	void CalculateReferenceBoneIndex();
+
+	/** Update the bone informations after a change of skeleton. */
+	void UpdateSkeleton(bool bRebuildClothSimulationModel = true);
+
 #if WITH_EDITORONLY_DATA
 	/** Build the SkeletalMeshLODModel for this asset. */
-	void BuildModel();
+	void BuildMeshModel();
 #endif
 
+	/** Build the clothing simulation meshes from the Cloth Collection. */
+	void BuildClothSimulationModel();
+
 	/** Initialize all render resources. */
 	void InitResources();
 
@@ -232,5 +242,8 @@ private:
 	TSharedPtr<FSkeletalMeshModel> MeshModel;
 #endif
 
+	/** Simulation mesh Lods as fed to the solver for constraints creation. */
+	TUniquePtr<FChaosClothSimulationModel> ClothSimulationModel;
+
 	friend class UClothAssetBuilderEditor;
 };
diff --git a/Engine/Plugins/Experimental/ChaosClothAsset/Source/ChaosClothAssetEngine/Public/ChaosClothAsset/ClothAssetBuilder.h b/Engine/Plugins/Experimental/ChaosClothAsset/Source/ChaosClothAssetEngine/Public/ChaosClothAsset/ClothAssetBuilder.h
--- a/Engine/Plugins/Experimental/ChaosClothAsset/Source/ChaosClothAssetEngine/Public/ChaosClothAsset/ClothAssetBuilder.h
+++ b/Engine/Plugins/Experimental/ChaosClothAsset/Source/ChaosClothAssetEngine/Public/ChaosClothAsset/ClothAssetBuilder.h
@@ -2,12 +2,13 @@
 
 #pragma once
 
-#include "ChaosClothAsset/ClothAsset.h"
 #include "Features/IModularFeature.h"
 
 #include "ClothAssetBuilder.generated.h"
 
 class UClothAssetBuilder;
+class UChaosClothAsset;
+class FSkeletalMeshLODModel;
 
 /**
  * A modular interface to provide ways to build a SkeletalMesh LODModel from a cloth asset.
@@ -32,4 +33,10 @@ public:
 	/** Build a FSkeletalMeshLODModel out of the cloth asset for the specified LOD index. */
 	virtual void BuildLod(FSkeletalMeshLODModel& LODModel, const UChaosClothAsset& ClothAsset, int32 LodIndex) const
 	PURE_VIRTUAL(UClothAssetBuilder::BuildLod, );
+
+protected:
+	// Helper functions used to retrieve LOD information without having to make public the FChaosClothSimulationModel class
+	static int32 GetNumVertices(const UChaosClothAsset& ClothAsset, int32 LodIndex);
+	static TConstArrayView<FVector3f> GetSimPositions(const UChaosClothAsset& ClothAsset, int32 LodIndex);
+	static TConstArrayView<uint32> GetSimIndices(const UChaosClothAsset& ClothAsset, int32 LodIndex);
 };
diff --git a/Engine/Plugins/Experimental/ChaosClothAsset/Source/ChaosClothAssetEngine/Public/ChaosClothAsset/ClothComponent.h b/Engine/Plugins/Experimental/ChaosClothAsset/Source/ChaosClothAssetEngine/Public/ChaosClothAsset/ClothComponent.h
--- a/Engine/Plugins/Experimental/ChaosClothAsset/Source/ChaosClothAssetEngine/Public/ChaosClothAsset/ClothComponent.h
+++ b/Engine/Plugins/Experimental/ChaosClothAsset/Source/ChaosClothAssetEngine/Public/ChaosClothAsset/ClothComponent.h
@@ -3,20 +3,28 @@
 #pragma once
 
 #include "Components/SkinnedMeshComponent.h"
+#include "ClothingSystemRuntimeTypes.h"
 
 #include "ClothComponent.generated.h"
 
 class UChaosClothAsset;
+class UChaosClothComponent;
 
+namespace UE::Chaos::ClothAsset
+{
+	class FClothSimulationProxy;
+}
 /**
  * Cloth simulation component.
  */
-UCLASS(ClassGroup = Physics, hidecategories = (Object, "Mesh|SkeletalAsset"), meta = (BlueprintSpawnableComponent, ToolTip = "Chaos cloth simulation component."), DisplayName = "Cloth Simulation", HideCategories = (Physics, Constraints, Advanced, Cooking, Collision))
+UCLASS(ClassGroup = Physics, Meta = (BlueprintSpawnableComponent, ToolTip = "Chaos cloth simulation component."), DisplayName = "Cloth Simulation", HideCategories = (Object, "Mesh|SkeletalAsset", Physics, Constraints, Advanced, Cooking, Collision))
 class CHAOSCLOTHASSETENGINE_API UChaosClothComponent : public USkinnedMeshComponent
 {
 	GENERATED_BODY()
 public:
 	UChaosClothComponent(const FObjectInitializer& ObjectInitializer);
+	UChaosClothComponent(FVTableHelper& Helper);
+	~UChaosClothComponent();
 
 	/** Set the cloth asset used by this component. */
 	UFUNCTION(BlueprintCallable, Category = "Components|ClothAsset")
@@ -26,6 +34,27 @@ public:
 	UFUNCTION(BlueprintPure, Category = "Components|ClothAsset")
 	UChaosClothAsset* GetClothAsset() const;
 
+	/** Reset the teleport mode. */
+	UFUNCTION(BlueprintCallable, Category = "Components|ClothAsset|Teleport")
+	void ResetClothTeleportMode() { ClothTeleportMode = EClothingTeleportMode::None; }
+
+	/**
+	 * Used to indicate we should force 'teleport' during the next call to UpdateClothState,
+	 * This will transform positions and velocities and thus keep the simulation state, just translate it to a new pose.
+	 */
+	UFUNCTION(BlueprintCallable, Category = "Components|ClothAsset|Teleport")
+	void ForceClothNextUpdateTeleport() { ClothTeleportMode = EClothingTeleportMode::Teleport; }
+
+	/**
+	 * Used to indicate we should force 'teleport and reset' during the next call to UpdateClothState.
+	 * This can be used to reset it from a bad state or by a teleport where the old state is not important anymore.
+	 */
+	UFUNCTION(BlueprintCallable, Category = "Components|ClothAsset|Teleport")
+	void ForceClothNextUpdateTeleportAndReset() { ClothTeleportMode = EClothingTeleportMode::TeleportAndReset; }
+
+	/** Return the current teleport mode currently requested if any. */
+	EClothingTeleportMode GetClothTeleportMode() const { return ClothTeleportMode; }
+
 protected:
 	//~ Begin UObject Interface
 	virtual void PostLoad() override;
@@ -36,17 +65,50 @@ protected:
 
 	//~ Begin UActorComponent Interface
 	virtual void OnRegister() override;
+	virtual void OnUnregister() override;
+	virtual void TickComponent(float DeltaTime, enum ELevelTick TickType, FActorComponentTickFunction* ThisTickFunction) override;
+	virtual bool RequiresPreEndOfFrameSync() const override;
+	virtual void OnPreEndOfFrameSync() override;
 	//~ End UActorComponent Interface
 
-	//~ Begin USkinClothComponent Interface
+	//~ Begin USceneComponent Interface.
+	virtual FBoxSphereBounds CalcBounds(const FTransform& LocalToWorld) const override;
+	//~ End USceneComponent Interface
+
+	//~ Begin USkinnedMeshComponent Interface
 	virtual void RefreshBoneTransforms(FActorComponentTickFunction* TickFunction = nullptr) override;
-	//~ End USkinClothComponent Interface
+	virtual void GetUpdateClothSimulationData_AnyThread(TMap<int32, FClothSimulData>& OutClothSimulData, FMatrix& OutLocalToWorld, float& OutClothBlendWeight) override;
+	//~ End USkinnedMeshComponent Interface
 
 private:
+	void StartNewParallelSimulation(float DeltaTime);
+	void HandleExistingParallelSimulation();
+	bool ShouldWaitForClothInTickFunction() const;
+	bool IsSimulationSuspended() const;
+
 #if WITH_EDITORONLY_DATA
 	/** Cloth asset used by this component. */
 	UE_DEPRECATED(5.1, "This property isn't deprecated, but getter and setter must be used at all times to preserve correct operations.")
 	UPROPERTY(EditAnywhere, Transient, BlueprintSetter = SetClothAsset, BlueprintGetter = GetClothAsset, Category = Mesh)
 	TObjectPtr<UChaosClothAsset> ClothAsset;
 #endif
+
+	UPROPERTY()
+	uint8 bDisableClothSimulation : 1;
+
+	UPROPERTY()
+	uint8 bSuspendSimulation : 1;
+
+	UPROPERTY()
+	uint8 bWaitForParallelClothTask : 1;
+
+	UPROPERTY()
+	uint8 bBindClothToLeaderComponent : 1;
+
+	UPROPERTY()
+	float ClothBlendWeight = 1.f;
+
+	EClothingTeleportMode ClothTeleportMode = EClothingTeleportMode::None;
+
+	TUniquePtr<UE::Chaos::ClothAsset::FClothSimulationProxy> ClothSimulationProxy;
 };
diff --git a/Engine/Plugins/Experimental/ChaosClothAssetEditor/Source/ChaosClothAssetTools/Private/ChaosClothAsset/ClothAssetBuilderEditor.cpp b/Engine/Plugins/Experimental/ChaosClothAssetEditor/Source/ChaosClothAssetTools/Private/ChaosClothAsset/ClothAssetBuilderEditor.cpp
--- a/Engine/Plugins/Experimental/ChaosClothAssetEditor/Source/ChaosClothAssetTools/Private/ChaosClothAsset/ClothAssetBuilderEditor.cpp
+++ b/Engine/Plugins/Experimental/ChaosClothAssetEditor/Source/ChaosClothAssetTools/Private/ChaosClothAsset/ClothAssetBuilderEditor.cpp
@@ -1,12 +1,12 @@
 // Copyright Epic Games, Inc. All Rights Reserved.
 
 #include "ChaosClothAsset/ClothAssetBuilderEditor.h"
-
-#include "BoneWeights.h"
+#include "ChaosClothAsset/ClothAsset.h"
 #include "ChaosClothAsset/ClothCollection.h"
 #include "Modules/ModuleManager.h"
 #include "Rendering/SkeletalMeshLODModel.h"
 #include "Utils/ClothingMeshUtils.h"
+#include "BoneWeights.h"
 #include "MeshUtilities.h"
 #include "PointWeightMap.h"
 
@@ -65,36 +65,11 @@ void UClothAssetBuilderEditor::BuildLod(FSkeletalMeshLODModel& LODModel, const U
 	IMeshUtilities& MeshUtilities = FModuleManager::Get().LoadModuleChecked<IMeshUtilities>("MeshUtilities");
 
 	// Build the sim mesh descriptor for creation of the sections' mesh to mesh mapping data
-	// TODO: Generate Mesh to Mesh data from pattern space
-	const int32 LodSimVerticesStart = ClothCollection->SimVerticesStart[PatternStart];
-	const int32 LodSimVerticesEnd = ClothCollection->SimVerticesEnd[PatternEnd];
-	const int32 NumLodSimVertices = LodSimVerticesEnd - LodSimVerticesStart + 1;
-
-	const int32 LodSimFacesStart = ClothCollection->SimFacesStart[PatternStart];
-	const int32 LodSimFacesEnd = ClothCollection->SimFacesEnd[PatternEnd];
-
-	const int32 NumLodSimFaces = LodSimFacesEnd - LodSimFacesStart + 1;
-	const int32 NumLodIndices = NumLodSimFaces * 3;
-
-	TArray<uint32> LodSimIndices;
-	LodSimIndices.Reserve(NumLodIndices);
-
-	for (int32 SimFaceIndex = LodSimFacesStart; SimFaceIndex <= LodSimFacesEnd; ++SimFaceIndex)
-	{
-		const FIntVector3& SimIndices = ClothCollection->SimIndices[SimFaceIndex];
-		LodSimIndices.Add((uint32)SimIndices.X);
-		LodSimIndices.Add((uint32)SimIndices.Y);
-		LodSimIndices.Add((uint32)SimIndices.Z);
-
-		const FVector3f& A = ClothCollection->SimRestPosition[LodSimVerticesStart + SimIndices.X];
-		const FVector3f& B = ClothCollection->SimRestPosition[LodSimVerticesStart + SimIndices.Y];
-		const FVector3f& C = ClothCollection->SimRestPosition[LodSimVerticesStart + SimIndices.Z];
-		ensure(FVector3f::CrossProduct(B - A, C - A).SizeSquared() >= SMALL_NUMBER);
-	}
-
 	const ClothingMeshUtils::ClothMeshDesc SourceMesh(
-		TConstArrayView<FVector3f>(ClothCollection->SimRestPosition.GetData() + LodSimVerticesStart, NumLodSimVertices),
-		TConstArrayView<uint32>(LodSimIndices));  // Calculate averaged normals
+		GetSimPositions(ClothAsset, LodIndex),
+		GetSimIndices(ClothAsset, LodIndex));  // Let it calculate the averaged normals as to match the simulation data output
+
+	const int32 NumLodSimVertices = GetNumVertices(ClothAsset, LodIndex);
 
 	// Retrieve the MaxDistance map
 	FPointWeightMap MaxDistances;
@@ -112,7 +87,9 @@ void UClothAssetBuilderEditor::BuildLod(FSkeletalMeshLODModel& LODModel, const U
 	int32 BaseIndex = 0;
 	for (const TPair<int32, TArray<int32>>& SectionFaces : SectionFacesMap)
 	{
-		FSkelMeshSection& Section = LODModel.Sections[SectionIndex++];
+		FSkelMeshSection& Section = LODModel.Sections[SectionIndex];
+
+		Section.OriginalDataSectionIndex = SectionIndex++;
 
 		const int32 MaterialIndex = SectionFaces.Key;
 		const TArray<int32>& Faces = SectionFaces.Value;
@@ -266,6 +243,9 @@ void UClothAssetBuilderEditor::BuildLod(FSkeletalMeshLODModel& LODModel, const U
 				}
 			}
 		}
+
+		// Copy to user section data, otherwise the section data set above would get lost when the user section gets synced
+		FSkelMeshSourceSectionUserData::GetSourceSectionUserData(LODModel.UserSectionsData, Section);
 	}
 
 	// Remap the LOD indices with the new vertex indices
@@ -283,4 +263,3 @@ void UClothAssetBuilderEditor::BuildLod(FSkeletalMeshLODModel& LODModel, const U
 	// Compute the required bones for this model.
 	USkeletalMesh::CalculateRequiredBones(LODModel, ClothAsset.RefSkeleton, nullptr);
 }
-
