CIS - Validate Compiler Toolchains, fixed error C2039: 'Softs': is not a member of 'UE::Chaos', due to the compiler getting confused by the using namespace Softs in the statement while compiling an unrelated UE_CLOG macro.

#rnx

diff --git a/Engine/Plugins/Experimental/ChaosClothAsset/Source/ChaosClothAssetEngine/Private/ChaosClothAsset/ClothSimulationProxy.cpp b/Engine/Plugins/Experimental/ChaosClothAsset/Source/ChaosClothAssetEngine/Private/ChaosClothAsset/ClothSimulationProxy.cpp
--- a/Engine/Plugins/Experimental/ChaosClothAsset/Source/ChaosClothAssetEngine/Private/ChaosClothAsset/ClothSimulationProxy.cpp
+++ b/Engine/Plugins/Experimental/ChaosClothAsset/Source/ChaosClothAssetEngine/Private/ChaosClothAsset/ClothSimulationProxy.cpp
@@ -264,7 +264,6 @@ PRAGMA_ENABLE_DEPRECATION_WARNINGS
 	void FClothSimulationProxy::WriteSimulationData_GameThread()
 	{
 		using namespace ::Chaos;
-		using namespace ::Chaos::Softs;
 
 		check(IsInGameThread());
 
@@ -320,7 +319,7 @@ PRAGMA_DISABLE_DEPRECATION_WARNINGS  // TODO: CHAOS_IS_CLOTHINGSIMULATIONMESH_AB
 PRAGMA_ENABLE_DEPRECATION_WARNINGS
 			if (LODIndex != Cloth->GetLODIndex(Solver.Get()))
 			{
-				Solver->Update(FSolverReal(0.));  // Update for LOD switching, but do not simulate
+				Solver->Update(Softs::FSolverReal(0.));  // Update for LOD switching, but do not simulate
 			}
 
 			if (Cloth->GetOffset(Solver.Get()) == INDEX_NONE || Cloth->GetLODIndex(Solver.Get()) == INDEX_NONE)
