CIS - Fixed duplicate parallel cloth task priority name, and updated stat id to match.

diff --git a/Engine/Plugins/Experimental/ChaosClothAsset/Source/ChaosClothAssetEngine/Private/ChaosClothAsset/ClothSimulationProxy.cpp b/Engine/Plugins/Experimental/ChaosClothAsset/Source/ChaosClothAssetEngine/Private/ChaosClothAsset/ClothSimulationProxy.cpp
--- a/Engine/Plugins/Experimental/ChaosClothAsset/Source/ChaosClothAssetEngine/Private/ChaosClothAsset/ClothSimulationProxy.cpp
+++ b/Engine/Plugins/Experimental/ChaosClothAsset/Source/ChaosClothAssetEngine/Private/ChaosClothAsset/ClothSimulationProxy.cpp
@@ -33,24 +33,24 @@ namespace UE::Chaos::ClothAsset
 	FAutoConsoleVariableRef CVarTransformClothSimukDataISPCEnabled(TEXT("p.ChaosClothAsset.TransformClothSimulData.ISPC"), bTransformClothSimulData_ISPC_Enabled, TEXT("Whether to use ISPC optimizations when transforming simulation data back to reference bone space."));
 #endif
 
-	static FAutoConsoleTaskPriority CPrio_FParallelClothTask(
-		TEXT("TaskGraph.TaskPriorities.ParallelClothTask"),
-		TEXT("Task and thread priority for parallel cloth."),
+	static FAutoConsoleTaskPriority CPrio_ClothSimulationProxyParallelTask(
+		TEXT("TaskGraph.TaskPriorities.ClothSimulationProxyParallelTask"),
+		TEXT("Task and thread priority for the cloth simulation proxy."),
 		ENamedThreads::HighThreadPriority, // If we have high priority task threads, then use them...
 		ENamedThreads::NormalTaskPriority, // .. at normal task priority
 		ENamedThreads::HighTaskPriority);  // If we don't have high priority threads, then use normal priority threads at high task priority instead
 
-	class FClothSimulationProxy::FParallelClothTask
+	class FClothSimulationProxyParallelTask
 	{
 	public:
-		FParallelClothTask(FClothSimulationProxy& InClothSimulationProxy)
+		FClothSimulationProxyParallelTask(FClothSimulationProxy& InClothSimulationProxy)
 			: ClothSimulationProxy(InClothSimulationProxy)
 		{
 		}
 
 		TStatId GetStatId() const
 		{
-			RETURN_QUICK_DECLARE_CYCLE_STAT(FParallelClothTask, STATGROUP_TaskGraphTasks);
+			RETURN_QUICK_DECLARE_CYCLE_STAT(FClothSimulationProxyParallelTask, STATGROUP_TaskGraphTasks);
 		}
 
 		static ENamedThreads::Type GetDesiredThread()
@@ -59,7 +59,7 @@ namespace UE::Chaos::ClothAsset
 
 			if (CVarClothPhysicsUseTaskThread && CVarClothPhysicsUseTaskThread->GetBool())
 			{
-				return CPrio_FParallelClothTask.Get();
+				return CPrio_ClothSimulationProxyParallelTask.Get();
 			}
 			return ENamedThreads::GameThread;
 		}
@@ -182,7 +182,7 @@ PRAGMA_ENABLE_DEPRECATION_WARNINGS
 			ClothSimulationContext.Fill(ClothComponent, DeltaTime, MaxDeltaTime);
 
 			// Start the the cloth simulation thread
-			ParallelClothTask = TGraphTask<FParallelClothTask>::CreateTask(nullptr, ENamedThreads::GameThread).ConstructAndDispatchWhenReady(*this);
+			ParallelTask = TGraphTask<FClothSimulationProxyParallelTask>::CreateTask(nullptr, ENamedThreads::GameThread).ConstructAndDispatchWhenReady(*this);
 		}
 	}
 
@@ -245,16 +245,16 @@ PRAGMA_ENABLE_DEPRECATION_WARNINGS
 	{
 		check(IsInGameThread());
 
-		if (IsValidRef(ParallelClothTask))
+		if (IsValidRef(ParallelTask))
 		{
 			SCOPE_CYCLE_COUNTER(STAT_ClothSimulationProxy_EndParallelClothTask);
 			CSV_SCOPED_SET_WAIT_STAT(Cloth);
 
 			// There's a simulation in flight
-			FTaskGraphInterface::Get().WaitUntilTaskCompletes(ParallelClothTask, ENamedThreads::GameThread);
+			FTaskGraphInterface::Get().WaitUntilTaskCompletes(ParallelTask, ENamedThreads::GameThread);
 
 			// No longer need this task, it has completed
-			ParallelClothTask.SafeRelease();
+			ParallelTask.SafeRelease();
 
 			// Write back to the GT cache
 			WriteSimulationData_GameThread();
@@ -395,7 +395,7 @@ PRAGMA_ENABLE_DEPRECATION_WARNINGS
 		// This is called during EndOfFrameUpdates, usually in a parallel-for loop. We need to be sure that
 		// the cloth task (if there is one) is complete, but it cannot be waited for here. See OnPreEndOfFrameUpdateSync
 		// which is called just before EOF updates and is where we would have waited for the cloth task.
-		if (!IsValidRef(ParallelClothTask) || ParallelClothTask->IsComplete())
+		if (!IsValidRef(ParallelTask) || ParallelTask->IsComplete())
 		{
 			return CurrentSimulationData;
 		}
@@ -408,7 +408,7 @@ PRAGMA_ENABLE_DEPRECATION_WARNINGS
 		SCOPE_CYCLE_COUNTER(STAT_ClothSimulationProxy_CalculateBounds);
 
 		check(Solver);
-		if (!IsValidRef(ParallelClothTask) || ParallelClothTask->IsComplete())
+		if (!IsValidRef(ParallelTask) || ParallelTask->IsComplete())
 		{
 			FBoxSphereBounds Bounds = Solver->CalculateBounds();
 
diff --git a/Engine/Plugins/Experimental/ChaosClothAsset/Source/ChaosClothAssetEngine/Private/ChaosClothAsset/ClothSimulationProxy.h b/Engine/Plugins/Experimental/ChaosClothAsset/Source/ChaosClothAssetEngine/Private/ChaosClothAsset/ClothSimulationProxy.h
--- a/Engine/Plugins/Experimental/ChaosClothAsset/Source/ChaosClothAssetEngine/Private/ChaosClothAsset/ClothSimulationProxy.h
+++ b/Engine/Plugins/Experimental/ChaosClothAsset/Source/ChaosClothAssetEngine/Private/ChaosClothAsset/ClothSimulationProxy.h
@@ -58,10 +58,10 @@ namespace UE::Chaos::ClothAsset
 		void WriteSimulationData_GameThread();
 
 		// Internal physics thread object
-		class FParallelClothTask;
+		friend class FClothSimulationProxyParallelTask;
 
-		// Reference for the clothing parallel task, to detect whether or not a simulation is running
-		FGraphEventRef ParallelClothTask;
+		// Reference for the cloth parallel task, to detect whether or not a simulation is running
+		FGraphEventRef ParallelTask;
 
 		// Simulation data written back to the component after the simulation has taken place
 		TMap<int32, FClothSimulData> CurrentSimulationData;
