CIS - Data Pipeline Tests, fixed IntProperty FChaosClothSimulationModel::ReferenceBoneIndex is not initialized properly.

#rnx

diff --git a/Engine/Plugins/Experimental/ChaosClothAsset/Source/ChaosClothAssetEngine/Private/ChaosClothAsset/ClothSimulationModel.h b/Engine/Plugins/Experimental/ChaosClothAsset/Source/ChaosClothAssetEngine/Private/ChaosClothAsset/ClothSimulationModel.h
--- a/Engine/Plugins/Experimental/ChaosClothAsset/Source/ChaosClothAssetEngine/Private/ChaosClothAsset/ClothSimulationModel.h
+++ b/Engine/Plugins/Experimental/ChaosClothAsset/Source/ChaosClothAssetEngine/Private/ChaosClothAsset/ClothSimulationModel.h
@@ -64,7 +64,7 @@ struct FChaosClothSimulationModel
 
 	/** Bone to treat as the root of the simulation space. */
 	UPROPERTY()
-	int32 ReferenceBoneIndex;
+	int32 ReferenceBoneIndex = INDEX_NONE;
 
 	FChaosClothSimulationModel() = default;
 	FChaosClothSimulationModel(const TSharedPtr<const UE::Chaos::ClothAsset::FClothCollection>& ClothCollection, const FReferenceSkeleton& ReferenceSkeleton);
