Chaos Cloth:
- Fixed the seams creation on the Cloth Asset legacy exporter.
- Added open edge visualization.
- Fixed the element indices visualization.
- Moved the ClothSimulationModel to a shared pointer, so that its ownership can be transferred to the simulation when it is recreated.
- Tidied up a couple of issues related to the recent global change in header inclusions.

#rnx

diff --git a/Engine/Plugins/Enterprise/DatasmithImporter/Source/DatasmithImporter/Private/DatasmithImporter.cpp b/Engine/Plugins/Enterprise/DatasmithImporter/Source/DatasmithImporter/Private/DatasmithImporter.cpp
--- a/Engine/Plugins/Enterprise/DatasmithImporter/Source/DatasmithImporter/Private/DatasmithImporter.cpp
+++ b/Engine/Plugins/Enterprise/DatasmithImporter/Source/DatasmithImporter/Private/DatasmithImporter.cpp
@@ -45,6 +45,7 @@
 #include "Engine/Engine.h"
 #include "Engine/RendererSettings.h"
 #include "Engine/SkinnedAsset.h"
+#include "Engine/SkinnedAssetCommon.h"
 #include "Engine/StaticMesh.h"
 #include "Engine/StaticMeshActor.h"
 #include "HAL/FileManager.h"
diff --git a/Engine/Plugins/Experimental/ChaosCloth/Source/Chaos/Private/ChaosCloth/ChaosClothVisualization.cpp b/Engine/Plugins/Experimental/ChaosCloth/Source/Chaos/Private/ChaosCloth/ChaosClothVisualization.cpp
--- a/Engine/Plugins/Experimental/ChaosCloth/Source/Chaos/Private/ChaosCloth/ChaosClothVisualization.cpp
+++ b/Engine/Plugins/Experimental/ChaosCloth/Source/Chaos/Private/ChaosCloth/ChaosClothVisualization.cpp
@@ -181,7 +181,10 @@ namespace Chaos
 			for (int32 Index = 0; Index < Elements.Num(); ++Index)
 			{
 				const TVec3<int32>& Element = Elements[Index];
-				const FVector Position = LocalSpaceLocation + FVector(Positions[Element[0]] + Positions[Element[1]] + Positions[Element[2]]) / (FReal)3.;
+				const FVector Position = LocalSpaceLocation + (
+					FVector(Positions[Element[0] - Offset]) +
+					FVector(Positions[Element[1] - Offset]) +
+					FVector(Positions[Element[2] - Offset])) / (FReal)3.;
 
 				const bool bIsKinematic0 = (InvMasses[Element.X - Offset] == (Softs::FSolverReal)0.);
 				const bool bIsKinematic1 = (InvMasses[Element.Y - Offset] == (Softs::FSolverReal)0.);
@@ -505,7 +508,7 @@ namespace Chaos
 
 			for (int32 ElementIndex = 0; ElementIndex < Elements.Num(); ++ElementIndex)
 			{
-				const auto& Element = Elements[ElementIndex];
+				const TVec3<int32>& Element = Elements[ElementIndex];
 
 				const FVector Pos0 = LocalSpaceLocation + FVector(Positions[Element.X - Offset]); // TODO: Triangle Mesh shouldn't really be solver dependent (ie not use an offset)
 				const FVector Pos1 = LocalSpaceLocation + FVector(Positions[Element.Y - Offset]);
@@ -546,7 +549,7 @@ namespace Chaos
 
 			for (int32 ElementIndex = 0; ElementIndex < Elements.Num(); ++ElementIndex)
 			{
-				const auto& Element = Elements[ElementIndex];
+				const TVec3<int32>& Element = Elements[ElementIndex];
 
 				const FVector Pos0 = LocalSpaceLocation + FVector(Positions[Element.X - Offset]); // TODO: Triangle Mesh shouldn't really be solver dependent (ie not use an offset)
 				const FVector Pos1 = LocalSpaceLocation + FVector(Positions[Element.Y - Offset]);
@@ -559,6 +562,72 @@ namespace Chaos
 		}
 	}
 
+	void FClothVisualization::DrawOpenEdges(FPrimitiveDrawInterface* PDI) const
+	{
+		auto MakeSortedUintVector2 = [](uint32 Index0, uint32 Index1) -> FUintVector2
+			{
+				return Index0 < Index1 ? FUintVector2(Index0, Index1) : FUintVector2(Index1, Index0);
+			};
+
+		auto BuildEdgeMap = [&MakeSortedUintVector2](const TConstArrayView<TVec3<int32>>& Elements, TMap<FUintVector2, TArray<uint32>>& OutEdgeToTrianglesMap)
+			{
+			OutEdgeToTrianglesMap.Empty(Elements.Num() * 2);  // Rough estimate for the number of edges
+
+				for (int32 ElementIndex = 0; ElementIndex < Elements.Num(); ++ElementIndex)
+				{
+					const TVec3<int32>& Element = Elements[ElementIndex];
+					const uint32 Index0 = Element[0];
+					const uint32 Index1 = Element[1];
+					const uint32 Index2 = Element[2];
+
+					const FUintVector2 Edge0 = MakeSortedUintVector2(Index0, Index1);
+					const FUintVector2 Edge1 = MakeSortedUintVector2(Index1, Index2);
+					const FUintVector2 Edge2 = MakeSortedUintVector2(Index2, Index0);
+
+					OutEdgeToTrianglesMap.FindOrAdd(Edge0).Add(ElementIndex);
+					OutEdgeToTrianglesMap.FindOrAdd(Edge1).Add(ElementIndex);
+					OutEdgeToTrianglesMap.FindOrAdd(Edge2).Add(ElementIndex);
+				}
+			};
+
+		if (!Solver)
+		{
+			return;
+		}
+
+		static const FLinearColor OpenedEdgeColor = FColor::Emerald;
+		static const FLinearColor ClosedEdgeColor = FColor::White;
+
+		const FVec3& LocalSpaceLocation = Solver->GetLocalSpaceLocation();
+
+		for (const FClothingSimulationCloth* const Cloth : Solver->GetCloths())
+		{
+			const int32 Offset = Cloth->GetOffset(Solver);
+			if (Offset == INDEX_NONE)
+			{
+				continue;
+			}
+
+			const TConstArrayView<TVec3<int32>> Elements = Cloth->GetTriangleMesh(Solver).GetElements();
+			const TConstArrayView<Softs::FSolverVec3> Positions = Cloth->GetAnimationPositions(Solver);
+
+			TMap<FUintVector2, TArray<uint32>> EdgeToTrianglesMap;
+			BuildEdgeMap(Elements, EdgeToTrianglesMap);
+
+			for (const TPair<FUintVector2, TArray<uint32>>& EdgeToTriangles : EdgeToTrianglesMap)
+			{
+				const FUintVector2& Edge = EdgeToTriangles.Key;
+				const TArray<uint32>& Triangles = EdgeToTriangles.Value;
+
+				const FVector Pos0 = LocalSpaceLocation + FVector(Positions[Edge[0] - Offset]); // TODO: Triangle Mesh shouldn't really be solver dependent (ie not use an offset)
+				const FVector Pos1 = LocalSpaceLocation + FVector(Positions[Edge[1] - Offset]);
+				const FLinearColor& Color = (Triangles.Num() > 1) ? ClosedEdgeColor : OpenedEdgeColor;
+
+				DrawLine(PDI, Pos0, Pos1, Color);
+			}
+		}
+	}
+
 	void FClothVisualization::DrawAnimNormals(FPrimitiveDrawInterface* PDI) const
 	{
 		if (!Solver)
@@ -810,15 +879,15 @@ namespace Chaos
 		}
 
 		auto DrawBackstop = [PDI](const FVector& Position, const FVector& Normal, FReal Radius, const FVector& Axis, const FLinearColor& Color)
-		{
-			static const FReal MaxCosAngle = (FReal)0.99;
-			if (FMath::Abs(FVector::DotProduct(Normal, Axis)) < MaxCosAngle)
 			{
-				static const FReal ArcLength = (FReal)5.; // Arch length in cm
-				const FReal ArcAngle = (FReal)360. * ArcLength / FMath::Max((Radius * (FReal)2. * (FReal)PI), ArcLength);
-				DrawArc(PDI, Position, Normal, FVector::CrossProduct(Axis, Normal).GetSafeNormal(), -ArcAngle / (FReal)2., ArcAngle / (FReal)2., Radius, Color);
-			}
-		};
+				static const FReal MaxCosAngle = (FReal)0.99;
+				if (FMath::Abs(FVector::DotProduct(Normal, Axis)) < MaxCosAngle)
+				{
+					static const FReal ArcLength = (FReal)5.; // Arch length in cm
+					const FReal ArcAngle = (FReal)360. * ArcLength / FMath::Max((Radius * (FReal)2. * (FReal)PI), ArcLength);
+					DrawArc(PDI, Position, Normal, FVector::CrossProduct(Axis, Normal).GetSafeNormal(), -ArcAngle / (FReal)2., ArcAngle / (FReal)2., Radius, Color);
+				}
+			};
 
 		const FVec3& LocalSpaceLocation = Solver->GetLocalSpaceLocation();
 
@@ -1026,9 +1095,9 @@ namespace Chaos
 			const uint8 HueOffset = 196 / ConstraintColorNum;
 			auto ConstraintColor =
 				[HueOffset](int32 ColorIndex)->FLinearColor
-			{
-				return FLinearColor::MakeFromHSV8(ColorIndex * HueOffset, 255, 255);
-			};
+				{
+					return FLinearColor::MakeFromHSV8(ColorIndex * HueOffset, 255, 255);
+				};
 
 			for (int32 ConstraintColorIndex = 0; ConstraintColorIndex < ConstraintColorNum; ++ConstraintColorIndex)
 			{
diff --git a/Engine/Plugins/Experimental/ChaosCloth/Source/Chaos/Public/ChaosCloth/ChaosClothVisualization.h b/Engine/Plugins/Experimental/ChaosCloth/Source/Chaos/Public/ChaosCloth/ChaosClothVisualization.h
--- a/Engine/Plugins/Experimental/ChaosCloth/Source/Chaos/Public/ChaosCloth/ChaosClothVisualization.h
+++ b/Engine/Plugins/Experimental/ChaosCloth/Source/Chaos/Public/ChaosCloth/ChaosClothVisualization.h
@@ -33,6 +33,7 @@ namespace Chaos
 		void DrawPhysMeshWired(FPrimitiveDrawInterface* PDI = nullptr) const;
 		void DrawAnimMeshWired(FPrimitiveDrawInterface* PDI = nullptr) const;
 		void DrawAnimNormals(FPrimitiveDrawInterface* PDI = nullptr) const;
+		void DrawOpenEdges(FPrimitiveDrawInterface* PDI = nullptr) const;
 		void DrawPointNormals(FPrimitiveDrawInterface* PDI = nullptr) const;
 		void DrawPointVelocities(FPrimitiveDrawInterface* PDI = nullptr) const;
 		void DrawCollision(FPrimitiveDrawInterface* PDI = nullptr) const;
@@ -56,6 +57,7 @@ namespace Chaos
 		void DrawPhysMeshWired(FPrimitiveDrawInterface* /*PDI*/ = nullptr) const {}
 		void DrawAnimMeshWired(FPrimitiveDrawInterface* /*PDI*/ = nullptr) const {}
 		void DrawAnimNormals(FPrimitiveDrawInterface* /*PDI*/ = nullptr) const {}
+		void DrawOpenEdges(FPrimitiveDrawInterface* /*PDI*/ = nullptr) const {}
 		void DrawPointNormals(FPrimitiveDrawInterface* /*PDI*/ = nullptr) const {}
 		void DrawPointVelocities(FPrimitiveDrawInterface* /*PDI*/ = nullptr) const {}
 		void DrawCollision(FPrimitiveDrawInterface* /*PDI*/ = nullptr) const {}
diff --git a/Engine/Plugins/Experimental/ChaosClothAsset/Source/ChaosClothAsset/Private/ChaosClothAsset/ClothLodAdapter.cpp b/Engine/Plugins/Experimental/ChaosClothAsset/Source/ChaosClothAsset/Private/ChaosClothAsset/ClothLodAdapter.cpp
--- a/Engine/Plugins/Experimental/ChaosClothAsset/Source/ChaosClothAsset/Private/ChaosClothAsset/ClothLodAdapter.cpp
+++ b/Engine/Plugins/Experimental/ChaosClothAsset/Source/ChaosClothAsset/Private/ChaosClothAsset/ClothLodAdapter.cpp
@@ -305,30 +305,65 @@ namespace UE::Chaos::ClothAsset::Private
 					const uint32 TriangleBase = Triangle * 3;
 					const uint32 TriangleIndex0 = Island.SourceIndices[TriangleBase + 0];
 					const uint32 TriangleIndex1 = Island.SourceIndices[TriangleBase + 1];
-					return
-						(SourceEdge[0] == TriangleIndex0 || SourceEdge[1] == TriangleIndex0) ?
-						(SourceEdge[0] == TriangleIndex1 || SourceEdge[1] == TriangleIndex1) ?
-						MakeSortedUintVector2(Island.Indices[TriangleBase + 0], Island.Indices[TriangleBase + 1]) :  // Edge 01
-						MakeSortedUintVector2(Island.Indices[TriangleBase + 2], Island.Indices[TriangleBase + 0]) :  // Edge 20
-						MakeSortedUintVector2(Island.Indices[TriangleBase + 1], Island.Indices[TriangleBase + 2]);   // Edge 12
+					const uint32 TriangleIndex2 = Island.SourceIndices[TriangleBase + 2];
+
+					if (SourceEdge[0] == TriangleIndex0)
+					{
+						if (SourceEdge[1] == TriangleIndex1)
+						{
+							return FUintVector2(Island.Indices[TriangleBase + 0], Island.Indices[TriangleBase + 1]);  // Edge 0 1
+						}
+						else if (SourceEdge[1] == TriangleIndex2)
+						{
+							return FUintVector2(Island.Indices[TriangleBase + 0], Island.Indices[TriangleBase + 2]);  // Edge 0 2
+						}
+					}
+					else if (SourceEdge[0] == TriangleIndex1)
+					{
+						if (SourceEdge[1] == TriangleIndex0)
+						{
+							return FUintVector2(Island.Indices[TriangleBase + 1], Island.Indices[TriangleBase + 0]);  // Edge 1 0
+						}
+						else if (SourceEdge[1] == TriangleIndex2)
+						{
+							return FUintVector2(Island.Indices[TriangleBase + 1], Island.Indices[TriangleBase + 2]);  // Edge 1 2
+						}
+					}
+					else if (SourceEdge[0] == TriangleIndex2)
+					{
+						if (SourceEdge[1] == TriangleIndex0)
+						{
+							return FUintVector2(Island.Indices[TriangleBase + 2], Island.Indices[TriangleBase + 0]);  // Edge 2 0
+						}
+						else if (SourceEdge[1] == TriangleIndex1)
+						{
+							return FUintVector2(Island.Indices[TriangleBase + 2], Island.Indices[TriangleBase + 1]);  // Edge 2 1
+						}
+					}
+					check(false);  // Sanity check
+					return FUintVector2(0);
 				};
 
 			// Look for disconnected triangles
 			for (const TPair<FUintVector2, TArray<uint32>>& SourceEdgeToTriangles : SourceEdgeToTrianglesMap)
 			{
 				const TArray<uint32>& Triangles = SourceEdgeToTriangles.Value;
-				const uint32 NumTriangles = (uint32)Triangles.Num();
+				const int32 NumTriangles = Triangles.Num();
 				if (NumTriangles > 1)
 				{
 					const FUintVector2& SourceEdge = SourceEdgeToTriangles.Key;
 
-					for (uint32 TriangleIndex0 = 0; TriangleIndex0 < NumTriangles - 1; ++TriangleIndex0)
+					for (int32 Index0 = 0; Index0 < NumTriangles - 1; ++Index0)
 					{
+						const uint32 TriangleIndex0 = Triangles[Index0];
+
 						const FUintVector2 Edge0 = GetTriangleEdgeMatchingSourceEdge(TriangleIndex0, SourceEdge);
 						check(Edge0[0] <= (uint32)TNumericLimits<int32>::Max() && Edge0[1] <= (uint32)TNumericLimits<int32>::Max());
 
-						for (uint32 TriangleIndex1 = TriangleIndex0 + 1; TriangleIndex1 < NumTriangles; ++TriangleIndex1)
+						for (int32 Index1 = Index0 + 1; Index1 < NumTriangles; ++Index1)
 						{
+							const uint32 TriangleIndex1 = Triangles[Index1];
+
 							const FUintVector2 Edge1 = GetTriangleEdgeMatchingSourceEdge(TriangleIndex1, SourceEdge);
 							check(Edge1[0] <= (uint32)TNumericLimits<int32>::Max() && Edge1[1] <= (uint32)TNumericLimits<int32>::Max());
 
@@ -380,7 +415,9 @@ namespace UE::Chaos::ClothAsset
 	{
 		const int32 NumSimVertices = GetPatternsNumSimVertices();
 
-		// Initialize welding map
+		// Initialize welding map with same index
+		// The welding map redirects to an existing vertex index if these two are part of the same welding group.
+		// The redirected index must be the smallest index in the group.
 		TArray<int32> WeldingMap;
 		WeldingMap.SetNumUninitialized(NumSimVertices);
 		for (int32 SimVertexIndex = 0; SimVertexIndex < NumSimVertices; ++SimVertexIndex)
@@ -388,42 +425,58 @@ namespace UE::Chaos::ClothAsset
 			WeldingMap[SimVertexIndex] = SimVertexIndex;
 		}
 
+		// Define welding groups
+		// Welding groups contain all stitched pair of indices to be welded together that are required to build the welding map.
+		// Key is the smallest redirected index in the group, and will be the one index used in the welding map redirects.
 		TMap<int32, TSet<int32>> WeldingGroups;
 
 		auto UpdateWeldingMap =
 			[&WeldingMap, &WeldingGroups](int32 Index0, int32 Index1)
 		{
+			// Only process pairs that are not already redirected to the same index
 			if (WeldingMap[Index0] != WeldingMap[Index1])
 			{
+				// Make sure Index0 points to the the smallest redirected index, so that merges are done into the correct group
 				if (WeldingMap[Index0] > WeldingMap[Index1])
 				{
 					Swap(Index0, Index1);
 				}
 
-				TSet<int32>* WeldingGroup0 = WeldingGroups.Find(WeldingMap[Index0]);
+				// Find the group for Index0 if any
+				const int32 Key0 = WeldingMap[Index0];
+				TSet<int32>* WeldingGroup0 = WeldingGroups.Find(Key0);
 				if (!WeldingGroup0)
 				{
-					WeldingGroup0 = &WeldingGroups.Add(WeldingMap[Index0]);
+					// No existing group, create a new one
+					check(Key0 == Index0);  // No group means this index can't already have been redirected  // TODO: Make this a checkSlow
+					WeldingGroup0 = &WeldingGroups.Add(Key0);
 					WeldingGroup0->Add(Index0);
 				}
 
-				if (TSet<int32>* const WeldingGroup1 = WeldingGroups.Find(WeldingMap[Index1]))
+				// Find the group for Index1, if it exists merge the two groups
+				const int32 Key1 = WeldingMap[Index1];
+				if (TSet<int32>* const WeldingGroup1 = WeldingGroups.Find(Key1))
 				{
-					// Update group1 mapped indices
+					// Update group1 redirected indices with the new key
 					for (int32 Index : *WeldingGroup1)
 					{
-						WeldingMap[Index] = WeldingMap[Index0];
+						WeldingMap[Index] = Key0;
 					}
 
 					// Merge group0 & group1
 					WeldingGroup0->Append(*WeldingGroup1);
 
 					// Remove group1
-					WeldingGroups.Remove(WeldingMap[Index1]);
+					WeldingGroups.Remove(Key1);
+
+					// Sanity check
+					check(WeldingGroup0->Contains(Key0) && WeldingGroup0->Contains(Key1));  // TODO: Make this a checkSlow
 				}
 				else
 				{
-					WeldingMap[Index1] = WeldingMap[Index0];
+					// Otherwise add Index1 to Index0's group
+					check(Key1 == Index1);  // No group means this index can't already have been redirected  // TODO: Make this a checkSlow
+					WeldingMap[Index1] = Key0;
 					WeldingGroup0->Add(Index1);
 				}
 			}
@@ -432,7 +485,6 @@ namespace UE::Chaos::ClothAsset
 		// Apply all seams
 		const int32 NumSeams = GetNumSeams();
 		const TConstArrayView<TArray<FIntVector2>> SeamStitches = GetSeamStitches();
-		const TConstArrayView<FIntVector2> SeamPatternFirst = GetSeamPatterns();
 		for (int32 SeamIndex = 0; SeamIndex < NumSeams; ++SeamIndex)
 		{
 			for (const FIntVector2& Stitch : SeamStitches[SeamIndex])
@@ -591,10 +643,19 @@ namespace UE::Chaos::ClothAsset
 
 		const TArrayView<TArray<FIntVector2>> SeamStitches = GetSeamStitches();
 		const TArrayView<FIntVector2> SeamPatterns = GetSeamPatterns();
+		const TArrayView<int32> SimVerticesStart = GetSimVerticesStart();
+
 		for (int32 SeamIndex = 0; SeamIndex < Seams.Num(); ++SeamIndex)
 		{
-			SeamPatterns[SeamIndex] = Seams[SeamIndex].Patterns;
-			SeamStitches[SeamIndex] = Seams[SeamIndex].Stitches.Array();
+			const FIntVector2& Patterns = Seams[SeamIndex].Patterns;
+			SeamPatterns[SeamIndex] = Patterns;
+
+			const int32 NumStitches = Seams[SeamIndex].Stitches.Num();
+			SeamStitches[SeamIndex].Reset(NumStitches);
+			for (const FIntVector2& Stitch : Seams[SeamIndex].Stitches)
+			{
+				SeamStitches[SeamIndex].Emplace(Stitch[0] + SimVerticesStart[Patterns[0]], Stitch[1] + SimVerticesStart[Patterns[1]]);
+			}
 		}
 	}
 }  // End namespace UE::Chaos::ClothAsset
diff --git a/Engine/Plugins/Experimental/ChaosClothAsset/Source/ChaosClothAssetEngine/Private/ChaosClothAsset/ClothAsset.cpp b/Engine/Plugins/Experimental/ChaosClothAsset/Source/ChaosClothAssetEngine/Private/ChaosClothAsset/ClothAsset.cpp
--- a/Engine/Plugins/Experimental/ChaosClothAsset/Source/ChaosClothAssetEngine/Private/ChaosClothAsset/ClothAsset.cpp
+++ b/Engine/Plugins/Experimental/ChaosClothAsset/Source/ChaosClothAssetEngine/Private/ChaosClothAsset/ClothAsset.cpp
@@ -40,6 +40,16 @@ UChaosClothAsset::UChaosClothAsset(FVTableHelper& Helper)
 
 UChaosClothAsset::~UChaosClothAsset() = default;
 
+FSkeletalMeshLODInfo* UChaosClothAsset::GetLODInfo(int32 Index)
+{
+	return LODInfo.IsValidIndex(Index) ? &LODInfo[Index] : nullptr;
+}
+
+const FSkeletalMeshLODInfo* UChaosClothAsset::GetLODInfo(int32 Index) const
+{
+	return LODInfo.IsValidIndex(Index) ? &LODInfo[Index] : nullptr;
+}
+
 FMatrix UChaosClothAsset::GetComposedRefPoseMatrix(FName InBoneName) const
 {
 	FMatrix LocalPose(FMatrix::Identity);
@@ -93,17 +103,7 @@ void UChaosClothAsset::PostEditChangeProperty(FPropertyChangedEvent& PropertyCha
 {
 	if (PropertyChangedEvent.GetPropertyName() == GET_MEMBER_NAME_CHECKED(UChaosClothAsset, PhysicsAsset))
 	{
-		// Recreate the simulation proxies with the updated physics asset
-		for (TObjectIterator<UChaosClothComponent> ObjectIterator; ObjectIterator; ++ObjectIterator)
-		{
-			if (UChaosClothComponent* const Component = *ObjectIterator)
-			{
-				if (Component->GetClothAsset() == this)
-				{
-					const FComponentReregisterContext Context(Component);  // Context goes out of scope, causing the Component to be re-registered
-				}
-			}
-		}
+		ReregisterComponents();
 	}
 	Super::PostEditChangeProperty(PropertyChangedEvent);
 }
@@ -365,6 +365,9 @@ void UChaosClothAsset::Build()
 	{
 		InitResources();
 	}
+
+	// Re-register any components using this asset to restart the simulation with the updated asset
+	ReregisterComponents();
 }
 
 #if WITH_EDITORONLY_DATA
@@ -396,7 +399,7 @@ void UChaosClothAsset::BuildMeshModel()
 
 void UChaosClothAsset::BuildClothSimulationModel()
 {
-	ClothSimulationModel = MakeUnique<FChaosClothSimulationModel>(GetClothCollection(), GetRefSkeleton());
+	ClothSimulationModel = MakeShared<FChaosClothSimulationModel>(GetClothCollection(), GetRefSkeleton());
 }
 
 const FMeshUVChannelInfo* UChaosClothAsset::GetUVChannelData(int32 MaterialIndex) const
@@ -534,3 +537,18 @@ void UChaosClothAsset::CopySimMeshToRenderMesh(int32 MaterialIndex)
 	check(ClothCollection.IsValid());
 	FClothGeometryTools::CopySimMeshToRenderMesh(ClothCollection, MaterialIndex);
 }
+
+void UChaosClothAsset::ReregisterComponents()
+{
+	// Recreate the simulation proxies with the updated physics asset
+	for (TObjectIterator<UChaosClothComponent> ObjectIterator; ObjectIterator; ++ObjectIterator)
+	{
+		if (UChaosClothComponent* const Component = *ObjectIterator)
+		{
+			if (Component->GetClothAsset() == this)
+			{
+				const FComponentReregisterContext Context(Component);  // Context goes out of scope, causing the Component to be re-registered
+			}
+		}
+	}
+}
diff --git a/Engine/Plugins/Experimental/ChaosClothAsset/Source/ChaosClothAssetEngine/Private/ChaosClothAsset/ClothAssetBuilder.cpp b/Engine/Plugins/Experimental/ChaosClothAsset/Source/ChaosClothAssetEngine/Private/ChaosClothAsset/ClothAssetBuilder.cpp
--- a/Engine/Plugins/Experimental/ChaosClothAsset/Source/ChaosClothAssetEngine/Private/ChaosClothAsset/ClothAssetBuilder.cpp
+++ b/Engine/Plugins/Experimental/ChaosClothAsset/Source/ChaosClothAssetEngine/Private/ChaosClothAsset/ClothAssetBuilder.cpp
@@ -6,7 +6,7 @@
 
 int32 UClothAssetBuilder::GetNumVertices(const UChaosClothAsset& ClothAsset, int32 LodIndex)
 {
-	const FChaosClothSimulationModel* const ClothSimulationModel = ClothAsset.GetClothSimulationModel();
+	const TSharedPtr<const FChaosClothSimulationModel> ClothSimulationModel = ClothAsset.GetClothSimulationModel();
 	return ClothSimulationModel && ClothSimulationModel->IsValidLodIndex(LodIndex) ?
 		ClothSimulationModel->GetNumVertices(LodIndex) :
 		0;
@@ -14,7 +14,7 @@ int32 UClothAssetBuilder::GetNumVertices(const UChaosClothAsset& ClothAsset, int
 
 TConstArrayView<FVector3f> UClothAssetBuilder::GetSimPositions(const UChaosClothAsset& ClothAsset, int32 LodIndex)
 {
-	const FChaosClothSimulationModel* const ClothSimulationModel = ClothAsset.GetClothSimulationModel();
+	const TSharedPtr<const FChaosClothSimulationModel> ClothSimulationModel = ClothAsset.GetClothSimulationModel();
 	return ClothSimulationModel && ClothSimulationModel->IsValidLodIndex(LodIndex) ?
 		ClothSimulationModel->GetPositions(LodIndex) :
 		TConstArrayView<FVector3f>();
@@ -22,7 +22,7 @@ TConstArrayView<FVector3f> UClothAssetBuilder::GetSimPositions(const UChaosCloth
 
 TConstArrayView<uint32> UClothAssetBuilder::GetSimIndices(const UChaosClothAsset& ClothAsset, int32 LodIndex)
 {
-	const FChaosClothSimulationModel* const ClothSimulationModel = ClothAsset.GetClothSimulationModel();
+	const TSharedPtr<const FChaosClothSimulationModel> ClothSimulationModel = ClothAsset.GetClothSimulationModel();
 	return ClothSimulationModel && ClothSimulationModel->IsValidLodIndex(LodIndex) ?
 		ClothSimulationModel->GetIndices(LodIndex) :
 		TConstArrayView<uint32>();
diff --git a/Engine/Plugins/Experimental/ChaosClothAsset/Source/ChaosClothAssetEngine/Private/ChaosClothAsset/ClothComponent.cpp b/Engine/Plugins/Experimental/ChaosClothAsset/Source/ChaosClothAssetEngine/Private/ChaosClothAsset/ClothComponent.cpp
--- a/Engine/Plugins/Experimental/ChaosClothAsset/Source/ChaosClothAssetEngine/Private/ChaosClothAsset/ClothComponent.cpp
+++ b/Engine/Plugins/Experimental/ChaosClothAsset/Source/ChaosClothAssetEngine/Private/ChaosClothAsset/ClothComponent.cpp
@@ -92,7 +92,7 @@ void UChaosClothComponent::OnRegister()
 
 	if (GetClothAsset())
 	{
-		const FChaosClothSimulationModel* const ClothSimulationModel = GetClothAsset()->GetClothSimulationModel();
+		const TSharedPtr<const FChaosClothSimulationModel> ClothSimulationModel = GetClothAsset()->GetClothSimulationModel();
 		if (ensure(ClothSimulationModel) && ClothSimulationModel->GetNumLods())
 		{
 			FSkeletalMeshLODRenderData& LODData = GetClothAsset()->GetResourceForRendering()->LODRenderData[GetPredictedLODLevel()];
diff --git a/Engine/Plugins/Experimental/ChaosClothAsset/Source/ChaosClothAssetEngine/Private/ChaosClothAsset/ClothSimulationProxy.cpp b/Engine/Plugins/Experimental/ChaosClothAsset/Source/ChaosClothAssetEngine/Private/ChaosClothAsset/ClothSimulationProxy.cpp
--- a/Engine/Plugins/Experimental/ChaosClothAsset/Source/ChaosClothAssetEngine/Private/ChaosClothAsset/ClothSimulationProxy.cpp
+++ b/Engine/Plugins/Experimental/ChaosClothAsset/Source/ChaosClothAssetEngine/Private/ChaosClothAsset/ClothSimulationProxy.cpp
@@ -102,8 +102,8 @@ namespace UE::Chaos::ClothAsset
 		{
 			return;
 		}
-		const FChaosClothSimulationModel* const ClothSimulationModel = ClothAsset->GetClothSimulationModel();
-		check(ClothSimulationModel);
+		ClothSimulationModel = ClothAsset->GetClothSimulationModel();
+		check(ClothSimulationModel.IsValid());
 
 		FString DebugName;
 #if !UE_BUILD_SHIPPING
@@ -255,6 +255,7 @@ PRAGMA_ENABLE_DEPRECATION_WARNINGS
 		// Visualization
 		Visualization->DrawPhysMeshWired();
 		Visualization->DrawCollision();
+		Visualization->DrawOpenEdges();
 	}
 
 	void FClothSimulationProxy::CompleteParallelSimulation_GameThread()
diff --git a/Engine/Plugins/Experimental/ChaosClothAsset/Source/ChaosClothAssetEngine/Private/ChaosClothAsset/ClothSimulationProxy.h b/Engine/Plugins/Experimental/ChaosClothAsset/Source/ChaosClothAssetEngine/Private/ChaosClothAsset/ClothSimulationProxy.h
--- a/Engine/Plugins/Experimental/ChaosClothAsset/Source/ChaosClothAssetEngine/Private/ChaosClothAsset/ClothSimulationProxy.h
+++ b/Engine/Plugins/Experimental/ChaosClothAsset/Source/ChaosClothAssetEngine/Private/ChaosClothAsset/ClothSimulationProxy.h
@@ -20,6 +20,7 @@ namespace Chaos
 
 class UChaosClothComponent;
 class UChaosClothAsset;
+struct FChaosClothSimulationModel;
 struct FReferenceSkeleton;
 
 namespace UE::Chaos::ClothAsset
@@ -73,6 +74,9 @@ namespace UE::Chaos::ClothAsset
 		// Simulation context used to store the required component data for the duration of the simulation
 		FClothSimulationContext ClothSimulationContext;
 
+		// The cloth simulation model used to create this simulation, ownership might get transferred to this proxy if it changes during the simulation
+		TSharedPtr<const FChaosClothSimulationModel> ClothSimulationModel;
+
 		// Simulation objects
 		TUniquePtr<::Chaos::FClothingSimulationSolver> Solver;
 PRAGMA_DISABLE_DEPRECATION_WARNINGS  // TODO: CHAOS_IS_CLOTHINGSIMULATIONMESH_ABSTRACT
diff --git a/Engine/Plugins/Experimental/ChaosClothAsset/Source/ChaosClothAssetEngine/Public/ChaosClothAsset/ClothAsset.h b/Engine/Plugins/Experimental/ChaosClothAsset/Source/ChaosClothAssetEngine/Public/ChaosClothAsset/ClothAsset.h
--- a/Engine/Plugins/Experimental/ChaosClothAsset/Source/ChaosClothAssetEngine/Public/ChaosClothAsset/ClothAsset.h
+++ b/Engine/Plugins/Experimental/ChaosClothAsset/Source/ChaosClothAssetEngine/Public/ChaosClothAsset/ClothAsset.h
@@ -4,7 +4,6 @@
 
 #include "ChaosClothAsset/ClothPreset.h"
 #include "Engine/SkinnedAsset.h"
-#include "Engine/SkinnedAssetCommon.h"
 #include "ReferenceSkeleton.h"
 #include "RenderCommandFence.h"
 #include "ClothAsset.generated.h"
@@ -49,8 +48,8 @@ public:
 	//~ Begin USkinnedAsset interface
 	virtual FReferenceSkeleton& GetRefSkeleton()								{ return RefSkeleton; }
 	virtual const FReferenceSkeleton& GetRefSkeleton() const					{ return RefSkeleton; }
-	virtual FSkeletalMeshLODInfo* GetLODInfo(int32 Index)						{ return LODInfo.IsValidIndex(Index) ? &LODInfo[Index] : nullptr; }
-	virtual const FSkeletalMeshLODInfo* GetLODInfo(int32 Index) const			{ return LODInfo.IsValidIndex(Index) ? &LODInfo[Index] : nullptr; }
+	virtual FSkeletalMeshLODInfo* GetLODInfo(int32 Index);
+	virtual const FSkeletalMeshLODInfo* GetLODInfo(int32 Index) const;
 	UFUNCTION(BlueprintGetter)
 	virtual class UPhysicsAsset* GetShadowPhysicsAsset() const override			{ return ShadowPhysicsAsset; }
 	virtual FMatrix GetComposedRefPoseMatrix(FName InBoneName) const override;
@@ -103,7 +102,7 @@ public:
 	const TSharedPtr<UE::Chaos::ClothAsset::FClothCollection> GetClothCollection() const { return ClothCollection; }
 
 	/** Return the cloth simulation ready LOD model data. */
-	const FChaosClothSimulationModel* GetClothSimulationModel() const { return ClothSimulationModel.Get(); }
+	TSharedPtr<const FChaosClothSimulationModel> GetClothSimulationModel() const { return ClothSimulationModel; }
 
 	/** Build this asset static render and simulation data. This needs to be done every time the asset has changed. */
 	void Build();
@@ -115,6 +114,9 @@ public:
 	 */
 	void CopySimMeshToRenderMesh(int32 MaterialIndex);
 
+	/** Set the physics asset for this cloth. */
+	void SetPhysicsAsset(UPhysicsAsset* InPhysicsAsset) { PhysicsAsset = InPhysicsAsset; }
+
 private:
 	//~ Begin USkinnedAsset interface
 	/** Initial step for the Post Load process - Can't be done in parallel. */
@@ -167,6 +169,9 @@ private:
 	void CacheDerivedData(FSkinnedAssetPostLoadContext* Context);
 #endif
 
+	/** Reregister all components using this asset to reset the simulation in case anything has changed. */
+	void ReregisterComponents();
+
 	/** List of cloth presets for this cloth asset. */
 	UPROPERTY(EditAnywhere, AssetRegistrySearchable, Category = ClothPresets)
 	TArray<TObjectPtr<UChaosClothPreset>> ClothPresets;
@@ -251,8 +256,8 @@ private:
 	TSharedPtr<FSkeletalMeshModel> MeshModel;
 #endif
 
-	/** Simulation mesh Lods as fed to the solver for constraints creation. */
-	TUniquePtr<FChaosClothSimulationModel> ClothSimulationModel;
+	/** Simulation mesh Lods as fed to the solver for constraints creation. Ownership gets transferred to the proxy when it is changed during a simulation. */
+	TSharedPtr<FChaosClothSimulationModel> ClothSimulationModel;
 
 	friend class UClothAssetBuilderEditor;
 };
diff --git a/Engine/Plugins/Experimental/ChaosClothAssetEditor/Source/ChaosClothAssetTools/Private/ChaosClothAsset/ClothingAssetToClothAssetExporter.cpp b/Engine/Plugins/Experimental/ChaosClothAssetEditor/Source/ChaosClothAssetTools/Private/ChaosClothAsset/ClothingAssetToClothAssetExporter.cpp
--- a/Engine/Plugins/Experimental/ChaosClothAssetEditor/Source/ChaosClothAssetTools/Private/ChaosClothAsset/ClothingAssetToClothAssetExporter.cpp
+++ b/Engine/Plugins/Experimental/ChaosClothAssetEditor/Source/ChaosClothAssetTools/Private/ChaosClothAsset/ClothingAssetToClothAssetExporter.cpp
@@ -4,6 +4,8 @@
 #include "ChaosClothAsset/ClothAsset.h"
 #include "ChaosClothAsset/ClothAdapter.h"
 #include "Animation/Skeleton.h"
+#include "Engine/SkinnedAssetCommon.h"
+#include "Engine/SkeletalMesh.h"
 #include "Materials/Material.h"
 #include "Misc/MessageDialog.h"
 #include "ClothingAsset.h"
@@ -48,9 +50,17 @@ void UClothingAssetToChaosClothAssetExporter::Export(const UClothingAssetBase* C
 	UMaterialInterface* const DefaultMaterial = LoadObject<UMaterialInterface>(nullptr, TEXT("/Engine/EditorMaterials/Cloth/CameraLitDoubleSided.CameraLitDoubleSided"), nullptr, LOAD_None, nullptr);
 	const int32 MaterialId = ClothAsset->GetMaterials().Emplace(DefaultMaterial, true, false, DefaultMaterial->GetFName());
 
-	// Set a default skeleton
-	USkeleton* const DefaultSkeleton = LoadObject<USkeleton>(nullptr, TEXT("/Engine/EditorMeshes/SkeletalMesh/DefaultSkeletalMesh_Skeleton.DefaultSkeletalMesh_Skeleton"), nullptr, LOAD_None, nullptr);
-	ClothAsset->SetSkeleton(DefaultSkeleton);
+	// Assign the physics asset if any
+	ClothAsset->SetPhysicsAsset(ClothingAssetCommon->PhysicsAsset);
+
+	// Set the skeleton from the skeletal mesh, or create a default skeleton if it can't find one
+	USkeletalMesh* const SkeletalMesh = Cast<USkeletalMesh>(ClothingAssetCommon->GetOuter());
+
+	USkeleton* const Skeleton = SkeletalMesh ?
+		SkeletalMesh->GetSkeleton() :
+		LoadObject<USkeleton>(nullptr, TEXT("/Engine/EditorMeshes/SkeletalMesh/DefaultSkeletalMesh_Skeleton.DefaultSkeletalMesh_Skeleton"), nullptr, LOAD_None, nullptr);
+
+	ClothAsset->SetSkeleton(Skeleton);
 
 	// Set the render mesh to duplicate the sim mesh
 	ClothAsset->CopySimMeshToRenderMesh(MaterialId);
