Chaos Cloth:
- Added new simulation control.
- Tidied up some variable naming.

#rnx

diff --git a/Engine/Plugins/Experimental/ChaosClothAsset/Source/ChaosClothAssetEngine/Private/ChaosClothAsset/ClothComponent.cpp b/Engine/Plugins/Experimental/ChaosClothAsset/Source/ChaosClothAssetEngine/Private/ChaosClothAsset/ClothComponent.cpp
--- a/Engine/Plugins/Experimental/ChaosClothAsset/Source/ChaosClothAssetEngine/Private/ChaosClothAsset/ClothComponent.cpp
+++ b/Engine/Plugins/Experimental/ChaosClothAsset/Source/ChaosClothAssetEngine/Private/ChaosClothAsset/ClothComponent.cpp
@@ -14,10 +14,10 @@ CSV_DECLARE_CATEGORY_MODULE_EXTERN(ENGINE_API, Animation);
 
 UChaosClothComponent::UChaosClothComponent(const FObjectInitializer& ObjectInitializer)
 	: Super(ObjectInitializer)
-	, bDisableClothSimulation(0)
+	, bDisableSimulation(0)
 	, bSuspendSimulation(0)
-	, bWaitForParallelClothTask(0)
-	, bBindClothToLeaderComponent(0)
+	, bWaitForParallelTask(0)
+	, bBindToLeaderComponent(0)
 {
 	bAutoActivate = true;
 	PrimaryComponentTick.bCanEverTick = true;
@@ -108,6 +108,13 @@ void UChaosClothComponent::OnRegister()
 	}
 }
 
+bool UChaosClothComponent::IsSimulationSuspended() const
+{
+	static IConsoleVariable* const CVarClothPhysics = IConsoleManager::Get().FindConsoleVariable(TEXT("p.ClothPhysics"));
+
+	return bSuspendSimulation || !ClothSimulationProxy.IsValid() || (CVarClothPhysics && !CVarClothPhysics->GetBool());
+}
+
 void UChaosClothComponent::OnUnregister()
 {
 	Super::OnUnregister();
@@ -116,6 +123,11 @@ void UChaosClothComponent::OnUnregister()
 	ClothSimulationProxy.Reset(nullptr);
 }
 
+bool UChaosClothComponent::IsComponentTickEnabled() const
+{
+	return !bDisableSimulation && Super::IsComponentTickEnabled();
+}
+
 void UChaosClothComponent::TickComponent(float DeltaTime, enum ELevelTick TickType, FActorComponentTickFunction* ThisTickFunction)
 {
 	CSV_SCOPED_TIMING_STAT_EXCLUSIVE(Physics);
@@ -126,18 +138,6 @@ void UChaosClothComponent::TickComponent(float DeltaTime, enum ELevelTick TickTy
 	//	ClothingSimulation->UpdateWorldForces(this);
 	//}
 
-	// TODO: Suspended update
-	//// If we are suspended, we will not simulate clothing, but as clothing is simulated in local space
-	//// relative to a root bone we need to extract simulation positions as this bone could be animated.
-	//if ((!CVarEnableClothPhysics.GetValueOnGameThread() || bClothingSimulationSuspended) && ClothingSimulation)
-	//{
-	//	// First update the simulation context, since the simulation isn't ticking
-	//	// and it is still required to get the correct simulation data and bounds.
-	//	constexpr bool bIsInitialization = false;
-	//	ClothingSimulation->FillContext(this, DeltaTime, ClothingSimulationContext, bIsInitialization);
-	//	ClothingSimulation->GetSimulationData(CurrentSimulationData, this, Cast<USkeletalMeshComponent>(MasterPoseComponent.Get()));
-	//}
-
 	// Make sure that the previous frame simulation has completed
 	HandleExistingParallelSimulation();
 
@@ -168,19 +168,19 @@ void UChaosClothComponent::RefreshBoneTransforms(FActorComponentTickFunction* Ti
 	bHasValidBoneTransform = true;
 }
 
-void UChaosClothComponent::GetUpdateClothSimulationData_AnyThread(TMap<int32, FClothSimulData>& OutClothSimulData, FMatrix& OutLocalToWorld, float& OutClothBlendWeight)
+void UChaosClothComponent::GetUpdateClothSimulationData_AnyThread(TMap<int32, FClothSimulData>& OutClothSimulData, FMatrix& OutLocalToWorld, float& OutBlendWeight)
 {
 	OutLocalToWorld = GetComponentToWorld().ToMatrixWithScale();
 
 	const UChaosClothComponent* const LeaderPoseClothComponent = Cast<UChaosClothComponent>(LeaderPoseComponent.Get());
-	if (LeaderPoseClothComponent && LeaderPoseClothComponent->ClothSimulationProxy && bBindClothToLeaderComponent)
+	if (LeaderPoseClothComponent && LeaderPoseClothComponent->ClothSimulationProxy && bBindToLeaderComponent)
 	{
-		OutClothBlendWeight = ClothBlendWeight;
+		OutBlendWeight = BlendWeight;
 		OutClothSimulData = LeaderPoseClothComponent->ClothSimulationProxy->GetCurrentSimulationData_AnyThread();
 	}
-	else if (!bDisableClothSimulation && !bBindClothToLeaderComponent)
+	else if (!bDisableSimulation && !bBindToLeaderComponent)
 	{
-		OutClothBlendWeight = ClothBlendWeight;
+		OutBlendWeight = BlendWeight;
 		OutClothSimulData = ClothSimulationProxy->GetCurrentSimulationData_AnyThread();
 	}
 	else
@@ -191,7 +191,7 @@ void UChaosClothComponent::GetUpdateClothSimulationData_AnyThread(TMap<int32, FC
 	// Blend cloth out whenever the simulation data is invalid
 	if (!OutClothSimulData.Num())
 	{
-		OutClothBlendWeight = 0.0f;
+		OutBlendWeight = 0.0f;
 	}
 }
 
@@ -244,7 +244,7 @@ PRAGMA_ENABLE_DEPRECATION_WARNINGS
 
 bool UChaosClothComponent::RequiresPreEndOfFrameSync() const
 {
-	if (!IsSimulationSuspended() && !ShouldWaitForClothInTickFunction())
+	if (!IsSimulationSuspended() && !ShouldWaitForParallelSimulationInTickFunction())
 	{
 		// By default we await the cloth task in TickComponent, but...
 		// If we have cloth and have no game-thread dependencies on the cloth output, 
@@ -298,7 +298,7 @@ void UChaosClothComponent::StartNewParallelSimulation(float DeltaTime)
 
 void UChaosClothComponent::HandleExistingParallelSimulation()
 {
-	if (bBindClothToLeaderComponent)
+	if (bBindToLeaderComponent)
 	{
 		if (UChaosClothComponent* const LeaderComponent = Cast<UChaosClothComponent>(LeaderPoseComponent.Get()))
 		{
@@ -312,16 +312,9 @@ void UChaosClothComponent::HandleExistingParallelSimulation()
 	}
 }
 
-bool UChaosClothComponent::ShouldWaitForClothInTickFunction() const
+bool UChaosClothComponent::ShouldWaitForParallelSimulationInTickFunction() const
 {
 	static IConsoleVariable* const CVarClothPhysicsWaitForParallelClothTask = IConsoleManager::Get().FindConsoleVariable(TEXT("p.ClothPhysics.WaitForParallelClothTask"));
 
-	return bWaitForParallelClothTask || (CVarClothPhysicsWaitForParallelClothTask && CVarClothPhysicsWaitForParallelClothTask->GetBool());
-}
-
-bool UChaosClothComponent::IsSimulationSuspended() const
-{
-	static IConsoleVariable* const CVarClothPhysics = IConsoleManager::Get().FindConsoleVariable(TEXT("p.ClothPhysics"));
-
-	return bSuspendSimulation || !ClothSimulationProxy.IsValid() || (CVarClothPhysics && !CVarClothPhysics->GetBool());
+	return bWaitForParallelTask || (CVarClothPhysicsWaitForParallelClothTask && CVarClothPhysicsWaitForParallelClothTask->GetBool());
 }
diff --git a/Engine/Plugins/Experimental/ChaosClothAsset/Source/ChaosClothAssetEngine/Private/ChaosClothAsset/ClothSimulationContext.cpp b/Engine/Plugins/Experimental/ChaosClothAsset/Source/ChaosClothAssetEngine/Private/ChaosClothAsset/ClothSimulationContext.cpp
--- a/Engine/Plugins/Experimental/ChaosClothAsset/Source/ChaosClothAssetEngine/Private/ChaosClothAsset/ClothSimulationContext.cpp
+++ b/Engine/Plugins/Experimental/ChaosClothAsset/Source/ChaosClothAssetEngine/Private/ChaosClothAsset/ClothSimulationContext.cpp
@@ -22,12 +22,12 @@ namespace UE::Chaos::ClothAsset
 		constexpr float MaxDeltaTimeTeleportMultiplierDefault = 1.5f;
 		const float MaxDeltaTimeTeleportMultiplier = CVarMaxDeltaTimeTeleportMultiplier ? CVarMaxDeltaTimeTeleportMultiplier->GetFloat() : MaxDeltaTimeTeleportMultiplierDefault;
 
-		TeleportMode = (DeltaTime > MaxDeltaTime * MaxDeltaTimeTeleportMultiplier) ?
-			EClothingTeleportMode::Teleport :
-			ClothComponent.GetClothTeleportMode();
+		bTeleport = (DeltaTime > MaxDeltaTime * MaxDeltaTimeTeleportMultiplier) ? true : ClothComponent.NeedsTeleport();
+		bReset = ClothComponent.NeedsReset();
 
-		VelocityScale = (TeleportMode == EClothingTeleportMode::None) ?
-			FMath::Min(InDeltaTime, MaxDeltaTime) / InDeltaTime : 1.f;
+		VelocityScale = (!bTeleport && !bReset && InDeltaTime > 0.f) ?
+			FMath::Min(InDeltaTime, MaxDeltaTime) / InDeltaTime :
+			bReset ? 1.f : 0.f;  // Set to 0 when teleporting and 1 when resetting to match the internal solver's behavior
 
 		// Copy component transform
 		ComponentTransform = ClothComponent.GetComponentTransform();
diff --git a/Engine/Plugins/Experimental/ChaosClothAsset/Source/ChaosClothAssetEngine/Private/ChaosClothAsset/ClothSimulationContext.h b/Engine/Plugins/Experimental/ChaosClothAsset/Source/ChaosClothAssetEngine/Private/ChaosClothAsset/ClothSimulationContext.h
--- a/Engine/Plugins/Experimental/ChaosClothAsset/Source/ChaosClothAssetEngine/Private/ChaosClothAsset/ClothSimulationContext.h
+++ b/Engine/Plugins/Experimental/ChaosClothAsset/Source/ChaosClothAssetEngine/Private/ChaosClothAsset/ClothSimulationContext.h
@@ -4,7 +4,6 @@
 
 #include "Containers/Array.h"
 #include "Math/Transform.h"
-#include "ClothingSystemRuntimeTypes.h"
 
 class UChaosClothComponent;
 
@@ -33,8 +32,11 @@ namespace UE::Chaos::ClothAsset
 		/** Bone matrices used for skinning. */
 		TArray<FMatrix44f> RefToLocalMatrices;
 
-		/** Teleport mode. */
-		EClothingTeleportMode TeleportMode = EClothingTeleportMode::None;
+		/** Teleport. */
+		bool bTeleport = false;
+
+		/** Reset. */
+		bool bReset = false;
 
 		/** Velocity scale to compensate for time differences when the MaxPhysicsDeltaTime kicks in. */
 		float VelocityScale;
diff --git a/Engine/Plugins/Experimental/ChaosClothAsset/Source/ChaosClothAssetEngine/Private/ChaosClothAsset/ClothSimulationProxy.cpp b/Engine/Plugins/Experimental/ChaosClothAsset/Source/ChaosClothAssetEngine/Private/ChaosClothAsset/ClothSimulationProxy.cpp
--- a/Engine/Plugins/Experimental/ChaosClothAsset/Source/ChaosClothAssetEngine/Private/ChaosClothAsset/ClothSimulationProxy.cpp
+++ b/Engine/Plugins/Experimental/ChaosClothAsset/Source/ChaosClothAssetEngine/Private/ChaosClothAsset/ClothSimulationProxy.cpp
@@ -188,14 +188,19 @@ PRAGMA_ENABLE_DEPRECATION_WARNINGS
 	{
 		SCOPE_CYCLE_COUNTER(STAT_ClothSimulationProxy_TickGame);
 
-		if (DeltaTime > 0.f)
-		{
-			// Fill a new context
-			ClothSimulationContext.Fill(ClothComponent, DeltaTime, MaxDeltaTime);
+		// Fill a new context, note the context is also needed when the simulation is suspended
+		ClothSimulationContext.Fill(ClothComponent, DeltaTime, MaxDeltaTime);
 
+		if (DeltaTime > 0.f && !ClothComponent.IsSimulationSuspended())
+		{
 			// Start the the cloth simulation thread
 			ParallelTask = TGraphTask<FClothSimulationProxyParallelTask>::CreateTask(nullptr, ENamedThreads::GameThread).ConstructAndDispatchWhenReady(*this);
 		}
+		else
+		{
+			// It still needs to write back to the GT cache as the context has changed
+			WriteSimulationData_GameThread();
+		}
 	}
 
 	void FClothSimulationProxy::Tick_PhysicsThread()
@@ -212,8 +217,8 @@ PRAGMA_ENABLE_DEPRECATION_WARNINGS
 		const double StartTime = FPlatformTime::Seconds();
 		const float PrevSimulationTime = SimulationTime;  // Copy the atomic to prevent a re-read
 
-		const bool bNeedsReset = (ClothSimulationContext.TeleportMode == EClothingTeleportMode::TeleportAndReset || PrevSimulationTime == 0.f);
-		const bool bNeedsTeleport = (ClothSimulationContext.TeleportMode > EClothingTeleportMode::None);
+		const bool bNeedsReset = (ClothSimulationContext.bReset || PrevSimulationTime == 0.f);  // Reset on the first frame too since the simulation is created in bind pose, and not in start pose
+		const bool bNeedsTeleport = ClothSimulationContext.bTeleport;
 		bIsTeleported = bNeedsTeleport;
 
 		// Update Solver animatable parameters
diff --git a/Engine/Plugins/Experimental/ChaosClothAsset/Source/ChaosClothAssetEngine/Private/ChaosClothAsset/ClothSimulationProxy.h b/Engine/Plugins/Experimental/ChaosClothAsset/Source/ChaosClothAssetEngine/Private/ChaosClothAsset/ClothSimulationProxy.h
--- a/Engine/Plugins/Experimental/ChaosClothAsset/Source/ChaosClothAssetEngine/Private/ChaosClothAsset/ClothSimulationProxy.h
+++ b/Engine/Plugins/Experimental/ChaosClothAsset/Source/ChaosClothAssetEngine/Private/ChaosClothAsset/ClothSimulationProxy.h
@@ -43,6 +43,7 @@ namespace UE::Chaos::ClothAsset
 
 		void Tick_GameThread(float DeltaTime);
 
+		/** Wait for the parallel task to complete if one was running, and update the simulation data. */
 		void CompleteParallelSimulation_GameThread();
 
 		/**
diff --git a/Engine/Plugins/Experimental/ChaosClothAsset/Source/ChaosClothAssetEngine/Public/ChaosClothAsset/ClothComponent.h b/Engine/Plugins/Experimental/ChaosClothAsset/Source/ChaosClothAssetEngine/Public/ChaosClothAsset/ClothComponent.h
--- a/Engine/Plugins/Experimental/ChaosClothAsset/Source/ChaosClothAssetEngine/Public/ChaosClothAsset/ClothComponent.h
+++ b/Engine/Plugins/Experimental/ChaosClothAsset/Source/ChaosClothAssetEngine/Public/ChaosClothAsset/ClothComponent.h
@@ -3,7 +3,6 @@
 #pragma once
 
 #include "Components/SkinnedMeshComponent.h"
-#include "ClothingSystemRuntimeTypes.h"
 
 #include "ClothComponent.generated.h"
 
@@ -27,33 +26,42 @@ public:
 	~UChaosClothComponent();
 
 	/** Set the cloth asset used by this component. */
-	UFUNCTION(BlueprintCallable, Category = "Components|ClothAsset")
+	UFUNCTION(BlueprintCallable, Category = "Components|ClothAsset", Meta = (Keywords = "Chaos Cloth Asset"))
 	void SetClothAsset(UChaosClothAsset* InClothAsset);
 
 	/** Get the cloth asset used by this component. */
-	UFUNCTION(BlueprintPure, Category = "Components|ClothAsset")
+	UFUNCTION(BlueprintPure, Category = "Components|ClothAsset", Meta = (Keywords = "Chaos Cloth Asset"))
 	UChaosClothAsset* GetClothAsset() const;
 
 	/** Reset the teleport mode. */
-	UFUNCTION(BlueprintCallable, Category = "Components|ClothAsset|Teleport")
-	void ResetClothTeleportMode() { ClothTeleportMode = EClothingTeleportMode::None; }
-
-	/**
-	 * Used to indicate we should force 'teleport' during the next call to UpdateClothState,
-	 * This will transform positions and velocities and thus keep the simulation state, just translate it to a new pose.
-	 */
-	UFUNCTION(BlueprintCallable, Category = "Components|ClothAsset|Teleport")
-	void ForceClothNextUpdateTeleport() { ClothTeleportMode = EClothingTeleportMode::Teleport; }
-
-	/**
-	 * Used to indicate we should force 'teleport and reset' during the next call to UpdateClothState.
-	 * This can be used to reset it from a bad state or by a teleport where the old state is not important anymore.
-	 */
-	UFUNCTION(BlueprintCallable, Category = "Components|ClothAsset|Teleport")
-	void ForceClothNextUpdateTeleportAndReset() { ClothTeleportMode = EClothingTeleportMode::TeleportAndReset; }
-
-	/** Return the current teleport mode currently requested if any. */
-	EClothingTeleportMode GetClothTeleportMode() const { return ClothTeleportMode; }
+	UFUNCTION(BlueprintCallable, Category = "Components|Teleport", Meta = (Keywords = "Chaos Cloth Teleport"))
+	void ResetTeleportMode() { bTeleport = bReset = false; }
+
+	/** Teleport the cloth particles to the new reference bone location keeping pose and velocities prior to advancing the simulation. */
+	UFUNCTION(BlueprintCallable, Category = "Components||Teleport", Meta = (Keywords = "Chaos Cloth Teleport"))
+	void ForceNextUpdateTeleport() { bTeleport = true; bReset = false; }
+
+	/** Teleport the cloth particles to the new reference bone location while reseting the pose and velocities prior to advancing the simulation. */
+	UFUNCTION(BlueprintCallable, Category = "Components|Simulation", Meta = (Keywords = "Chaos Cloth Teleport Reset"))
+	void ForceNextUpdateTeleportAndReset() { bTeleport = bReset = true; }
+
+	/** Return whether teleport is currently requested. */
+	bool NeedsTeleport() const { return bTeleport; }
+
+	/** Return whether reseting the pose is currently requested. */
+	bool NeedsReset() const { return bReset; }
+
+	/** Stop the simulation, and keep the cloth in its last pose. */
+	UFUNCTION(BlueprintCallable, Category = "Components|Simulation", Meta = (UnsafeDuringActorConstruction, Keywords = "Chaos Cloth Simulation Suspend"))
+	void SuspendSimulation() { bSuspendSimulation = true; }
+
+	/** Resume a previously suspended simulation. */
+	UFUNCTION(BlueprintCallable, Category = "Components|Simulation", Meta = (UnsafeDuringActorConstruction, Keywords = "Chaos Cloth Simulation Resume"))
+	void ResumeSimulation() { bSuspendSimulation = false; }
+
+	/** Return whether or not the simulation is currently suspended. */
+	UFUNCTION(BlueprintCallable, Category = "Components|Simulation", Meta = (Keywords = "Chaos Cloth Simulation Suspend"))
+	bool IsSimulationSuspended() const;
 
 protected:
 	//~ Begin UObject Interface
@@ -66,6 +74,7 @@ protected:
 	//~ Begin UActorComponent Interface
 	virtual void OnRegister() override;
 	virtual void OnUnregister() override;
+	virtual bool IsComponentTickEnabled() const override;
 	virtual void TickComponent(float DeltaTime, enum ELevelTick TickType, FActorComponentTickFunction* ThisTickFunction) override;
 	virtual bool RequiresPreEndOfFrameSync() const override;
 	virtual void OnPreEndOfFrameSync() override;
@@ -77,14 +86,13 @@ protected:
 
 	//~ Begin USkinnedMeshComponent Interface
 	virtual void RefreshBoneTransforms(FActorComponentTickFunction* TickFunction = nullptr) override;
-	virtual void GetUpdateClothSimulationData_AnyThread(TMap<int32, FClothSimulData>& OutClothSimulData, FMatrix& OutLocalToWorld, float& OutClothBlendWeight) override;
+	virtual void GetUpdateClothSimulationData_AnyThread(TMap<int32, FClothSimulData>& OutClothSimulData, FMatrix& OutLocalToWorld, float& OutBlendWeight) override;
 	//~ End USkinnedMeshComponent Interface
 
 private:
 	void StartNewParallelSimulation(float DeltaTime);
 	void HandleExistingParallelSimulation();
-	bool ShouldWaitForClothInTickFunction() const;
-	bool IsSimulationSuspended() const;
+	bool ShouldWaitForParallelSimulationInTickFunction() const;
 
 #if WITH_EDITORONLY_DATA
 	/** Cloth asset used by this component. */
@@ -93,22 +101,33 @@ private:
 	TObjectPtr<UChaosClothAsset> ClothAsset;
 #endif
 
+	/** Whether to disable the simulation and use the skinned pose instead. */
 	UPROPERTY()
-	uint8 bDisableClothSimulation : 1;
+	uint8 bDisableSimulation : 1;
 
+	/** Whether to suspend the simulation and use the last simulated pose. */
 	UPROPERTY()
 	uint8 bSuspendSimulation : 1;
 
+	/** Whether to wait for the cloth simulation to end in the TickComponent instead of in the EndOfFrameUpdates. */
 	UPROPERTY()
-	uint8 bWaitForParallelClothTask : 1;
+	uint8 bWaitForParallelTask : 1;
 
+	/** Whether to use the leader component pose. */
 	UPROPERTY()
-	uint8 bBindClothToLeaderComponent : 1;
+	uint8 bBindToLeaderComponent : 1;
 
+	/** Whether to teleport the cloth prior to advancing the simulation. */
 	UPROPERTY()
-	float ClothBlendWeight = 1.f;
+	uint8 bTeleport : 1;
 
-	EClothingTeleportMode ClothTeleportMode = EClothingTeleportMode::None;
+	/** Whether to reset the pose, bTeleport must be true. */
+	UPROPERTY()
+	uint8 bReset : 1;
+
+	/** Blend amount between the skinned (=0) and the simulated pose (=1). */
+	UPROPERTY()
+	float BlendWeight = 1.f;
 
 	TUniquePtr<UE::Chaos::ClothAsset::FClothSimulationProxy> ClothSimulationProxy;
 };
diff --git a/Engine/Plugins/Experimental/ChaosClothAssetEditor/Source/ChaosClothAssetEditor/Private/ChaosClothAsset/ClothEditorMode.cpp b/Engine/Plugins/Experimental/ChaosClothAssetEditor/Source/ChaosClothAssetEditor/Private/ChaosClothAsset/ClothEditorMode.cpp
--- a/Engine/Plugins/Experimental/ChaosClothAssetEditor/Source/ChaosClothAssetEditor/Private/ChaosClothAsset/ClothEditorMode.cpp
+++ b/Engine/Plugins/Experimental/ChaosClothAssetEditor/Source/ChaosClothAssetEditor/Private/ChaosClothAsset/ClothEditorMode.cpp
@@ -710,7 +710,7 @@ void UChaosClothAssetEditorMode::ModeTick(float DeltaTime)
 
 	if (bShouldClearTeleportFlag)
 	{
-		ClothComponent->ResetClothTeleportMode();
+		ClothComponent->ResetTeleportMode();
 		bShouldClearTeleportFlag = false;
 	}
 
@@ -722,7 +722,7 @@ void UChaosClothAssetEditorMode::ModeTick(float DeltaTime)
 		}
 		else
 		{
-			ClothComponent->ForceClothNextUpdateTeleportAndReset();
+			ClothComponent->ForceNextUpdateTeleportAndReset();
 		}
 
 		bShouldResetSimulation = false;
