ChaosCloth: if the cloth asset has a scalar weight map attribute named MaxDistance, pass it into the cloth simulation as the MaxDistance constraint. (This is simply hard coded for now, until we have the ability to specify which weight maps are to be used for which simulation parameters)

diff --git a/Engine/Plugins/Experimental/ChaosClothAsset/Source/ChaosClothAssetEngine/Private/ChaosClothAsset/ClothSimulationMesh.cpp b/Engine/Plugins/Experimental/ChaosClothAsset/Source/ChaosClothAssetEngine/Private/ChaosClothAsset/ClothSimulationMesh.cpp
--- a/Engine/Plugins/Experimental/ChaosClothAsset/Source/ChaosClothAssetEngine/Private/ChaosClothAsset/ClothSimulationMesh.cpp
+++ b/Engine/Plugins/Experimental/ChaosClothAsset/Source/ChaosClothAssetEngine/Private/ChaosClothAsset/ClothSimulationMesh.cpp
@@ -60,8 +60,15 @@ PRAGMA_ENABLE_DEPRECATION_WARNINGS
 	TArray<TConstArrayView<::Chaos::FRealSingle>> FClothSimulationMesh::GetWeightMaps(int32 LODIndex) const
 	{
 		constexpr int32 MaxNumWeightMaps = 15; // TODO: Refactor how weight maps are used in base simulation classes
-		static const TConstArrayView<::Chaos::FRealSingle> WeightMaps[MaxNumWeightMaps];
-		return TArray<TConstArrayView<::Chaos::FRealSingle>>(&WeightMaps[0], MaxNumWeightMaps);
+
+		TArray<TConstArrayView<::Chaos::FRealSingle>> WeightMaps;
+		WeightMaps.SetNum(MaxNumWeightMaps);
+
+		// Set max distance weight map
+		constexpr int32 MaxDistanceWeightMapTarget = 1;		// EWeightMapTargetCommon::MaxDistance
+		WeightMaps[MaxDistanceWeightMapTarget] = TConstArrayView<::Chaos::FRealSingle>(ClothSimulationModel.ClothSimulationLodModels[LODIndex].MaxDistance);
+
+		return WeightMaps;
 	}
 
 	TArray<TConstArrayView<TTuple<int32, int32, float>>> FClothSimulationMesh::GetTethers(int32 LODIndex, bool bUseGeodesicTethers) const
diff --git a/Engine/Plugins/Experimental/ChaosClothAsset/Source/ChaosClothAssetEngine/Private/ChaosClothAsset/ClothSimulationModel.cpp b/Engine/Plugins/Experimental/ChaosClothAsset/Source/ChaosClothAssetEngine/Private/ChaosClothAsset/ClothSimulationModel.cpp
--- a/Engine/Plugins/Experimental/ChaosClothAsset/Source/ChaosClothAssetEngine/Private/ChaosClothAsset/ClothSimulationModel.cpp
+++ b/Engine/Plugins/Experimental/ChaosClothAsset/Source/ChaosClothAssetEngine/Private/ChaosClothAsset/ClothSimulationModel.cpp
@@ -122,6 +122,16 @@ FChaosClothSimulationModel::FChaosClothSimulationModel(const TSharedPtr<const UE
 		const FClothLodConstAdapter ClothLod = Cloth.GetLod(LodIndex);
 		ClothLod.BuildSimulationMesh(LodModel.Positions, LodModel.Normals, LodModel.Indices);
 
+		const TManagedArray<float>* const MaxDistanceValues = ClothCollection->FindAttributeTyped<float>("MaxDistance", FClothCollection::SimVerticesGroup);
+		if (MaxDistanceValues)
+		{
+			LodModel.MaxDistance = MaxDistanceValues->GetConstArray();
+		}
+		else
+		{
+			LodModel.MaxDistance.Init(1.0, LodModel.Positions.Num());
+		}
+
 		const int32 NumSimVertices = LodModel.Positions.Num();
 		LodModel.BoneData.SetNum(NumSimVertices);
 		for (int32 VertexIndex = 0; VertexIndex < NumSimVertices; ++VertexIndex)
diff --git a/Engine/Plugins/Experimental/ChaosClothAsset/Source/ChaosClothAssetEngine/Private/ChaosClothAsset/ClothSimulationModel.h b/Engine/Plugins/Experimental/ChaosClothAsset/Source/ChaosClothAssetEngine/Private/ChaosClothAsset/ClothSimulationModel.h
--- a/Engine/Plugins/Experimental/ChaosClothAsset/Source/ChaosClothAssetEngine/Private/ChaosClothAsset/ClothSimulationModel.h
+++ b/Engine/Plugins/Experimental/ChaosClothAsset/Source/ChaosClothAssetEngine/Private/ChaosClothAsset/ClothSimulationModel.h
@@ -39,6 +39,9 @@ struct FChaosClothSimulationLodModel
 	/** Skinning weights. */
 	UPROPERTY()
 	TArray<FClothVertBoneData> BoneData;
+
+	UPROPERTY()
+	TArray<float> MaxDistance;
 };
 
 /**
diff --git a/Engine/Plugins/Experimental/ChaosClothAsset/Source/ChaosClothAssetEngine/Private/ChaosClothAsset/ClothSimulationProxy.cpp b/Engine/Plugins/Experimental/ChaosClothAsset/Source/ChaosClothAssetEngine/Private/ChaosClothAsset/ClothSimulationProxy.cpp
--- a/Engine/Plugins/Experimental/ChaosClothAsset/Source/ChaosClothAssetEngine/Private/ChaosClothAsset/ClothSimulationProxy.cpp
+++ b/Engine/Plugins/Experimental/ChaosClothAsset/Source/ChaosClothAssetEngine/Private/ChaosClothAsset/ClothSimulationProxy.cpp
@@ -234,7 +234,7 @@ PRAGMA_ENABLE_DEPRECATION_WARNINGS
 		for (const TUniquePtr<FClothingSimulationCloth>& Cloth : Cloths)
 		{
 			// Update Cloth animatable parameters while in the cloth loop
-			//Cloth->SetMaxDistancesMultiplier(Context->MaxDistanceScale);
+			Cloth->SetMaxDistancesMultiplier(1000.0);
 
 			if (bNeedsReset)
 			{
diff --git a/Engine/Plugins/Experimental/ChaosClothAssetEditor/Source/ChaosClothAssetEditor/Private/ChaosClothAsset/ClothEditorMode.cpp b/Engine/Plugins/Experimental/ChaosClothAssetEditor/Source/ChaosClothAssetEditor/Private/ChaosClothAsset/ClothEditorMode.cpp
--- a/Engine/Plugins/Experimental/ChaosClothAssetEditor/Source/ChaosClothAssetEditor/Private/ChaosClothAsset/ClothEditorMode.cpp
+++ b/Engine/Plugins/Experimental/ChaosClothAssetEditor/Source/ChaosClothAssetEditor/Private/ChaosClothAsset/ClothEditorMode.cpp
@@ -389,6 +389,13 @@ void UChaosClothAssetEditorMode::UpdateSimulationMeshes()
 	const TSet<FName> WeightMapsToRemove = ClothAssetWeightMapNames.Difference(CommonDynamicMeshWeightMapNames);
 	ChaosClothAssetEditorModeHelpers::RemoveClothWeightMaps(ClothAdapter, WeightMapsToRemove.Array());
 
+	ChaosClothAsset->Build();
+
+	// Reset cloth component
+	{
+		const FComponentReregisterContext Context(ClothComponent);
+	}
+
 	GetToolManager()->EndUndoTransaction();
 }
 
