Chaos Cloth:
- Fix ref pose initialization on collision bones.
- Fix crash when re-exporting over an existing cloth asset.

#rnx

diff --git a/Engine/Plugins/Experimental/ChaosClothAsset/Source/ChaosClothAssetEngine/Private/ChaosClothAsset/ClothAsset.cpp b/Engine/Plugins/Experimental/ChaosClothAsset/Source/ChaosClothAssetEngine/Private/ChaosClothAsset/ClothAsset.cpp
--- a/Engine/Plugins/Experimental/ChaosClothAsset/Source/ChaosClothAssetEngine/Private/ChaosClothAsset/ClothAsset.cpp
+++ b/Engine/Plugins/Experimental/ChaosClothAsset/Source/ChaosClothAssetEngine/Private/ChaosClothAsset/ClothAsset.cpp
@@ -310,12 +310,6 @@ void UChaosClothAsset::CalculateBounds()
 
 void UChaosClothAsset::UpdateSkeleton(bool bRebuildClothSimulationModel)
 {
-	// Rebuild Skeleton
-	RefSkeleton = Skeleton->GetReferenceSkeleton();  // TODO: Does a copy work? Try to see what's exactly needed here...
-
-	constexpr bool bRebuildNameMap = false;
-	RefSkeleton.RebuildRefSkeleton(Skeleton, bRebuildNameMap);
-
 	CalculateInvRefMatrices();
 
 	if (bRebuildClothSimulationModel)
diff --git a/Engine/Plugins/Experimental/ChaosClothAsset/Source/ChaosClothAssetEngine/Public/ChaosClothAsset/ClothAsset.h b/Engine/Plugins/Experimental/ChaosClothAsset/Source/ChaosClothAssetEngine/Public/ChaosClothAsset/ClothAsset.h
--- a/Engine/Plugins/Experimental/ChaosClothAsset/Source/ChaosClothAssetEngine/Public/ChaosClothAsset/ClothAsset.h
+++ b/Engine/Plugins/Experimental/ChaosClothAsset/Source/ChaosClothAssetEngine/Public/ChaosClothAsset/ClothAsset.h
@@ -72,9 +72,9 @@ public:
 	virtual TArray<class USkeletalMeshSocket*> GetActiveSocketList() const override { static TArray<class USkeletalMeshSocket*> Dummy; return Dummy; }
 	virtual USkeletalMeshSocket* FindSocket(FName InSocketName) const override	{ return nullptr; }
 	virtual USkeletalMeshSocket* FindSocketInfo(FName InSocketName, FTransform& OutTransform, int32& OutBoneIndex, int32& OutIndex) const override { return nullptr; }
-	virtual USkeleton* GetSkeleton() override									{ return Skeleton; }
-	virtual const USkeleton* GetSkeleton() const override						{ return Skeleton; }
-	virtual void SetSkeleton(USkeleton* InSkeleton) override					{ Skeleton = InSkeleton; UpdateSkeleton(); }
+	virtual USkeleton* GetSkeleton() override									{ return nullptr; }
+	virtual const USkeleton* GetSkeleton() const override						{ return nullptr; }
+	virtual void SetSkeleton(USkeleton* InSkeleton) override					{}
 	virtual UMeshDeformer* GetDefaultMeshDeformer() const override				{ return nullptr; }
 	virtual bool IsValidLODIndex(int32 Index) const override					{ return LODInfo.IsValidIndex(Index); }
 	virtual int32 GetMinLodIdx(bool bForceLowestLODIdx = false) const override	{ return 0; }
@@ -117,6 +117,9 @@ public:
 	/** Set the physics asset for this cloth. */
 	void SetPhysicsAsset(UPhysicsAsset* InPhysicsAsset) { PhysicsAsset = InPhysicsAsset; }
 
+	/** Set the bone hierachy to use for this cloth. */
+	void SetReferenceSkeleton(const FReferenceSkeleton& InReferenceSkeleton, bool bRebuildClothSimulationModel = true) { RefSkeleton = InReferenceSkeleton; UpdateSkeleton(bRebuildClothSimulationModel); }
+
 private:
 	//~ Begin USkinnedAsset interface
 	/** Initial step for the Post Load process - Can't be done in parallel. */
@@ -185,10 +188,6 @@ private:
 	UPROPERTY(EditAnywhere, Category = Collision)
 	TObjectPtr<UPhysicsAsset> PhysicsAsset;
 
-	/** Skeleton. */
-	UPROPERTY(EditAnywhere, AssetRegistrySearchable, Category = Skeleton)
-	TObjectPtr<USkeleton> Skeleton;
-
 	/** Struct containing information for each LOD level, such as materials to use, and when use the LOD. */
 	UPROPERTY(EditAnywhere, EditFixedSize, Category = LevelOfDetails)
 	TArray<FSkeletalMeshLODInfo> LODInfo;
diff --git a/Engine/Plugins/Experimental/ChaosClothAssetEditor/Source/ChaosClothAssetTools/Private/ChaosClothAsset/ClothingAssetToClothAssetExporter.cpp b/Engine/Plugins/Experimental/ChaosClothAssetEditor/Source/ChaosClothAssetTools/Private/ChaosClothAsset/ClothingAssetToClothAssetExporter.cpp
--- a/Engine/Plugins/Experimental/ChaosClothAssetEditor/Source/ChaosClothAssetTools/Private/ChaosClothAsset/ClothingAssetToClothAssetExporter.cpp
+++ b/Engine/Plugins/Experimental/ChaosClothAssetEditor/Source/ChaosClothAssetTools/Private/ChaosClothAsset/ClothingAssetToClothAssetExporter.cpp
@@ -56,11 +56,9 @@ void UClothingAssetToChaosClothAssetExporter::Export(const UClothingAssetBase* C
 	// Set the skeleton from the skeletal mesh, or create a default skeleton if it can't find one
 	USkeletalMesh* const SkeletalMesh = Cast<USkeletalMesh>(ClothingAssetCommon->GetOuter());
 
-	USkeleton* const Skeleton = SkeletalMesh ?
-		SkeletalMesh->GetSkeleton() :
-		LoadObject<USkeleton>(nullptr, TEXT("/Engine/EditorMeshes/SkeletalMesh/DefaultSkeletalMesh_Skeleton.DefaultSkeletalMesh_Skeleton"), nullptr, LOAD_None, nullptr);
-
-	ClothAsset->SetSkeleton(Skeleton);
+	ClothAsset->SetReferenceSkeleton(SkeletalMesh ?
+		SkeletalMesh->GetRefSkeleton() :
+		LoadObject<USkeleton>(nullptr, TEXT("/Engine/EditorMeshes/SkeletalMesh/DefaultSkeletalMesh_Skeleton.DefaultSkeletalMesh_Skeleton"), nullptr, LOAD_None, nullptr)->GetReferenceSkeleton());
 
 	// Set the render mesh to duplicate the sim mesh
 	ClothAsset->CopySimMeshToRenderMesh(MaterialId);
diff --git a/Engine/Source/Editor/ClothPainter/Private/ClothingAssetExporter.cpp b/Engine/Source/Editor/ClothPainter/Private/ClothingAssetExporter.cpp
--- a/Engine/Source/Editor/ClothPainter/Private/ClothingAssetExporter.cpp
+++ b/Engine/Source/Editor/ClothPainter/Private/ClothingAssetExporter.cpp
@@ -7,6 +7,7 @@
 #include "Containers/Array.h"
 #include "Containers/UnrealString.h"
 #include "ContentBrowserModule.h"
+#include "Engine/SkinnedAsset.h"
 #include "Features/IModularFeatures.h"
 #include "FileHelpers.h"
 #include "HAL/PlatformMisc.h"
@@ -22,6 +23,7 @@
 #include "UObject/Class.h"
 #include "UObject/Package.h"
 #include "UObject/TopLevelAssetPath.h"
+#include "SkinnedAssetCompiler.h"
 
 #define LOCTEXT_NAMESPACE "ClothingAssetExporter"
 
@@ -81,7 +83,14 @@ void ExportClothingAsset(const UClothingAssetBase* ClothingAsset, UClass* Export
 	}
 
 	// Try loading the package, when overwriting an existing asset it needs to be loaded first
-	LoadPackage(nullptr, *NewPackageName, LOAD_Quiet | LOAD_EditorOnly);
+	if (UPackage* const Package = LoadPackage(nullptr, *NewPackageName, LOAD_Quiet | LOAD_EditorOnly))
+	{
+		if (USkinnedAsset* const SkinnedAsset = Cast<USkinnedAsset>(Package->FindAssetInPackage()))
+		{
+			// Finish compilation now otherwise the asset might be deleted and replaced before it ends
+			FSkinnedAssetCompilingManager::Get().FinishCompilation({ SkinnedAsset });
+		}
+	}
 
 	// Find a matching exporter
 	UClothingAssetExporter* ClothingAssetExporter = nullptr;
