Chaos Cloth:
- Add Leader Pose support for the new Cloth Asset.
- Add an optional wait for the simulation in the Cloth Component TickComponent function.

#rnx

diff --git a/Engine/Plugins/Experimental/ChaosClothAsset/Source/ChaosClothAssetEngine/Private/ChaosClothAsset/ClothComponent.cpp b/Engine/Plugins/Experimental/ChaosClothAsset/Source/ChaosClothAssetEngine/Private/ChaosClothAsset/ClothComponent.cpp
--- a/Engine/Plugins/Experimental/ChaosClothAsset/Source/ChaosClothAssetEngine/Private/ChaosClothAsset/ClothComponent.cpp
+++ b/Engine/Plugins/Experimental/ChaosClothAsset/Source/ChaosClothAssetEngine/Private/ChaosClothAsset/ClothComponent.cpp
@@ -14,9 +14,10 @@ CSV_DECLARE_CATEGORY_MODULE_EXTERN(ENGINE_API, Animation);
 
 UChaosClothComponent::UChaosClothComponent(const FObjectInitializer& ObjectInitializer)
 	: Super(ObjectInitializer)
+	, bUseAttachedParentAsPoseComponent(1)  // By default use the parent component as leader pose component
+	, bWaitForParallelTask(0)
 	, bDisableSimulation(0)
 	, bSuspendSimulation(0)
-	, bWaitForParallelTask(0)
 	, bBindToLeaderComponent(0)
 {
 	bAutoActivate = true;
@@ -93,12 +94,15 @@ void UChaosClothComponent::OnRegister()
 		const TSharedPtr<const FChaosClothSimulationModel> ClothSimulationModel = GetClothAsset()->GetClothSimulationModel();
 		if (ensure(ClothSimulationModel) && ClothSimulationModel->GetNumLods())
 		{
-			FSkeletalMeshLODRenderData& LODData = GetClothAsset()->GetResourceForRendering()->LODRenderData[GetPredictedLODLevel()];
-			GetClothAsset()->FillComponentSpaceTransforms(GetClothAsset()->GetRefSkeleton().GetRefBonePose(), LODData.RequiredBones, GetEditableComponentSpaceTransforms());
+			if (!LeaderPoseComponent.IsValid())
+			{
+				FSkeletalMeshLODRenderData& LODData = GetClothAsset()->GetResourceForRendering()->LODRenderData[GetPredictedLODLevel()];
+				GetClothAsset()->FillComponentSpaceTransforms(GetClothAsset()->GetRefSkeleton().GetRefBonePose(), LODData.RequiredBones, GetEditableComponentSpaceTransforms());
 
-			bNeedToFlipSpaceBaseBuffers = true; // Have updated space bases so need to flip
-			FlipEditableSpaceBases();
-			bHasValidBoneTransform = true;
+				bNeedToFlipSpaceBaseBuffers = true; // Have updated space bases so need to flip
+				FlipEditableSpaceBases();
+				bHasValidBoneTransform = true;
+			}
 
 			// Create simulation proxy
 			ClothSimulationProxy = MakeUnique<FClothSimulationProxy>(*this);
@@ -144,14 +148,11 @@ void UChaosClothComponent::TickComponent(float DeltaTime, enum ELevelTick TickTy
 	// Update the proxy and start the simulation parallel task
 	StartNewParallelSimulation(DeltaTime);
 
-	// TODO: Wait in tick function
-	//if (ShouldWaitForClothInTickFunction())
-	//{
-	//	FGraphEventArray Prerequisites;
-	//	Prerequisites.Add(ParallelClothTask);
-	//	FGraphEventRef ClothCompletionEvent = TGraphTask<FParallelClothCompletionTask>::CreateTask(&Prerequisites, ENamedThreads::GameThread).ConstructAndDispatchWhenReady(this);
-	//	ThisTickFunction.GetCompletionHandle()->DontCompleteUntil(ClothCompletionEvent);
-	//}
+	// Wait in tick function for the simulation results if required
+	if (ShouldWaitForParallelSimulationInTickComponent())
+	{
+		HandleExistingParallelSimulation();
+	}
 
 	Super::TickComponent(DeltaTime, TickType, ThisTickFunction);
 }
@@ -242,7 +243,7 @@ PRAGMA_ENABLE_DEPRECATION_WARNINGS
 
 bool UChaosClothComponent::RequiresPreEndOfFrameSync() const
 {
-	if (!IsSimulationSuspended() && !ShouldWaitForParallelSimulationInTickFunction())
+	if (!IsSimulationSuspended() && !ShouldWaitForParallelSimulationInTickComponent())
 	{
 		// By default we await the cloth task in TickComponent, but...
 		// If we have cloth and have no game-thread dependencies on the cloth output, 
@@ -285,6 +286,16 @@ FBoxSphereBounds UChaosClothComponent::CalcBounds(const FTransform& LocalToWorld
 	return NewBounds;
 }
 
+void UChaosClothComponent::OnAttachmentChanged()
+{
+	if (bUseAttachedParentAsPoseComponent)
+	{
+		SetLeaderPoseComponent(Cast<USkinnedMeshComponent>(GetAttachParent()));  // If the cast fail, remove the current leader
+	}
+
+	Super::OnAttachmentChanged();
+}
+
 void UChaosClothComponent::StartNewParallelSimulation(float DeltaTime)
 {
 	if (ClothSimulationProxy.IsValid())
@@ -310,7 +321,7 @@ void UChaosClothComponent::HandleExistingParallelSimulation()
 	}
 }
 
-bool UChaosClothComponent::ShouldWaitForParallelSimulationInTickFunction() const
+bool UChaosClothComponent::ShouldWaitForParallelSimulationInTickComponent() const
 {
 	static IConsoleVariable* const CVarClothPhysicsWaitForParallelClothTask = IConsoleManager::Get().FindConsoleVariable(TEXT("p.ClothPhysics.WaitForParallelClothTask"));
 
diff --git a/Engine/Plugins/Experimental/ChaosClothAsset/Source/ChaosClothAssetEngine/Public/ChaosClothAsset/ClothComponent.h b/Engine/Plugins/Experimental/ChaosClothAsset/Source/ChaosClothAssetEngine/Public/ChaosClothAsset/ClothComponent.h
--- a/Engine/Plugins/Experimental/ChaosClothAsset/Source/ChaosClothAssetEngine/Public/ChaosClothAsset/ClothComponent.h
+++ b/Engine/Plugins/Experimental/ChaosClothAsset/Source/ChaosClothAssetEngine/Public/ChaosClothAsset/ClothComponent.h
@@ -16,7 +16,11 @@ namespace UE::Chaos::ClothAsset
 /**
  * Cloth simulation component.
  */
-UCLASS(ClassGroup = Physics, Meta = (BlueprintSpawnableComponent, ToolTip = "Chaos cloth simulation component."), DisplayName = "Cloth Simulation", HideCategories = (Object, "Mesh|SkeletalAsset", Constraints, Advanced, Cooking, Collision))
+UCLASS(
+	ClassGroup = Physics, 
+	Meta = (BlueprintSpawnableComponent, ToolTip = "Chaos cloth simulation component."),
+	DisplayName = "Cloth Simulation",
+	HideCategories = (Object, "Mesh|SkeletalAsset", Constraints, Advanced, Cooking, Collision, Navigation))
 class CHAOSCLOTHASSETENGINE_API UChaosClothComponent : public USkinnedMeshComponent
 {
 	GENERATED_BODY()
@@ -82,6 +86,7 @@ protected:
 
 	//~ Begin USceneComponent Interface.
 	virtual FBoxSphereBounds CalcBounds(const FTransform& LocalToWorld) const override;
+	virtual void OnAttachmentChanged() override;
 	//~ End USceneComponent Interface
 
 	//~ Begin USkinnedMeshComponent Interface
@@ -92,7 +97,7 @@ protected:
 private:
 	void StartNewParallelSimulation(float DeltaTime);
 	void HandleExistingParallelSimulation();
-	bool ShouldWaitForParallelSimulationInTickFunction() const;
+	bool ShouldWaitForParallelSimulationInTickComponent() const;
 
 #if WITH_EDITORONLY_DATA
 	/** Cloth asset used by this component. */
@@ -101,6 +106,14 @@ private:
 	TObjectPtr<UChaosClothAsset> ClothAsset;
 #endif
 
+	/** If enabled, and the parent is another Skinned Mesh Component (e.g. another Cloth Component, Poseable Mesh Component, Skeletal Mesh Component, ...etc.), use its pose. */
+	UPROPERTY(EditAnywhere, Category = ClothComponent)
+	uint8 bUseAttachedParentAsPoseComponent : 1;
+
+	/** Whether to wait for the cloth simulation to end in the TickComponent instead of in the EndOfFrameUpdates. */
+	UPROPERTY(EditAnywhere, Category = ClothComponent)
+	uint8 bWaitForParallelTask : 1;
+
 	/** Whether to disable the simulation and use the skinned pose instead. */
 	UPROPERTY()
 	uint8 bDisableSimulation : 1;
@@ -109,9 +122,6 @@ private:
 	UPROPERTY()
 	uint8 bSuspendSimulation : 1;
 
-	/** Whether to wait for the cloth simulation to end in the TickComponent instead of in the EndOfFrameUpdates. */
-	UPROPERTY()
-	uint8 bWaitForParallelTask : 1;
 
 	/** Whether to use the leader component pose. */
 	UPROPERTY()
