Applying bone influences limit change to ChaosClothAsset code, which was previously restricted

diff --git a/Engine/Plugins/Experimental/ChaosClothAsset/Source/ChaosClothAssetEngine/Private/ChaosClothAsset/ClothAsset.cpp b/Engine/Plugins/Experimental/ChaosClothAsset/Source/ChaosClothAssetEngine/Private/ChaosClothAsset/ClothAsset.cpp
--- a/Engine/Plugins/Experimental/ChaosClothAsset/Source/ChaosClothAssetEngine/Private/ChaosClothAsset/ClothAsset.cpp
+++ b/Engine/Plugins/Experimental/ChaosClothAsset/Source/ChaosClothAssetEngine/Private/ChaosClothAsset/ClothAsset.cpp
@@ -8,6 +8,7 @@
 #include "ChaosClothAsset/ClothAssetPrivate.h"
 #include "ChaosClothAsset/ClothSimulationModel.h"
 #include "Animation/Skeleton.h"
+#include "Engine/RendererSettings.h"
 #include "Engine/SkinnedAssetAsyncCompileUtils.h"
 #include "Features/IModularFeatures.h"
 #include "Rendering/SkeletalMeshModel.h"
@@ -477,6 +478,9 @@ FString UChaosClothAsset::BuildDerivedDataKey(const ITargetPlatform* TargetPlatf
 	const bool bUnlimitedBoneInfluences = FGPUBaseSkinVertexFactory::GetUnlimitedBoneInfluences();
 	KeySuffix += bUnlimitedBoneInfluences ? "1" : "0";
 
+	// Include the global default bone influences limit in case any LODs don't set an explicit limit (highly likely)
+	KeySuffix += FString::FromInt(GetDefault<URendererSettings>()->DefaultBoneInfluenceLimit.GetValueForPlatform(*TargetPlatform->IniPlatformName()));
+
 	// Add LODInfoArray
 	TmpPartialKeySuffix = TEXT("");
 	TArray<FSkeletalMeshLODInfo>& LODInfos = GetLODInfoArray();
