Chaos Cloth - Add skinning weights data to the Cloth Collection.

#rnx

diff --git a/Engine/Plugins/Experimental/ChaosClothAsset/Source/ChaosClothAsset/Private/ChaosClothAsset/ClothCollection.cpp b/Engine/Plugins/Experimental/ChaosClothAsset/Source/ChaosClothAsset/Private/ChaosClothAsset/ClothCollection.cpp
--- a/Engine/Plugins/Experimental/ChaosClothAsset/Source/ChaosClothAsset/Private/ChaosClothAsset/ClothCollection.cpp
+++ b/Engine/Plugins/Experimental/ChaosClothAsset/Source/ChaosClothAsset/Private/ChaosClothAsset/ClothCollection.cpp
@@ -43,7 +43,9 @@ namespace UE::Chaos::ClothAsset
 		AddExternalAttribute<FVector2f>("SimPosition", SimVerticesGroup, SimPosition);
 		AddExternalAttribute<FVector3f>("SimRestPosition", SimVerticesGroup, SimRestPosition);
 		AddExternalAttribute<FVector3f>("SimRestNormal", SimVerticesGroup, SimRestNormal);
-		//AddExternalAttribute<FVertexBoneData>("SimBoneData", SimVerticesGroup, SimBoneData);
+		AddExternalAttribute<int32>("SimNumBoneInfluences", SimVerticesGroup, SimNumBoneInfluences);
+		AddExternalAttribute<TArray<int32>>("SimBoneIndices", SimVerticesGroup, SimBoneIndices);
+		AddExternalAttribute<TArray<float>>("SimBoneWeights", SimVerticesGroup, SimBoneWeights);
 
 		// Sim Faces Group
 		AddExternalAttribute<FIntVector3>("SimIndices", SimFacesGroup, SimIndices, SimVerticesDependency);
@@ -55,7 +57,9 @@ namespace UE::Chaos::ClothAsset
 		AddExternalAttribute<FVector3f>("RenderTangentV", RenderVerticesGroup, RenderTangentV);
 		AddExternalAttribute<TArray<FVector2f>>("RenderUVs", RenderVerticesGroup, RenderUVs);
 		AddExternalAttribute<FLinearColor>("RenderColor", RenderVerticesGroup, RenderColor);
-		//AddExternalAttribute<FVertexBoneData>("RenderBoneData", RenderVerticesGroup, RenderBoneData);
+		AddExternalAttribute<int32>("RenderNumBoneInfluences", RenderVerticesGroup, RenderNumBoneInfluences);
+		AddExternalAttribute<TArray<int32>>("RenderBoneIndices", RenderVerticesGroup, RenderBoneIndices);
+		AddExternalAttribute<TArray<float>>("RenderBoneWeights", RenderVerticesGroup, RenderBoneWeights);
 
 		// Render Faces Group
 		AddExternalAttribute<FIntVector3>("RenderIndices", RenderFacesGroup, RenderIndices, RenderVerticesDependency);
diff --git a/Engine/Plugins/Experimental/ChaosClothAsset/Source/ChaosClothAsset/Public/ChaosClothAsset/ClothCollection.h b/Engine/Plugins/Experimental/ChaosClothAsset/Source/ChaosClothAsset/Public/ChaosClothAsset/ClothCollection.h
--- a/Engine/Plugins/Experimental/ChaosClothAsset/Source/ChaosClothAsset/Public/ChaosClothAsset/ClothCollection.h
+++ b/Engine/Plugins/Experimental/ChaosClothAsset/Source/ChaosClothAsset/Public/ChaosClothAsset/ClothCollection.h
@@ -7,23 +7,9 @@
 // TODO:
 // - Transition Sim LOD data
 // - Move wrap deformer into struct similar to FMeshToMeshVertData
-// - Add FVertexBoneData to MANAGED_ARRAY_TYPE?
 
 namespace UE::Chaos::ClothAsset
 {
-	// TODO: Do we want to move the skinning data away from the static render data?
-	///**
-	// * Skinning data per vertex
-	// */
-	//struct FVertexBoneData
-	//{
-	//	static const int8 MaxTotalInfluences = 12;  // Up to MAX_TOTAL_INFLUENCES bone indices that this vert is weighted to
-	//
-	//	int32 NumInfluences;
-	//	uint16 BoneIndices[MaxTotalInfluences];
-	//	float BoneWeights[MaxTotalInfluences];
-	//};
-
 	/**
 	 * Tailored Cloth Asset Collection containing draping and pattern information.
 	 */
@@ -78,7 +64,9 @@ namespace UE::Chaos::ClothAsset
 		TManagedArray<FVector2f> SimPosition;
 		TManagedArray<FVector3f> SimRestPosition;
 		TManagedArray<FVector3f> SimRestNormal;  // Used for capture, maxdistance, backstop authoring ...etc
-		//TManagedArray<FVertexBoneData> SimBoneData;
+		TManagedArray<int32> SimNumBoneInfluences;
+		TManagedArray<TArray<int32>> SimBoneIndices;
+		TManagedArray<TArray<float>> SimBoneWeights;
 
 		// Sim Faces Group
 		TManagedArray<FIntVector3> SimIndices;  // The indices point to the elements in the Sim Vertices arrays but don't include the LOD start offset
@@ -90,7 +78,9 @@ namespace UE::Chaos::ClothAsset
 		TManagedArray<FVector3f> RenderTangentV;
 		TManagedArray<TArray<FVector2f>> RenderUVs;
 		TManagedArray<FLinearColor> RenderColor;
-		//TManagedArray<FVertexBoneData> RenderBoneData;
+		TManagedArray<int32> RenderNumBoneInfluences;
+		TManagedArray<TArray<int32>> RenderBoneIndices;
+		TManagedArray<TArray<float>> RenderBoneWeights;
 
 		// Render Faces Group
 		TManagedArray<FIntVector3> RenderIndices;  // The indices point to the elements in the Render Vertices arrays but don't include the LOD start offset
diff --git a/Engine/Plugins/Experimental/ChaosClothAsset/Source/ChaosClothAsset/Public/ChaosClothAsset/ClothLodAdapter.h b/Engine/Plugins/Experimental/ChaosClothAsset/Source/ChaosClothAsset/Public/ChaosClothAsset/ClothLodAdapter.h
--- a/Engine/Plugins/Experimental/ChaosClothAsset/Source/ChaosClothAsset/Public/ChaosClothAsset/ClothLodAdapter.h
+++ b/Engine/Plugins/Experimental/ChaosClothAsset/Source/ChaosClothAsset/Public/ChaosClothAsset/ClothLodAdapter.h
@@ -54,6 +54,9 @@ namespace UE::Chaos::ClothAsset
 		TConstArrayView<FVector2f> GetPatternsSimPosition() const { return GetClothCollection()->GetPatternsElements(GetClothCollection()->SimPosition, GetClothCollection()->SimVerticesStart, GetClothCollection()->SimVerticesEnd, GetElementIndex()); }
 		TConstArrayView<FVector3f> GetPatternsSimRestPosition() const { return GetClothCollection()->GetPatternsElements(GetClothCollection()->SimRestPosition, GetClothCollection()->SimVerticesStart, GetClothCollection()->SimVerticesEnd, GetElementIndex()); }
 		TConstArrayView<FVector3f> GetPatternsSimRestNormal() const { return GetClothCollection()->GetPatternsElements(GetClothCollection()->SimRestNormal, GetClothCollection()->SimVerticesStart, GetClothCollection()->SimVerticesEnd, GetElementIndex()); }
+		TConstArrayView<int32> GetPatternsSimNumBoneInfluences() const { return GetClothCollection()->GetPatternsElements(GetClothCollection()->SimNumBoneInfluences, GetClothCollection()->SimVerticesStart, GetClothCollection()->SimVerticesEnd, GetElementIndex()); }
+		TConstArrayView<TArray<int32>> GetPatternsSimBoneIndices() const { return GetClothCollection()->GetPatternsElements(GetClothCollection()->SimBoneIndices, GetClothCollection()->SimVerticesStart, GetClothCollection()->SimVerticesEnd, GetElementIndex()); }
+		TConstArrayView<TArray<float>> GetPatternsSimBoneWeights() const { return GetClothCollection()->GetPatternsElements(GetClothCollection()->SimBoneWeights, GetClothCollection()->SimVerticesStart, GetClothCollection()->SimVerticesEnd, GetElementIndex()); }
 
 		// All patterns Sim Faces Group, use these to access all pattern at once
 		TTuple<int32, int32> GetPatternsSimFacesRange() const { return GetClothCollection()->GetPatternsElementsStartEnd<true, true>(GetClothCollection()->SimFacesStart, GetClothCollection()->SimFacesEnd, GetElementIndex()); }
@@ -72,6 +75,9 @@ namespace UE::Chaos::ClothAsset
 		TConstArrayView<FVector3f> GetPatternsRenderTangentV() const { return GetClothCollection()->GetPatternsElements(GetClothCollection()->RenderTangentV, GetClothCollection()->RenderVerticesStart, GetClothCollection()->RenderVerticesEnd, GetElementIndex()); }
 		TConstArrayView<TArray<FVector2f>> GePatternstRenderUVs() const { return GetClothCollection()->GetPatternsElements(GetClothCollection()->RenderUVs, GetClothCollection()->RenderVerticesStart, GetClothCollection()->RenderVerticesEnd, GetElementIndex()); }
 		TConstArrayView<FLinearColor> GetPatternsRenderColor() const { return GetClothCollection()->GetPatternsElements(GetClothCollection()->RenderColor, GetClothCollection()->RenderVerticesStart, GetClothCollection()->RenderVerticesEnd, GetElementIndex()); }
+		TConstArrayView<int32> GetPatternsRenderNumBoneInfluences() const { return GetClothCollection()->GetPatternsElements(GetClothCollection()->RenderNumBoneInfluences, GetClothCollection()->RenderVerticesStart, GetClothCollection()->RenderVerticesEnd, GetElementIndex()); }
+		TConstArrayView<TArray<int32>> GetPatternsRenderBoneIndices() const { return GetClothCollection()->GetPatternsElements(GetClothCollection()->RenderBoneIndices, GetClothCollection()->RenderVerticesStart, GetClothCollection()->RenderVerticesEnd, GetElementIndex()); }
+		TConstArrayView<TArray<float>> GetPatternsRenderBoneWeights() const { return GetClothCollection()->GetPatternsElements(GetClothCollection()->RenderBoneWeights, GetClothCollection()->RenderVerticesStart, GetClothCollection()->RenderVerticesEnd, GetElementIndex()); }
 
 		// All patterns Sim Faces Group, use these to access all pattern at once
 		TTuple<int32, int32> GetPatternsRenderFacesRange() const { return GetClothCollection()->GetPatternsElementsStartEnd<true, true>(GetClothCollection()->RenderFacesStart, GetClothCollection()->RenderFacesEnd, GetElementIndex()); }
@@ -147,6 +153,9 @@ namespace UE::Chaos::ClothAsset
 		TArrayView<FVector2f> GetPatternsSimPosition() { return GetClothCollection()->GetPatternsElements(GetClothCollection()->SimPosition, GetClothCollection()->SimVerticesStart, GetClothCollection()->SimVerticesEnd, GetElementIndex()); }
 		TArrayView<FVector3f> GetPatternsSimRestPosition() { return GetClothCollection()->GetPatternsElements(GetClothCollection()->SimRestPosition, GetClothCollection()->SimVerticesStart, GetClothCollection()->SimVerticesEnd, GetElementIndex()); }
 		TArrayView<FVector3f> GetPatternsSimRestNormal() { return GetClothCollection()->GetPatternsElements(GetClothCollection()->SimRestNormal, GetClothCollection()->SimVerticesStart, GetClothCollection()->SimVerticesEnd, GetElementIndex()); }
+		TArrayView<int32> GetPatternsSimNumBoneInfluences() { return GetClothCollection()->GetPatternsElements(GetClothCollection()->SimNumBoneInfluences, GetClothCollection()->SimVerticesStart, GetClothCollection()->SimVerticesEnd, GetElementIndex()); }
+		TArrayView<TArray<int32>> GetPatternsSimBoneIndices() { return GetClothCollection()->GetPatternsElements(GetClothCollection()->SimBoneIndices, GetClothCollection()->SimVerticesStart, GetClothCollection()->SimVerticesEnd, GetElementIndex()); }
+		TArrayView<TArray<float>> GetPatternsSimBoneWeights() { return GetClothCollection()->GetPatternsElements(GetClothCollection()->SimBoneWeights, GetClothCollection()->SimVerticesStart, GetClothCollection()->SimVerticesEnd, GetElementIndex()); }
 
 		// All patterns Sim Faces Group, use these to access all pattern at once
 		TArrayView<FIntVector3> GetPatternsSimIndices() { return GetClothCollection()->GetPatternsElements(GetClothCollection()->SimIndices, GetClothCollection()->SimFacesStart, GetClothCollection()->SimFacesEnd, GetElementIndex()); }
@@ -158,6 +167,9 @@ namespace UE::Chaos::ClothAsset
 		TArrayView<FVector3f> GetPatternsRenderTangentV() { return GetClothCollection()->GetPatternsElements(GetClothCollection()->RenderTangentV, GetClothCollection()->RenderVerticesStart, GetClothCollection()->RenderVerticesEnd, GetElementIndex()); }
 		TArrayView<TArray<FVector2f>> GePatternstRenderUVs() { return GetClothCollection()->GetPatternsElements(GetClothCollection()->RenderUVs, GetClothCollection()->RenderVerticesStart, GetClothCollection()->RenderVerticesEnd, GetElementIndex()); }
 		TArrayView<FLinearColor> GetPatternsRenderColor() { return GetClothCollection()->GetPatternsElements(GetClothCollection()->RenderColor, GetClothCollection()->RenderVerticesStart, GetClothCollection()->RenderVerticesEnd, GetElementIndex()); }
+		TArrayView<int32> GetPatternsRenderNumBoneInfluences() { return GetClothCollection()->GetPatternsElements(GetClothCollection()->RenderNumBoneInfluences, GetClothCollection()->RenderVerticesStart, GetClothCollection()->RenderVerticesEnd, GetElementIndex()); }
+		TArrayView<TArray<int32>> GetPatternsRenderBoneIndices() { return GetClothCollection()->GetPatternsElements(GetClothCollection()->RenderBoneIndices, GetClothCollection()->RenderVerticesStart, GetClothCollection()->RenderVerticesEnd, GetElementIndex()); }
+		TArrayView<TArray<float>> GetPatternsRenderBoneWeights() { return GetClothCollection()->GetPatternsElements(GetClothCollection()->RenderBoneWeights, GetClothCollection()->RenderVerticesStart, GetClothCollection()->RenderVerticesEnd, GetElementIndex()); }
 
 		// All patterns Render Faces Group, use these to access all pattern at once
 		TArrayView<FIntVector3> GetPatternsRenderIndices() { return GetClothCollection()->GetPatternsElements(GetClothCollection()->RenderIndices, GetClothCollection()->RenderFacesStart, GetClothCollection()->RenderFacesEnd, GetElementIndex()); }
diff --git a/Engine/Plugins/Experimental/ChaosClothAsset/Source/ChaosClothAsset/Public/ChaosClothAsset/ClothPatternAdapter.h b/Engine/Plugins/Experimental/ChaosClothAsset/Source/ChaosClothAsset/Public/ChaosClothAsset/ClothPatternAdapter.h
--- a/Engine/Plugins/Experimental/ChaosClothAsset/Source/ChaosClothAsset/Public/ChaosClothAsset/ClothPatternAdapter.h
+++ b/Engine/Plugins/Experimental/ChaosClothAsset/Source/ChaosClothAsset/Public/ChaosClothAsset/ClothPatternAdapter.h
@@ -26,6 +26,9 @@ namespace UE::Chaos::ClothAsset
 		TConstArrayView<FVector2f> GetSimPosition() const { return GetClothCollection()->GetElements(GetClothCollection()->SimPosition, GetClothCollection()->SimVerticesStart, GetClothCollection()->SimVerticesEnd, GetElementIndex()); }
 		TConstArrayView<FVector3f> GetSimRestPosition() const { return GetClothCollection()->GetElements(GetClothCollection()->SimRestPosition, GetClothCollection()->SimVerticesStart, GetClothCollection()->SimVerticesEnd, GetElementIndex()); }
 		TConstArrayView<FVector3f> GetSimRestNormal() const { return GetClothCollection()->GetElements(GetClothCollection()->SimRestNormal, GetClothCollection()->SimVerticesStart, GetClothCollection()->SimVerticesEnd, GetElementIndex()); }
+		TConstArrayView<int32> GetSimNumBoneInfluences() const { return GetClothCollection()->GetElements(GetClothCollection()->SimNumBoneInfluences, GetClothCollection()->SimVerticesStart, GetClothCollection()->SimVerticesEnd, GetElementIndex()); }
+		TConstArrayView<TArray<int32>> GetSimBoneIndices() const { return GetClothCollection()->GetElements(GetClothCollection()->SimBoneIndices, GetClothCollection()->SimVerticesStart, GetClothCollection()->SimVerticesEnd, GetElementIndex()); }
+		TConstArrayView<TArray<float>> GetSimBoneWeights() const { return GetClothCollection()->GetElements(GetClothCollection()->SimBoneWeights, GetClothCollection()->SimVerticesStart, GetClothCollection()->SimVerticesEnd, GetElementIndex()); }
 
 		// Sim Faces Group, note: SimIndices points to the LOD arrays, not the pattern arrays
 		int32 GetNumSimFaces() const { return GetClothCollection()->GetNumElements(GetClothCollection()->SimFacesStart, GetClothCollection()->SimFacesEnd, GetElementIndex()); }
@@ -40,6 +43,9 @@ namespace UE::Chaos::ClothAsset
 		TConstArrayView<FVector3f> GetRenderTangentV() const { return GetClothCollection()->GetElements(GetClothCollection()->RenderTangentV, GetClothCollection()->RenderVerticesStart, GetClothCollection()->RenderVerticesEnd, GetElementIndex()); }
 		TConstArrayView<TArray<FVector2f>> GetRenderUVs() const { return GetClothCollection()->GetElements(GetClothCollection()->RenderUVs, GetClothCollection()->RenderVerticesStart, GetClothCollection()->RenderVerticesEnd, GetElementIndex()); }
 		TConstArrayView<FLinearColor> GetRenderColor() const { return GetClothCollection()->GetElements(GetClothCollection()->RenderColor, GetClothCollection()->RenderVerticesStart, GetClothCollection()->RenderVerticesEnd, GetElementIndex()); }
+		TConstArrayView<int32> GetRenderNumBoneInfluences() const { return GetClothCollection()->GetElements(GetClothCollection()->RenderNumBoneInfluences, GetClothCollection()->RenderVerticesStart, GetClothCollection()->RenderVerticesEnd, GetElementIndex()); }
+		TConstArrayView<TArray<int32>> GetRenderBoneIndices() const { return GetClothCollection()->GetElements(GetClothCollection()->RenderBoneIndices, GetClothCollection()->RenderVerticesStart, GetClothCollection()->RenderVerticesEnd, GetElementIndex()); }
+		TConstArrayView<TArray<float>> GetRenderBoneWeights() const { return GetClothCollection()->GetElements(GetClothCollection()->RenderBoneWeights, GetClothCollection()->RenderVerticesStart, GetClothCollection()->RenderVerticesEnd, GetElementIndex()); }
 
 		// Render Faces Group, note: RenderIndices points to the LOD arrays, not the pattern arrays
 		int32 GetNumRenderFaces() const { return GetClothCollection()->GetNumElements(GetClothCollection()->RenderFacesStart, GetClothCollection()->RenderFacesEnd, GetElementIndex()); }
@@ -111,6 +117,9 @@ namespace UE::Chaos::ClothAsset
 		TArrayView<FVector2f> GetSimPosition() { return GetClothCollection()->GetElements(GetClothCollection()->SimPosition, GetClothCollection()->SimVerticesStart, GetClothCollection()->SimVerticesEnd, GetElementIndex()); }
 		TArrayView<FVector3f> GetSimRestPosition() { return GetClothCollection()->GetElements(GetClothCollection()->SimRestPosition, GetClothCollection()->SimVerticesStart, GetClothCollection()->SimVerticesEnd, GetElementIndex()); }
 		TArrayView<FVector3f> GetSimRestNormal() { return GetClothCollection()->GetElements(GetClothCollection()->SimRestNormal, GetClothCollection()->SimVerticesStart, GetClothCollection()->SimVerticesEnd, GetElementIndex()); }
+		TArrayView<int32> GetSimNumBoneInfluences() { return GetClothCollection()->GetElements(GetClothCollection()->SimNumBoneInfluences, GetClothCollection()->SimVerticesStart, GetClothCollection()->SimVerticesEnd, GetElementIndex()); }
+		TArrayView<TArray<int32>> GetSimBoneIndices() { return GetClothCollection()->GetElements(GetClothCollection()->SimBoneIndices, GetClothCollection()->SimVerticesStart, GetClothCollection()->SimVerticesEnd, GetElementIndex()); }
+		TArrayView<TArray<float>> GetSimBoneWeights() { return GetClothCollection()->GetElements(GetClothCollection()->SimBoneWeights, GetClothCollection()->SimVerticesStart, GetClothCollection()->SimVerticesEnd, GetElementIndex()); }
 
 		// Sim Faces Group, note: SimIndices points to the LOD arrays, not the pattern arrays
 		TArrayView<FIntVector3> GetSimIndices() { return GetClothCollection()->GetElements(GetClothCollection()->SimIndices, GetClothCollection()->SimFacesStart, GetClothCollection()->SimFacesEnd, GetElementIndex()); }
@@ -123,6 +132,9 @@ namespace UE::Chaos::ClothAsset
 		TArrayView<FVector3f> GetRenderTangentV() { return GetClothCollection()->GetElements(GetClothCollection()->RenderTangentV, GetClothCollection()->RenderVerticesStart, GetClothCollection()->RenderVerticesEnd, GetElementIndex()); }
 		TArrayView<TArray<FVector2f>> GetRenderUVs() { return GetClothCollection()->GetElements(GetClothCollection()->RenderUVs, GetClothCollection()->RenderVerticesStart, GetClothCollection()->RenderVerticesEnd, GetElementIndex()); }
 		TArrayView<FLinearColor> GetRenderColor() { return GetClothCollection()->GetElements(GetClothCollection()->RenderColor, GetClothCollection()->RenderVerticesStart, GetClothCollection()->RenderVerticesEnd, GetElementIndex()); }
+		TArrayView<int32> GetRenderNumBoneInfluences() { return GetClothCollection()->GetElements(GetClothCollection()->RenderNumBoneInfluences, GetClothCollection()->RenderVerticesStart, GetClothCollection()->RenderVerticesEnd, GetElementIndex()); }
+		TArrayView<TArray<int32>> GetRenderBoneIndices() { return GetClothCollection()->GetElements(GetClothCollection()->RenderBoneIndices, GetClothCollection()->RenderVerticesStart, GetClothCollection()->RenderVerticesEnd, GetElementIndex()); }
+		TArrayView<TArray<float>> GetRenderBoneWeights() { return GetClothCollection()->GetElements(GetClothCollection()->RenderBoneWeights, GetClothCollection()->RenderVerticesStart, GetClothCollection()->RenderVerticesEnd, GetElementIndex()); }
 
 		// Render Faces Group, note: RenderIndices points to the LOD arrays, not the pattern arrays
 		TArrayView<FIntVector3> GetRenderIndices() { return GetClothCollection()->GetElements(GetClothCollection()->RenderIndices, GetClothCollection()->RenderFacesStart, GetClothCollection()->RenderFacesEnd, GetElementIndex()); }
