Added support for setting r.GPUSkin.UnlimitedBoneInfluences and r.GPUSkin.UnlimitedBoneInfluencesThreshold per platform

Second CL for restricted file

diff --git a/Engine/Plugins/Experimental/ChaosClothAsset/Source/ChaosClothAssetEngine/Private/ChaosClothAsset/ClothAsset.cpp b/Engine/Plugins/Experimental/ChaosClothAsset/Source/ChaosClothAssetEngine/Private/ChaosClothAsset/ClothAsset.cpp
--- a/Engine/Plugins/Experimental/ChaosClothAsset/Source/ChaosClothAssetEngine/Private/ChaosClothAsset/ClothAsset.cpp
+++ b/Engine/Plugins/Experimental/ChaosClothAsset/Source/ChaosClothAssetEngine/Private/ChaosClothAsset/ClothAsset.cpp
@@ -475,7 +475,7 @@ FString UChaosClothAsset::BuildDerivedDataKey(const ITargetPlatform* TargetPlatf
 	KeySuffix += FString::FromInt(MaxGPUSkinBones);
 	// Add unlimited bone influences mode
 	IMeshBuilderModule::GetForPlatform(TargetPlatform).AppendToDDCKey(KeySuffix, true);
-	const bool bUnlimitedBoneInfluences = FGPUBaseSkinVertexFactory::GetUnlimitedBoneInfluences();
+	const bool bUnlimitedBoneInfluences = FGPUBaseSkinVertexFactory::GetUnlimitedBoneInfluences(TargetPlatform);
 	KeySuffix += bUnlimitedBoneInfluences ? "1" : "0";
 
 	// Include the global default bone influences limit in case any LODs don't set an explicit limit (highly likely)
