Chaos Cloth - Fixed cloth collection comment.

#rnx

diff --git a/Engine/Plugins/Experimental/ChaosClothAsset/Source/ChaosClothAsset/Public/ChaosClothAsset/ClothCollection.h b/Engine/Plugins/Experimental/ChaosClothAsset/Source/ChaosClothAsset/Public/ChaosClothAsset/ClothCollection.h
--- a/Engine/Plugins/Experimental/ChaosClothAsset/Source/ChaosClothAsset/Public/ChaosClothAsset/ClothCollection.h
+++ b/Engine/Plugins/Experimental/ChaosClothAsset/Source/ChaosClothAsset/Public/ChaosClothAsset/ClothCollection.h
@@ -48,7 +48,7 @@ namespace UE::Chaos::ClothAsset
 		template<bool bStart = true, bool bEnd = true>
 		inline TTuple<int32, int32> GetPatternsElementsStartEnd(const TManagedArray<int32>& StartArray, const TManagedArray<int32>& EndArray, int32 StartEndIndex) const;
 
-		// Attribute groups, predefined data member of the FClothLod object.
+		// Attribute groups, predefined data member of the collection
 		static const FName SimVerticesGroup;  // Contains patterns' 2D positions, 3D draped position (rest)
 		static const FName SimFacesGroup;  // Contains indices to sim vertex
 		static const FName RenderVerticesGroup;  // Contains pattern's 3D render model
