Chaos Cloth - Fix crash on cloth component with empty simulation proxy.

#rnx

diff --git a/Engine/Plugins/Experimental/ChaosClothAsset/Source/ChaosClothAssetEngine/Private/ChaosClothAsset/ClothComponent.cpp b/Engine/Plugins/Experimental/ChaosClothAsset/Source/ChaosClothAssetEngine/Private/ChaosClothAsset/ClothComponent.cpp
--- a/Engine/Plugins/Experimental/ChaosClothAsset/Source/ChaosClothAssetEngine/Private/ChaosClothAsset/ClothComponent.cpp
+++ b/Engine/Plugins/Experimental/ChaosClothAsset/Source/ChaosClothAssetEngine/Private/ChaosClothAsset/ClothComponent.cpp
@@ -203,7 +203,7 @@ void UChaosClothComponent::GetUpdateClothSimulationData_AnyThread(TMap<int32, FC
 		OutBlendWeight = BlendWeight;
 		OutClothSimulData = LeaderPoseClothComponent->ClothSimulationProxy->GetCurrentSimulationData_AnyThread();
 	}
-	else if (!bDisableSimulation && !bBindToLeaderComponent)
+	else if (!bDisableSimulation && !bBindToLeaderComponent && ClothSimulationProxy)
 	{
 		OutBlendWeight = BlendWeight;
 		OutClothSimulData = ClothSimulationProxy->GetCurrentSimulationData_AnyThread();
