ClothAsset: change category for cloth deformer properties

diff --git a/Engine/Plugins/Experimental/ChaosClothAsset/Source/ChaosClothAssetEngine/Public/ChaosClothAsset/ClothAsset.h b/Engine/Plugins/Experimental/ChaosClothAsset/Source/ChaosClothAssetEngine/Public/ChaosClothAsset/ClothAsset.h
--- a/Engine/Plugins/Experimental/ChaosClothAsset/Source/ChaosClothAssetEngine/Public/ChaosClothAsset/ClothAsset.h
+++ b/Engine/Plugins/Experimental/ChaosClothAsset/Source/ChaosClothAssetEngine/Public/ChaosClothAsset/ClothAsset.h
@@ -213,15 +213,15 @@ private:
 	int32 RayTracingMinLOD;
 
 	/** Whether to blend positions between the skinned/simulated transitions of the cloth render mesh. */
-	UPROPERTY(EditAnywhere, Category = RayTracing)
+	UPROPERTY(EditAnywhere, Category = ClothDeformer)
 	bool bSmoothTransition = true;
 
-	/** Whether to use multiple triangle influences on the proxy wrap deformer to help smoothen deformations. */
-	UPROPERTY(EditAnywhere, Category = RayTracing)
+	/** Whether to use multiple triangle influences on the proxy wrap deformer to help smoothe deformations. */
+	UPROPERTY(EditAnywhere, Category = ClothDeformer)
 	bool bUseMultipleInfluences = false;
 
 	/** The radius from which to get the multiple triangle influences from the simulated proxy mesh. */
-	UPROPERTY(EditAnywhere, Category = RayTracing)
+	UPROPERTY(EditAnywhere, Category = ClothDeformer)
 	float SkinningKernelRadius = 30.f;
 
 	/**
