ChaosCloth: Add ClothComponentCacheAdapter for chaoscache static pose mode

diff --git a/Engine/Plugins/Experimental/ChaosCloth/Source/Chaos/Private/ChaosCloth/ChaosClothingSimulationSolver.cpp b/Engine/Plugins/Experimental/ChaosCloth/Source/Chaos/Private/ChaosCloth/ChaosClothingSimulationSolver.cpp
--- a/Engine/Plugins/Experimental/ChaosCloth/Source/Chaos/Private/ChaosCloth/ChaosClothingSimulationSolver.cpp
+++ b/Engine/Plugins/Experimental/ChaosCloth/Source/Chaos/Private/ChaosCloth/ChaosClothingSimulationSolver.cpp
@@ -1162,12 +1162,13 @@ void FClothingSimulationSolver::UpdateFromCache(const TArray<FVector>& CachedPos
 {
 	Chaos::Softs::FSolverParticles& SolverParticles = Evolution->Particles();
 	const int32 NumParticles = GetNumParticles();
+	const bool bHasVelocity = CachedVelocities.Num() > 0;
 	if(CachedPositions.Num() == NumParticles)
 	{
 		for(int32 ParticleIndex = 0; ParticleIndex < NumParticles; ++ParticleIndex)
 		{
 			SolverParticles.X(ParticleIndex) = CachedPositions[ParticleIndex];
-			SolverParticles.V(ParticleIndex) = CachedVelocities[ParticleIndex];
+			SolverParticles.V(ParticleIndex) = bHasVelocity ? CachedVelocities[ParticleIndex] : FVector::ZeroVector;
 		}
 	}
 	PhysicsParallelFor(Cloths.Num(), [this](int32 ClothIndex)
diff --git a/Engine/Plugins/Experimental/ChaosCloth/Source/Chaos/Private/ChaosCloth/SkeletalMeshComponentCacheAdapter.cpp b/Engine/Plugins/Experimental/ChaosCloth/Source/Chaos/Private/ChaosCloth/SkeletalMeshComponentCacheAdapter.cpp
--- a/Engine/Plugins/Experimental/ChaosCloth/Source/Chaos/Private/ChaosCloth/SkeletalMeshComponentCacheAdapter.cpp
+++ b/Engine/Plugins/Experimental/ChaosCloth/Source/Chaos/Private/ChaosCloth/SkeletalMeshComponentCacheAdapter.cpp
@@ -33,7 +33,7 @@ namespace Chaos
 
 	uint8 FSkeletalMeshCacheAdapter::GetPriority() const
 	{
-		return EngineAdapterPriorityBegin;
+		return EngineAdapterPriorityBegin + 1; // A trick to prioritize cloth component over skeletal mesh component
 	}
 
 	void FSkeletalMeshCacheAdapter::Record_PostSolve(UPrimitiveComponent* InComponent, const FTransform& InRootTransform, FPendingFrameWrite& OutFrame, Chaos::FReal InTime) const
diff --git a/Engine/Plugins/Experimental/ChaosClothAsset/Source/ChaosClothAssetEngine/ChaosClothAssetEngine.Build.cs b/Engine/Plugins/Experimental/ChaosClothAsset/Source/ChaosClothAssetEngine/ChaosClothAssetEngine.Build.cs
--- a/Engine/Plugins/Experimental/ChaosClothAsset/Source/ChaosClothAssetEngine/ChaosClothAssetEngine.Build.cs
+++ b/Engine/Plugins/Experimental/ChaosClothAsset/Source/ChaosClothAssetEngine/ChaosClothAssetEngine.Build.cs
@@ -19,6 +19,7 @@ public class ChaosClothAssetEngine : ModuleRules
 				"Chaos",
 				"ChaosClothAsset",
 				"ChaosCloth",
+				"ChaosCaching",
 				"ProfileVisualizer",
 				"ClothingSystemRuntimeInterface",
 				"DataflowEngine"
diff --git a/Engine/Plugins/Experimental/ChaosClothAsset/Source/ChaosClothAssetEngine/Private/ChaosClothAsset/ChaosClothAssetEngineModule.cpp b/Engine/Plugins/Experimental/ChaosClothAsset/Source/ChaosClothAssetEngine/Private/ChaosClothAsset/ChaosClothAssetEngineModule.cpp
--- a/Engine/Plugins/Experimental/ChaosClothAsset/Source/ChaosClothAssetEngine/Private/ChaosClothAsset/ChaosClothAssetEngineModule.cpp
+++ b/Engine/Plugins/Experimental/ChaosClothAsset/Source/ChaosClothAssetEngine/Private/ChaosClothAsset/ChaosClothAssetEngineModule.cpp
@@ -3,6 +3,7 @@
 #include "ChaosClothAsset/ClothAssetPrivate.h"
 #include "Modules/ModuleInterface.h"
 #include "Modules/ModuleManager.h"
+#include "ChaosClothAsset/ClothComponentCacheAdapter.h"
 
 DEFINE_LOG_CATEGORY(LogChaosClothAsset);
 
@@ -13,11 +14,15 @@ namespace UE::Chaos::ClothAsset
 	public:
 		virtual void StartupModule() override
 		{
+			ClothComponentCacheAdapter = MakeUnique<FClothComponentCacheAdapter>();
+			RegisterAdapter(ClothComponentCacheAdapter.Get());
 		}
 
 		virtual void ShutdownModule() override
 		{
 		}
+	private:
+		TUniquePtr<FClothComponentCacheAdapter> ClothComponentCacheAdapter;
 	};
 }
 IMPLEMENT_MODULE(UE::Chaos::ClothAsset::FClothAssetEngineModule, ChaosClothAssetEngine);
diff --git a/Engine/Plugins/Experimental/ChaosClothAsset/Source/ChaosClothAssetEngine/Private/ChaosClothAsset/ClothComponentCacheAdapter.cpp b/Engine/Plugins/Experimental/ChaosClothAsset/Source/ChaosClothAssetEngine/Private/ChaosClothAsset/ClothComponentCacheAdapter.cpp
new file mode 100644
--- /dev/null
+++ b/Engine/Plugins/Experimental/ChaosClothAsset/Source/ChaosClothAssetEngine/Private/ChaosClothAsset/ClothComponentCacheAdapter.cpp
@@ -0,0 +1,180 @@
+// Copyright Epic Games, Inc. All Rights Reserved.
+
+#include "ChaosClothAsset/ClothComponentCacheAdapter.h"
+#include "Chaos/ParticleHandle.h"
+#include "Chaos/ChaosCache.h"
+#include "Chaos/CacheManagerActor.h"
+#include "Chaos/PBDEvolution.h"
+#include "ChaosCloth/ChaosClothingSimulation.h"
+#include "ChaosCloth/ChaosClothingSimulationSolver.h"
+#include "ChaosClothAsset/ClothComponent.h"
+#include "ChaosClothAsset/ClothSimulationProxy.h"
+#include "Components/SkeletalMeshComponent.h"
+
+namespace UE::Chaos::ClothAsset
+{
+	::Chaos::FComponentCacheAdapter::SupportType FClothComponentCacheAdapter::SupportsComponentClass(UClass* InComponentClass) const
+	{
+		/** The desired class is UChaosClothComponent. 
+		 * But when UChaosClothComponent is a child component of a USkeletalMeshComponent, UChaosClothComponent is not selectible using SoftComponentRef. 
+		 * To get around this issue, we need to select the parent USkeletalMeshComponent in SoftComponentRef, 
+		 * make USkeletalMeshComponent also a supported class for FClothComponentAdapter,
+		 *  and prioritize FClothComponentAdapter over FSkeletalMeshComponentAdapter in CacheManagerActor.
+		 */
+		const UClass* Desired = GetDesiredClass();
+		if(InComponentClass == Desired || InComponentClass == USkeletalMeshComponent::StaticClass())
+		{
+			return ::Chaos::FComponentCacheAdapter::SupportType::Direct;
+		}
+		else if(InComponentClass->IsChildOf(Desired) || InComponentClass->IsChildOf(USkeletalMeshComponent::StaticClass()))
+		{
+			return ::Chaos::FComponentCacheAdapter::SupportType::Derived;
+		}
+
+		return ::Chaos::FComponentCacheAdapter::SupportType::None;
+	}
+
+	UClass* FClothComponentCacheAdapter::GetDesiredClass() const
+	{
+		return UChaosClothComponent::StaticClass();
+	}
+
+	uint8 FClothComponentCacheAdapter::GetPriority() const
+	{
+		return EngineAdapterPriorityBegin;
+	}
+
+	void FClothComponentCacheAdapter::Record_PostSolve(UPrimitiveComponent* InComponent, const FTransform& InRootTransform, FPendingFrameWrite& OutFrame, ::Chaos::FReal InTime) const
+	{
+		ensureMsgf(false, TEXT("Record not implemented."));
+	}
+
+	void FClothComponentCacheAdapter::Playback_PreSolve(UPrimitiveComponent* InComponent, UChaosCache* InCache, ::Chaos::FReal InTime, FPlaybackTickRecord& TickRecord, TArray<::Chaos::TPBDRigidParticleHandle<::Chaos::FReal, 3>*>& OutUpdatedRigids) const
+	{
+		ensureMsgf(false, TEXT("Playback not implemented."));
+	}
+
+	FGuid FClothComponentCacheAdapter::GetGuid() const
+	{
+		FGuid NewGuid;
+		checkSlow(FGuid::Parse(TEXT("C704F4F536A34CD4973ABBB7BFEEE432"), NewGuid));
+		return NewGuid;
+	}
+
+	bool FClothComponentCacheAdapter::ValidForPlayback(UPrimitiveComponent* InComponent, UChaosCache* InCache) const
+	{
+		const UChaosClothComponent* ClothComp = GetClothComponent(InComponent);
+		return ClothComp && InCache->TrackToParticle.Num() > 0;
+	}
+
+	::Chaos::FPhysicsSolverEvents* FClothComponentCacheAdapter::BuildEventsSolver(UPrimitiveComponent* InComponent) const
+	{
+		ensureMsgf(false, TEXT("Playback or record not implemented."));
+		return nullptr;
+	}
+	
+	::Chaos::FPhysicsSolver* FClothComponentCacheAdapter::GetComponentSolver(UPrimitiveComponent* InComponent) const
+	{
+		return nullptr;
+	}
+	
+	void FClothComponentCacheAdapter::SetRestState(UPrimitiveComponent* InComponent, UChaosCache* InCache, const FTransform& InRootTransform, ::Chaos::FReal InTime) const
+	{
+		if (!InCache || InCache->GetDuration() == 0.0f)
+		{
+			return;
+		}
+
+		if (UChaosClothComponent* ClothComp = GetClothComponent(InComponent))
+		{
+			FPlaybackTickRecord TickRecord;
+			TickRecord.SetLastTime(InTime);
+		
+			FCacheEvaluationContext Context(TickRecord);
+			Context.bEvaluateTransform = false;
+			Context.bEvaluateCurves = false;
+			Context.bEvaluateEvents = false;
+			Context.bEvaluateChannels = true;
+
+			FCacheEvaluationResult EvaluatedResult = InCache->Evaluate(Context, nullptr);
+			const int32 NumCurves = EvaluatedResult.ParticleIndices.Num();
+
+			const TArray<float>* const PendingVX = EvaluatedResult.Channels.Find(VelocityXName);
+			const TArray<float>* const PendingVY = EvaluatedResult.Channels.Find(VelocityYName);
+			const TArray<float>* const PendingVZ = EvaluatedResult.Channels.Find(VelocityZName);
+			const TArray<float>* const PendingPX = EvaluatedResult.Channels.Find(PositionXName);
+			const TArray<float>* const PendingPY = EvaluatedResult.Channels.Find(PositionYName);
+			const TArray<float>* const PendingPZ = EvaluatedResult.Channels.Find(PositionZName);
+
+			TArray<FVector> CachedPositions; 
+			const bool bHasPositions = PendingPX && PendingPY && PendingPZ;
+			if (bHasPositions)
+			{
+				CachedPositions.SetNum(NumCurves);
+				for(int32 ParticleIndex = 0; ParticleIndex < NumCurves; ++ParticleIndex)
+				{
+					const int32 TransformIndex = EvaluatedResult.ParticleIndices[ParticleIndex];
+					CachedPositions[ParticleIndex] = FVector((*PendingPX)[TransformIndex], (*PendingPY)[TransformIndex], (*PendingPZ)[TransformIndex]);
+				}
+			}
+			TArray<FVector>* CachedPositionsPtr = bHasPositions ? &CachedPositions : nullptr;
+
+			TArray<FVector> CachedVelocities; 
+			const bool bHasVelocities = PendingVX && PendingVY && PendingVZ;
+			if (bHasVelocities)
+			{
+				CachedVelocities.SetNum(NumCurves);
+				for(int32 ParticleIndex = 0; ParticleIndex < NumCurves; ++ParticleIndex)
+				{
+					const int32 TransformIndex = EvaluatedResult.ParticleIndices[ParticleIndex];
+					CachedVelocities[ParticleIndex] = FVector((*PendingVX)[TransformIndex], (*PendingVY)[TransformIndex], (*PendingVZ)[TransformIndex]);
+				}
+			}
+			TArray<FVector>* CachedVelocitiesPtr = bHasVelocities ? &CachedVelocities : nullptr;
+
+			FClothSimulationProxy* Proxy = ClothComp->ClothSimulationProxy.Get();
+			Proxy->Tick_GameThread(0, CachedPositionsPtr, CachedVelocitiesPtr);
+			Proxy->CompleteParallelSimulation_GameThread();
+			ClothComp->MarkRenderDynamicDataDirty();
+			ClothComp->DoDeferredRenderUpdates_Concurrent();
+		}
+	}
+
+	bool FClothComponentCacheAdapter::InitializeForRecord(UPrimitiveComponent* InComponent, UChaosCache* InCache)
+	{
+		return true;
+	}
+
+	bool FClothComponentCacheAdapter::InitializeForPlayback(UPrimitiveComponent* InComponent, UChaosCache* InCache, float InTime)
+	{
+		::Chaos::EnsureIsInGameThreadContext();
+		return true;
+	}
+
+	UChaosClothComponent* FClothComponentCacheAdapter::GetClothComponent(UPrimitiveComponent* InComponent) const
+	{
+		if (InComponent)
+		{
+			if (UChaosClothComponent* ClothComp = Cast<UChaosClothComponent>(InComponent))
+			{
+				return ClothComp;
+			}
+			else if (USkeletalMeshComponent* SkelComp = Cast<USkeletalMeshComponent>(InComponent))
+			{
+				TArray<USceneComponent*> Children;
+				SkelComp->GetChildrenComponents(true, Children);
+				TArray<USceneComponent*> ClothComps = Children.FilterByPredicate([](USceneComponent* Child) { return Child->IsA<UChaosClothComponent>(); });
+				if (ClothComps.Num() > 1)
+				{
+					ensureMsgf(false, TEXT("Found more than one cloth component attached to a skeletal mesh component. This is not yet supported."));
+				}
+				else if (ClothComps.Num() == 1)
+				{
+					return Cast<UChaosClothComponent>(ClothComps[0]);
+				}
+			}
+		}
+		return nullptr;
+	}
+
+}    // namespace Chaos
diff --git a/Engine/Plugins/Experimental/ChaosClothAsset/Source/ChaosClothAssetEngine/Private/ChaosClothAsset/ClothComponentCacheAdapter.h b/Engine/Plugins/Experimental/ChaosClothAsset/Source/ChaosClothAssetEngine/Private/ChaosClothAsset/ClothComponentCacheAdapter.h
new file mode 100644
--- /dev/null
+++ b/Engine/Plugins/Experimental/ChaosClothAsset/Source/ChaosClothAssetEngine/Private/ChaosClothAsset/ClothComponentCacheAdapter.h
@@ -0,0 +1,55 @@
+// Copyright Epic Games, Inc. All Rights Reserved.
+
+#pragma once
+
+#include "Chaos/Adapters/CacheAdapter.h"
+
+class UChaosClothComponent;
+namespace Chaos
+{
+	class FClothingSimulationSolver;
+};
+
+namespace UE::Chaos::ClothAsset
+{	
+	/**
+	 * Skeletal mesh cache adapter to be able to cache cloth simulation datas through the chaos cache system
+	 */
+	class FClothComponentCacheAdapter : public ::Chaos::FComponentCacheAdapter
+	{
+	public:
+		
+		virtual ~FClothComponentCacheAdapter() = default;
+
+		// ~Begin FComponentCacheAdapter interface
+		virtual SupportType            SupportsComponentClass(UClass* InComponentClass) const override;
+		virtual UClass*                GetDesiredClass() const override;
+		virtual uint8                  GetPriority() const override;
+		virtual FGuid                  GetGuid() const override;
+		virtual bool                   ValidForPlayback(UPrimitiveComponent* InComponent, UChaosCache* InCache) const override;
+		virtual ::Chaos::FPhysicsSolver* GetComponentSolver(UPrimitiveComponent* InComponent) const override;
+		virtual ::Chaos::FPhysicsSolverEvents* BuildEventsSolver(UPrimitiveComponent* InComponent) const override;
+		virtual void				   SetRestState(UPrimitiveComponent* InComponent, UChaosCache* InCache, const FTransform& InRootTransform, ::Chaos::FReal InTime) const override;
+		virtual bool                   InitializeForRecord(UPrimitiveComponent* InComponent, UChaosCache* InCache) override;
+		virtual bool                   InitializeForPlayback(UPrimitiveComponent* InComponent, UChaosCache* InCache, float InTime) override;
+		virtual void                   Record_PostSolve(UPrimitiveComponent* InComp, const FTransform& InRootTransform, FPendingFrameWrite& OutFrame, ::Chaos::FReal InTime) const override;
+		virtual void                   Playback_PreSolve(UPrimitiveComponent* 				InComponent,
+										 UChaosCache*										InCache,
+										 ::Chaos::FReal										InTime,
+										 FPlaybackTickRecord&								TickRecord,
+										 TArray<::Chaos::TPBDRigidParticleHandle<::Chaos::FReal, 3>*>& OutUpdatedRigids) const override;
+		// ~End FComponentCacheAdapter interface
+
+	private :
+		/** Return the cloth component.
+		 */
+		UChaosClothComponent* GetClothComponent(UPrimitiveComponent* InComponent) const;
+
+		inline static const FName VelocityXName = TEXT("VelocityX");
+		inline static const FName VelocityYName = TEXT("VelocityY");
+		inline static const FName VelocityZName = TEXT("VelocityZ");
+		inline static const FName PositionXName = TEXT("PositionX");
+		inline static const FName PositionYName = TEXT("PositionY");
+		inline static const FName PositionZName = TEXT("PositionZ");
+	};
+}    // namespace Chaos
diff --git a/Engine/Plugins/Experimental/ChaosClothAsset/Source/ChaosClothAssetEngine/Private/ChaosClothAsset/ClothSimulationContext.cpp b/Engine/Plugins/Experimental/ChaosClothAsset/Source/ChaosClothAssetEngine/Private/ChaosClothAsset/ClothSimulationContext.cpp
--- a/Engine/Plugins/Experimental/ChaosClothAsset/Source/ChaosClothAssetEngine/Private/ChaosClothAsset/ClothSimulationContext.cpp
+++ b/Engine/Plugins/Experimental/ChaosClothAsset/Source/ChaosClothAssetEngine/Private/ChaosClothAsset/ClothSimulationContext.cpp
@@ -9,7 +9,7 @@
 
 namespace UE::Chaos::ClothAsset
 {
-	void FClothSimulationContext::Fill(const UChaosClothComponent& ClothComponent, float InDeltaTime, float MaxDeltaTime, bool bIsInitialization)
+	void FClothSimulationContext::Fill(const UChaosClothComponent& ClothComponent, float InDeltaTime, float MaxDeltaTime, bool bIsInitialization, const TArray<FVector>* InCachedPositions, const TArray<FVector>* InCachedVelocities)
 	{
 		// Set the time
 		DeltaTime = FMath::Min(InDeltaTime, MaxDeltaTime);
@@ -128,5 +128,23 @@ namespace UE::Chaos::ClothAsset
 		{
 			WindVelocity = FVector::ZeroVector;
 		}
+
+		if (InCachedPositions != nullptr)
+		{
+			CachedPositions = *InCachedPositions;
+		}
+		else
+		{
+			CachedPositions.Reset();
+		}
+
+		if (InCachedVelocities != nullptr)
+		{
+			CachedVelocities = *InCachedVelocities;
+		}
+		else
+		{
+			CachedVelocities.Reset();
+		}
 	}
 }  // End namespace UE::Chaos::ClothAsset
diff --git a/Engine/Plugins/Experimental/ChaosClothAsset/Source/ChaosClothAssetEngine/Private/ChaosClothAsset/ClothSimulationContext.h b/Engine/Plugins/Experimental/ChaosClothAsset/Source/ChaosClothAssetEngine/Private/ChaosClothAsset/ClothSimulationContext.h
--- a/Engine/Plugins/Experimental/ChaosClothAsset/Source/ChaosClothAssetEngine/Private/ChaosClothAsset/ClothSimulationContext.h
+++ b/Engine/Plugins/Experimental/ChaosClothAsset/Source/ChaosClothAssetEngine/Private/ChaosClothAsset/ClothSimulationContext.h
@@ -47,7 +47,13 @@ namespace UE::Chaos::ClothAsset
 		/** Wind velocity at the component location. */
 		FVector WindVelocity;
 
+		/** Positions used by ChaosCache */
+		TArray<FVector> CachedPositions;
+
+		/** Velocities used by ChaosCache */
+		TArray<FVector> CachedVelocities;
+
 		/** Fill the context from data collected from the specified component. */
-		void Fill(const UChaosClothComponent& ClothComponent, float InDeltaTime, float MaxDeltaTime, bool bIsInitialization = false);
+		void Fill(const UChaosClothComponent& ClothComponent, float InDeltaTime, float MaxDeltaTime, bool bIsInitialization = false, const TArray<FVector>* InCachedPositions = nullptr, const TArray<FVector>* InCachedVelocities = nullptr);
 	};
 }
diff --git a/Engine/Plugins/Experimental/ChaosClothAsset/Source/ChaosClothAssetEngine/Private/ChaosClothAsset/ClothSimulationProxy.cpp b/Engine/Plugins/Experimental/ChaosClothAsset/Source/ChaosClothAssetEngine/Private/ChaosClothAsset/ClothSimulationProxy.cpp
--- a/Engine/Plugins/Experimental/ChaosClothAsset/Source/ChaosClothAssetEngine/Private/ChaosClothAsset/ClothSimulationProxy.cpp
+++ b/Engine/Plugins/Experimental/ChaosClothAsset/Source/ChaosClothAssetEngine/Private/ChaosClothAsset/ClothSimulationProxy.cpp
@@ -148,14 +148,21 @@ PRAGMA_ENABLE_DEPRECATION_WARNINGS
 		CompleteParallelSimulation_GameThread();
 	}
 
-	void FClothSimulationProxy::Tick_GameThread(float DeltaTime)
+	void FClothSimulationProxy::Tick_GameThread(float DeltaTime, const TArray<FVector>* CachedPositions, const TArray<FVector>* CachedVelocities)
 	{
 		SCOPE_CYCLE_COUNTER(STAT_ClothSimulationProxy_TickGame);
 
 		// Fill a new context, note the context is also needed when the simulation is suspended
-		ClothSimulationContext.Fill(ClothComponent, DeltaTime, MaxDeltaTime);
+		ClothSimulationContext.Fill(ClothComponent, DeltaTime, MaxDeltaTime, false, CachedPositions, CachedVelocities);
 
-		if (DeltaTime > 0.f && !ClothComponent.IsSimulationSuspended())
+		const bool bUseCache = ClothSimulationContext.CachedPositions.Num() > 0 || ClothSimulationContext.CachedVelocities.Num() > 0;
+		if (bUseCache)
+		{
+			Solver->SetEnableSolver(false);
+		}
+
+		const bool bCreateParallelTask = (DeltaTime > 0.f && !ClothComponent.IsSimulationSuspended()) || bUseCache;
+		if (bCreateParallelTask)
 		{
 			// Replace physics thread's configs with the game thread's configs
 			for (const TUniquePtr<::Chaos::FClothingSimulationConfig>& Config : Configs)
@@ -179,7 +186,9 @@ PRAGMA_ENABLE_DEPRECATION_WARNINGS
 
 		TRACE_CPUPROFILER_EVENT_SCOPE(FClothSimulationProxy_TickPhysics);
 		SCOPE_CYCLE_COUNTER(STAT_ClothSimulationProxy_TickPhysics);
-		if (ClothSimulationContext.DeltaTime == 0.f)
+		const bool bUseCache = ClothSimulationContext.CachedPositions.Num() > 0 || ClothSimulationContext.CachedVelocities.Num() > 0;
+
+		if (ClothSimulationContext.DeltaTime == 0.f && !bUseCache)
 		{
 			return;
 		}
@@ -214,7 +223,14 @@ PRAGMA_ENABLE_DEPRECATION_WARNINGS
 		}
 
 		// Step the simulation
-		Solver->Update((Softs::FSolverReal)ClothSimulationContext.DeltaTime);
+		if (Solver->GetEnableSolver() || !bUseCache)
+		{
+			Solver->Update((Softs::FSolverReal)ClothSimulationContext.DeltaTime);
+		}
+		else
+		{
+			Solver->UpdateFromCache(ClothSimulationContext.CachedPositions, ClothSimulationContext.CachedVelocities);
+		}
 
 		// Keep the actual used number of iterations for the stats
 		NumIterations = Solver->GetNumUsedIterations();
diff --git a/Engine/Plugins/Experimental/ChaosClothAsset/Source/ChaosClothAssetEngine/Private/ChaosClothAsset/ClothSimulationProxy.h b/Engine/Plugins/Experimental/ChaosClothAsset/Source/ChaosClothAssetEngine/Private/ChaosClothAsset/ClothSimulationProxy.h
--- a/Engine/Plugins/Experimental/ChaosClothAsset/Source/ChaosClothAssetEngine/Private/ChaosClothAsset/ClothSimulationProxy.h
+++ b/Engine/Plugins/Experimental/ChaosClothAsset/Source/ChaosClothAssetEngine/Private/ChaosClothAsset/ClothSimulationProxy.h
@@ -42,7 +42,7 @@ namespace UE::Chaos::ClothAsset
 		FClothSimulationProxy& operator=(const FClothSimulationProxy&) = delete;
 		FClothSimulationProxy& operator=(FClothSimulationProxy&&) = delete;
 
-		void Tick_GameThread(float DeltaTime);
+		void Tick_GameThread(float DeltaTime, const TArray<FVector>* CachedPositions = nullptr, const TArray<FVector>* CachedVelocities = nullptr);
 
 		/** Wait for the parallel task to complete if one was running, and update the simulation data. */
 		void CompleteParallelSimulation_GameThread();
diff --git a/Engine/Plugins/Experimental/ChaosClothAsset/Source/ChaosClothAssetEngine/Public/ChaosClothAsset/ClothComponent.h b/Engine/Plugins/Experimental/ChaosClothAsset/Source/ChaosClothAssetEngine/Public/ChaosClothAsset/ClothComponent.h
--- a/Engine/Plugins/Experimental/ChaosClothAsset/Source/ChaosClothAssetEngine/Public/ChaosClothAsset/ClothComponent.h
+++ b/Engine/Plugins/Experimental/ChaosClothAsset/Source/ChaosClothAssetEngine/Public/ChaosClothAsset/ClothComponent.h
@@ -17,6 +17,7 @@ namespace Chaos::Softs
 namespace UE::Chaos::ClothAsset
 {
 	class FClothSimulationProxy;
+	class FClothComponentCacheAdapter;
 }
 
 /**
@@ -29,7 +30,7 @@ UCLASS(
 	HideCategories = (Object, "Mesh|SkeletalAsset", Constraints, Advanced, Cooking, Collision, Navigation))
 class CHAOSCLOTHASSETENGINE_API UChaosClothComponent : public USkinnedMeshComponent
 {
-	GENERATED_BODY()
+	GENERATED_BODY()	
 public:
 	UChaosClothComponent(const FObjectInitializer& ObjectInitializer);
 	UChaosClothComponent(FVTableHelper& Helper);
@@ -84,6 +85,8 @@ public:
 	 */
 	TSharedPtr<const FManagedArrayCollection> GetPropertyCollection() const { return TSharedPtr<const FManagedArrayCollection>(PropertyCollection); }
 
+	friend UE::Chaos::ClothAsset::FClothComponentCacheAdapter;
+
 protected:
 	//~ Begin UObject Interface
 	virtual void PostLoad() override;
