ChaosCache: Reverse hack to fix BestFitAdapter for skeletalmeshcomponents

diff --git a/Engine/Plugins/Experimental/ChaosCloth/Source/Chaos/Private/ChaosCloth/SkeletalMeshComponentCacheAdapter.cpp b/Engine/Plugins/Experimental/ChaosCloth/Source/Chaos/Private/ChaosCloth/SkeletalMeshComponentCacheAdapter.cpp
--- a/Engine/Plugins/Experimental/ChaosCloth/Source/Chaos/Private/ChaosCloth/SkeletalMeshComponentCacheAdapter.cpp
+++ b/Engine/Plugins/Experimental/ChaosCloth/Source/Chaos/Private/ChaosCloth/SkeletalMeshComponentCacheAdapter.cpp
@@ -33,7 +33,7 @@ namespace Chaos
 
 	uint8 FSkeletalMeshCacheAdapter::GetPriority() const
 	{
-		return EngineAdapterPriorityBegin + 1; // A trick to prioritize cloth component over skeletal mesh component
+		return EngineAdapterPriorityBegin;
 	}
 
 	void FSkeletalMeshCacheAdapter::Record_PostSolve(UPrimitiveComponent* InComponent, const FTransform& InRootTransform, FPendingFrameWrite& OutFrame, Chaos::FReal InTime) const
diff --git a/Engine/Plugins/Experimental/ChaosClothAsset/Source/ChaosClothAssetEngine/Private/ChaosClothAsset/ClothComponentCacheAdapter.cpp b/Engine/Plugins/Experimental/ChaosClothAsset/Source/ChaosClothAssetEngine/Private/ChaosClothAsset/ClothComponentCacheAdapter.cpp
--- a/Engine/Plugins/Experimental/ChaosClothAsset/Source/ChaosClothAssetEngine/Private/ChaosClothAsset/ClothComponentCacheAdapter.cpp
+++ b/Engine/Plugins/Experimental/ChaosClothAsset/Source/ChaosClothAssetEngine/Private/ChaosClothAsset/ClothComponentCacheAdapter.cpp
@@ -15,18 +15,12 @@ namespace UE::Chaos::ClothAsset
 {
 	::Chaos::FComponentCacheAdapter::SupportType FClothComponentCacheAdapter::SupportsComponentClass(UClass* InComponentClass) const
 	{
-		/** The desired class is UChaosClothComponent. 
-		 * But when UChaosClothComponent is a child component of a USkeletalMeshComponent, UChaosClothComponent is not selectible using SoftComponentRef. 
-		 * To get around this issue, we need to select the parent USkeletalMeshComponent in SoftComponentRef, 
-		 * make USkeletalMeshComponent also a supported class for FClothComponentAdapter,
-		 *  and prioritize FClothComponentAdapter over FSkeletalMeshComponentAdapter in CacheManagerActor.
-		 */
 		const UClass* Desired = GetDesiredClass();
-		if(InComponentClass == Desired || InComponentClass == USkeletalMeshComponent::StaticClass())
+		if(InComponentClass == Desired)
 		{
 			return ::Chaos::FComponentCacheAdapter::SupportType::Direct;
 		}
-		else if(InComponentClass->IsChildOf(Desired) || InComponentClass->IsChildOf(USkeletalMeshComponent::StaticClass()))
+		else if(InComponentClass->IsChildOf(Desired))
 		{
 			return ::Chaos::FComponentCacheAdapter::SupportType::Derived;
 		}
