Chaos Cloth - Move the Cloth Simulation Proxy header to the public folder, so that it can be used and inherited by the ML tool.

#rnx

diff --git a/Engine/Plugins/Experimental/ChaosClothAsset/Source/ChaosClothAssetEngine/Private/ChaosClothAsset/ClothSimulationProxy.h b/Engine/Plugins/Experimental/ChaosClothAsset/Source/ChaosClothAssetEngine/Private/ChaosClothAsset/ClothSimulationProxy.h
deleted file mode 100644
diff --git a/Engine/Plugins/Experimental/ChaosClothAsset/Source/ChaosClothAssetEngine/Public/ChaosClothAsset/ClothSimulationProxy.h b/Engine/Plugins/Experimental/ChaosClothAsset/Source/ChaosClothAssetEngine/Public/ChaosClothAsset/ClothSimulationProxy.h
new file mode 100644
--- /dev/null
+++ b/Engine/Plugins/Experimental/ChaosClothAsset/Source/ChaosClothAssetEngine/Public/ChaosClothAsset/ClothSimulationProxy.h
@@ -0,0 +1,117 @@
+// Copyright Epic Games, Inc. All Rights Reserved.
+
+#pragma once
+
+#include "ChaosClothAsset/ClothSimulationContext.h"
+#include "Async/TaskGraphInterfaces.h"
+#include "Containers/Array.h"
+#include "Math/Transform.h"
+#include "Templates/UniquePtr.h"
+#include "ClothingSystemRuntimeTypes.h"
+
+namespace Chaos
+{
+	class FClothingSimulationSolver;
+	class FClothingSimulationMesh;
+	class FClothingSimulationCloth;
+	class FClothingSimulationConfig;
+	class FClothingSimulationCollider;
+	class FClothVisualization;
+}
+
+class UChaosClothComponent;
+class UChaosClothAsset;
+struct FChaosClothSimulationModel;
+struct FReferenceSkeleton;
+
+namespace UE::Chaos::ClothAsset
+{
+	/**
+	 * Cloth simulation proxy.
+	 * Class used to share data between the cloth simulation and the cloth component.
+	 */
+	class FClothSimulationProxy final
+	{
+	public:
+		explicit FClothSimulationProxy(const UChaosClothComponent& InClothComponent);
+		~FClothSimulationProxy();
+
+		FClothSimulationProxy() = delete;  // This object cannot be created without a valid reference to a parent UChaosClothComponent 
+		FClothSimulationProxy(const FClothSimulationProxy&) = delete;  // Disable the copy as there must be a single unique proxy per component
+		FClothSimulationProxy(FClothSimulationProxy&&) = delete;  // Disable the move to force it to be associated with a valid component reference
+		FClothSimulationProxy& operator=(const FClothSimulationProxy&) = delete;
+		FClothSimulationProxy& operator=(FClothSimulationProxy&&) = delete;
+
+		void Tick_GameThread(float DeltaTime, const TArray<FVector>* CachedPositions = nullptr, const TArray<FVector>* CachedVelocities = nullptr);
+
+		/** Wait for the parallel task to complete if one was running, and update the simulation data. */
+		void CompleteParallelSimulation_GameThread();
+
+		/**
+		 * Return a map of all simulation data as used by the skeletal rendering code.
+		 * The map key is the rendering section's cloth index as set in FSkelMeshRenderSection::CorrespondClothAssetIndex,
+		 * which is 0 for the entire cloth component since all of its sections share the same simulation data.
+		 */
+		const TMap<int32, FClothSimulData>& GetCurrentSimulationData_AnyThread() const;
+
+		FBoxSphereBounds CalculateBounds_AnyThread() const;
+
+	private:
+		void Tick_PhysicsThread();
+
+		void WriteSimulationData_GameThread();
+
+		// Internal physics thread object
+		friend class FClothSimulationProxyParallelTask;
+
+		// Reference for the cloth parallel task, to detect whether or not a simulation is running
+		FGraphEventRef ParallelTask;
+
+		// Simulation data written back to the component after the simulation has taken place
+		TMap<int32, FClothSimulData> CurrentSimulationData;
+
+		// Owner component
+		const UChaosClothComponent& ClothComponent;
+
+		// Simulation context used to store the required component data for the duration of the simulation
+		FClothSimulationContext ClothSimulationContext;
+
+		// The cloth simulation model used to create this simulation, ownership might get transferred to this proxy if it changes during the simulation
+		TSharedPtr<const FChaosClothSimulationModel> ClothSimulationModel;
+
+		// Simulation objects
+		TUniquePtr<::Chaos::FClothingSimulationSolver> Solver;
+PRAGMA_DISABLE_DEPRECATION_WARNINGS  // TODO: CHAOS_IS_CLOTHINGSIMULATIONMESH_ABSTRACT
+		TArray<TUniquePtr<::Chaos::FClothingSimulationMesh>> Meshes;
+PRAGMA_ENABLE_DEPRECATION_WARNINGS
+		TArray<TUniquePtr<::Chaos::FClothingSimulationCloth>> Cloths;
+		TArray<TUniquePtr<::Chaos::FClothingSimulationConfig>> Configs;
+		TArray<TUniquePtr<::Chaos::FClothingSimulationCollider>> Colliders;
+		TUniquePtr<::Chaos::FClothVisualization> Visualization;
+
+		// Properties that must be readable from all threads
+		std::atomic<int32> NumCloths;
+		std::atomic<int32> NumKinematicParticles;
+		std::atomic<int32> NumDynamicParticles;
+		std::atomic<int32> NumIterations;
+		std::atomic<int32> NumSubsteps;
+		std::atomic<float> SimulationTime;
+		std::atomic<bool> bIsTeleported;
+
+		mutable bool bHasInvalidReferenceBoneTransforms = false;
+
+		// Cached value of the MaxPhysicsDeltaTime setting for the life of this proxy
+		const float MaxDeltaTime;
+	};
+
+#if !defined(CHAOS_TRANSFORM_CLOTH_SIMUL_DATA_ISPC_ENABLED_DEFAULT)
+#define CHAOS_TRANSFORM_CLOTH_SIMUL_DATA_ISPC_ENABLED_DEFAULT 1
+#endif
+
+	// Support run-time toggling on supported platforms in non-shipping configurations
+#if !INTEL_ISPC || UE_BUILD_SHIPPING
+	static constexpr bool bTransformClothSimulData_ISPC_Enabled = INTEL_ISPC && CHAOS_TRANSFORM_CLOTH_SIMUL_DATA_ISPC_ENABLED_DEFAULT;
+#else
+	extern bool bTransformClothSimulData_ISPC_Enabled;
+#endif
+}
