Chaos Cloth:
- Fixed Cloth Asset crashes with empty LOD.
- Tidied up ChaosClothAsset\ClothCollection.h and moved one template definition over to the cpp file.

#rnx

diff --git a/Engine/Plugins/Experimental/ChaosCloth/Source/Chaos/Private/ChaosCloth/ChaosClothingSimulationCloth.cpp b/Engine/Plugins/Experimental/ChaosCloth/Source/Chaos/Private/ChaosCloth/ChaosClothingSimulationCloth.cpp
--- a/Engine/Plugins/Experimental/ChaosCloth/Source/Chaos/Private/ChaosCloth/ChaosClothingSimulationCloth.cpp
+++ b/Engine/Plugins/Experimental/ChaosCloth/Source/Chaos/Private/ChaosCloth/ChaosClothingSimulationCloth.cpp
@@ -94,11 +94,6 @@ void FClothingSimulationCloth::FLODData::Add(FClothingSimulationSolver* Solver,
 	check(Cloth);
 	check(Cloth->Mesh);
 
-	if (!NumParticles)
-	{
-		return;
-	}
-
 	// Add a new solver data chunk
 	check(!SolverData.Find(Solver));
 	FSolverData& SolverDatum = SolverData.Add(Solver);
@@ -108,6 +103,11 @@ void FClothingSimulationCloth::FLODData::Add(FClothingSimulationSolver* Solver,
 	// Add particles
 	Offset = Solver->AddParticles(NumParticles, Cloth->GroupId);  // TODO: Have a per solver map of offset
 
+	if (!NumParticles)
+	{
+		return;
+	}
+
 PRAGMA_DISABLE_DEPRECATION_WARNINGS  // TODO: CHAOS_IS_CLOTHINGSIMULATIONMESH_ABSTRACT
 	// Update source mesh for this LOD, this is required prior to reset the start pose
 	Cloth->Mesh->Update(Solver, INDEX_NONE, InLODIndex, 0, Offset);
@@ -245,14 +245,15 @@ void FClothingSimulationCloth::FLODData::UpdateNormals(FClothingSimulationSolver
 	const int32 Offset = SolverDatum.Offset;
 	const FTriangleMesh& TriangleMesh = SolverDatum.TriangleMesh;
 
-	check(Offset != INDEX_NONE);
-
-	TConstArrayView<Softs::FSolverVec3> Points(Solver->GetParticleXs(Offset) - Offset, Offset + NumParticles);  // TODO: TriangleMesh still uses global array
-	TArray<Softs::FSolverVec3> FaceNormals;
-	TriangleMesh.GetFaceNormals(FaceNormals, Points, /*ReturnEmptyOnError =*/ false);
+	if (Offset != INDEX_NONE)
+	{
+		TConstArrayView<Softs::FSolverVec3> Points(Solver->GetParticleXs(Offset) - Offset, Offset + NumParticles);  // TODO: TriangleMesh still uses global array
+		TArray<Softs::FSolverVec3> FaceNormals;
+		TriangleMesh.GetFaceNormals(FaceNormals, Points, /*ReturnEmptyOnError =*/ false);
 
-	TArrayView<Softs::FSolverVec3> Normals(Solver->GetNormals(Offset), NumParticles);
-	TriangleMesh.GetPointNormals(Normals, TConstArrayView<Softs::FSolverVec3>(FaceNormals), /*bUseGlobalArray =*/ false);
+		TArrayView<Softs::FSolverVec3> Normals(Solver->GetNormals(Offset), NumParticles);
+		TriangleMesh.GetPointNormals(Normals, TConstArrayView<Softs::FSolverVec3>(FaceNormals), /*bUseGlobalArray =*/ false);
+	}
 }
 
 FClothingSimulationCloth::FClothingSimulationCloth(
@@ -813,7 +814,10 @@ PRAGMA_ENABLE_DEPRECATION_WARNINGS
 	// Update the source mesh skinned positions
 	const int32 PrevOffset = GetOffset(Solver, PrevLODIndex);
 	const int32 Offset = GetOffset(Solver, LODIndex);
-	check(PrevOffset != INDEX_NONE && Offset != INDEX_NONE);
+	if (PrevOffset == INDEX_NONE || Offset == INDEX_NONE)
+	{
+		return;
+	}
 
 PRAGMA_DISABLE_DEPRECATION_WARNINGS  // TODO: CHAOS_IS_CLOTHINGSIMULATIONMESH_ABSTRACT
 	Mesh->Update(Solver, PrevLODIndex, LODIndex, PrevOffset, Offset);
diff --git a/Engine/Plugins/Experimental/ChaosClothAsset/Source/ChaosClothAsset/Private/ChaosClothAsset/ClothCollection.cpp b/Engine/Plugins/Experimental/ChaosClothAsset/Source/ChaosClothAsset/Private/ChaosClothAsset/ClothCollection.cpp
--- a/Engine/Plugins/Experimental/ChaosClothAsset/Source/ChaosClothAsset/Private/ChaosClothAsset/ClothCollection.cpp
+++ b/Engine/Plugins/Experimental/ChaosClothAsset/Source/ChaosClothAsset/Private/ChaosClothAsset/ClothCollection.cpp
@@ -95,7 +95,7 @@ namespace UE::Chaos::ClothAsset
 		AddExternalAttribute<int32>("TetherStart", TetherBatchesGroup, TetherStart, TethersDependency);
 		AddExternalAttribute<int32>("TetherEnd", TetherBatchesGroup, TetherEnd, TethersDependency);
 
-		// LOD Group
+		// LODs Group
 		AddExternalAttribute<int32>("PatternStart", LodsGroup, PatternStart, PatternsDependency);
 		AddExternalAttribute<int32>("PatternEnd", LodsGroup, PatternEnd, PatternsDependency);
 		AddExternalAttribute<int32>("StitchingStart", LodsGroup, SeamStart, StitchingsDependency);
@@ -181,4 +181,33 @@ namespace UE::Chaos::ClothAsset
 
 		return Start == INDEX_NONE ? 0 : End - Start + 1;
 	}
+
+	template<bool bStart, bool bEnd>
+	TTuple<int32, int32> FClothCollection::GetPatternsElementsStartEnd(const TManagedArray<int32>& StartArray, const TManagedArray<int32>& EndArray, int32 LodIndex) const
+	{
+		const int32 LodPatternStart = PatternStart[LodIndex];
+		const int32 LodPatternEnd = PatternEnd[LodIndex];
+
+		int32 Start = INDEX_NONE;  // Find Start and End indices for the entire LOD minding empty patterns on the way
+		int32 End = INDEX_NONE;
+
+		if (LodPatternStart != INDEX_NONE && LodPatternEnd != INDEX_NONE)
+		{
+			for (int32 PatternIndex = LodPatternStart; PatternIndex <= LodPatternEnd; ++PatternIndex)
+			{
+				if (bStart && StartArray[PatternIndex] != INDEX_NONE)
+				{
+					Start = (Start == INDEX_NONE) ? StartArray[PatternIndex] : FMath::Min(Start, StartArray[PatternIndex]);
+				}
+				if (bEnd && EndArray[PatternIndex] != INDEX_NONE)
+				{
+					End = (End == INDEX_NONE) ? EndArray[PatternIndex] : FMath::Max(End, EndArray[PatternIndex]);
+				}
+			}
+		}
+		return TTuple<int32, int32>(Start, End);
+	}
+	template CHAOSCLOTHASSET_API TTuple<int32, int32> FClothCollection::GetPatternsElementsStartEnd<true, false>(const TManagedArray<int32>& StartArray, const TManagedArray<int32>& EndArray, int32 LodIndex) const;
+	template CHAOSCLOTHASSET_API TTuple<int32, int32> FClothCollection::GetPatternsElementsStartEnd<false, true>(const TManagedArray<int32>& StartArray, const TManagedArray<int32>& EndArray, int32 LodIndex) const;
+	template CHAOSCLOTHASSET_API TTuple<int32, int32> FClothCollection::GetPatternsElementsStartEnd<true, true>(const TManagedArray<int32>& StartArray, const TManagedArray<int32>& EndArray, int32 LodIndex) const;
 } // End namespace UE::Chaos::ClothAsset
diff --git a/Engine/Plugins/Experimental/ChaosClothAsset/Source/ChaosClothAsset/Public/ChaosClothAsset/ClothCollection.h b/Engine/Plugins/Experimental/ChaosClothAsset/Source/ChaosClothAsset/Public/ChaosClothAsset/ClothCollection.h
--- a/Engine/Plugins/Experimental/ChaosClothAsset/Source/ChaosClothAsset/Public/ChaosClothAsset/ClothCollection.h
+++ b/Engine/Plugins/Experimental/ChaosClothAsset/Source/ChaosClothAsset/Public/ChaosClothAsset/ClothCollection.h
@@ -5,8 +5,7 @@
 #include "GeometryCollection/ManagedArrayCollection.h"
 
 // TODO:
-// - Transition Sim LOD data
-// - Move wrap deformer into struct similar to FMeshToMeshVertData
+// - Transition Sim LOD data?
 
 namespace UE::Chaos::ClothAsset
 {
@@ -24,11 +23,13 @@ namespace UE::Chaos::ClothAsset
 		FClothCollection(FClothCollection&&) = default;
 		FClothCollection& operator=(FClothCollection&&) = default;
 
+		//~ Begin FManagedArrayCollection interface
+		virtual void Reset() override { Super::Reset(); Construct(); }
+		//~ Begin FManagedArrayCollection interface
+
 		using Super::Serialize;
 		void Serialize(FArchive& Ar);
 
-		virtual void Reset() override { Super::Reset(); Construct(); }
-
 		/** Set the number of elements to one of the groups that have start/end indices while maintaining the correct order of the data. */
 		int32 SetNumElements(int32 InNumElements, const FName& GroupName, TManagedArray<int32>& StartArray, TManagedArray<int32>& EndArray, int32 StartEndIndex);
 
@@ -48,7 +49,7 @@ namespace UE::Chaos::ClothAsset
 		inline TArrayView<T> GetPatternsElements(TManagedArray<T>& ElementArray, const TManagedArray<int32>& StartArray, const TManagedArray<int32>& EndArray, int32 LodIndex);
 
 		template<bool bStart = true, bool bEnd = true>
-		inline TTuple<int32, int32> GetPatternsElementsStartEnd(const TManagedArray<int32>& StartArray, const TManagedArray<int32>& EndArray, int32 StartEndIndex) const;
+		TTuple<int32, int32> GetPatternsElementsStartEnd(const TManagedArray<int32>& StartArray, const TManagedArray<int32>& EndArray, int32 StartEndIndex) const;
 
 		// Attribute groups, predefined data member of the collection
 		static const FName SimVerticesGroup;  // Contains patterns' 2D positions, 3D draped position (rest)
@@ -88,14 +89,6 @@ namespace UE::Chaos::ClothAsset
 		TManagedArray<FIntVector3> RenderIndices;  // The indices point to the elements in the Render Vertices arrays but don't include the LOD start offset
 		TManagedArray<int32> RenderMaterialIndex;  // Render material per triangle
 
-		// TODO: FMeshToMeshVertData
-		//// Wrap Deformers Group (render mesh captured, sim mesh capturing)
-		//TManagedArray<FVector4f> CapturedPosition;  // Barycentric coordinate and distance for captured render positions
-		//TManagedArray<FVector4f> CapturedNormal;  // Barycentric coordinate and distance for captured render normals
-		//TManagedArray<FVector4f> CapturedTangent;  // Barycentric coordinate and distance for captured render tangents
-		//TManagedArray<int32> CapturingFace;  // Sim triangle index
-		//TManagedArray<float> CapturingWeight;  // Weight for when using multiple capture influence per point, 1.0 otherwise
-
 		// Patterns Group
 		TManagedArray<int32> SimVerticesStart;
 		TManagedArray<int32> SimVerticesEnd;
@@ -124,7 +117,7 @@ namespace UE::Chaos::ClothAsset
 		TManagedArray<int32> TetherStart;
 		TManagedArray<int32> TetherEnd;
 
-		// LOD Group
+		// LODs Group
 		TManagedArray<int32> PatternStart;
 		TManagedArray<int32> PatternEnd;
 		TManagedArray<int32> SeamStart;
@@ -178,26 +171,4 @@ namespace UE::Chaos::ClothAsset
 		check(Start != INDEX_NONE || End == INDEX_NONE);  // Best to avoid situations where only one boundary of the range is set to INDEX_NONE
 		return Start == INDEX_NONE ? TArrayView<T>() : TArrayView<T>(ElementArray.GetData() + Start, End - Start + 1);
 	}
-
-	template<bool bStart, bool bEnd>
-	inline TTuple<int32, int32> FClothCollection::GetPatternsElementsStartEnd(const TManagedArray<int32>& StartArray, const TManagedArray<int32>& EndArray, int32 LodIndex) const
-	{
-		const int32 LodPatternStart = PatternStart[LodIndex];
-		const int32 LodPatternEnd = PatternEnd[LodIndex];
-
-		int32 Start = INDEX_NONE;  // Find Start and End indices for the entire LOD minding empty patterns on the way
-		int32 End = INDEX_NONE;
-		for (int32 PatternIndex = LodPatternStart; PatternIndex <= LodPatternEnd; ++PatternIndex)
-		{
-			if (bStart && StartArray[PatternIndex] != INDEX_NONE)
-			{
-				Start = (Start == INDEX_NONE) ? StartArray[PatternIndex] : FMath::Min(Start, StartArray[PatternIndex]);
-			}
-			if (bEnd && EndArray[PatternIndex] != INDEX_NONE)
-			{
-				End = (End == INDEX_NONE) ? EndArray[PatternIndex] : FMath::Max(End, EndArray[PatternIndex]);
-			}
-		}
-		return TTuple<int32, int32>(Start, End);
-	}
 }  // End namespace UE::Chaos::ClothAsset
diff --git a/Engine/Plugins/Experimental/ChaosClothAsset/Source/ChaosClothAssetEngine/Private/ChaosClothAsset/ClothAsset.cpp b/Engine/Plugins/Experimental/ChaosClothAsset/Source/ChaosClothAssetEngine/Private/ChaosClothAsset/ClothAsset.cpp
--- a/Engine/Plugins/Experimental/ChaosClothAsset/Source/ChaosClothAssetEngine/Private/ChaosClothAsset/ClothAsset.cpp
+++ b/Engine/Plugins/Experimental/ChaosClothAsset/Source/ChaosClothAssetEngine/Private/ChaosClothAsset/ClothAsset.cpp
@@ -298,10 +298,12 @@ void UChaosClothAsset::CalculateBounds()
 
 			const int32 RenderVerticesStart = ClothCollection->RenderVerticesStart[PatternElementIndex];
 			const int32 RenderVerticesEnd = ClothCollection->RenderVerticesEnd[PatternElementIndex];
-
-			for (int32 RenderVertexIndex = RenderVerticesStart; RenderVertexIndex <= RenderVerticesEnd; ++RenderVertexIndex)
+			if (RenderVerticesStart != INDEX_NONE && RenderVerticesStart != INDEX_NONE)
 			{
-				BoundingBox += (FVector)ClothCollection->RenderPosition[RenderVertexIndex];
+				for (int32 RenderVertexIndex = RenderVerticesStart; RenderVertexIndex <= RenderVerticesEnd; ++RenderVertexIndex)
+				{
+					BoundingBox += (FVector)ClothCollection->RenderPosition[RenderVertexIndex];
+				}
 			}
 		}
 	}
@@ -346,7 +348,7 @@ void UChaosClothAsset::Build()
 
 	// Add LODs to the render data
 	const FClothConstAdapter Cloth(ClothCollection);
-	const int32 NumLods = Cloth.GetNumLods();
+	const int32 NumLods = FMath::Max(Cloth.GetNumLods(), 1);  // The render data will always look for at least one default LOD 0
 
 	// Rebuild LOD Infos
 	LODInfo.Reset(NumLods);
diff --git a/Engine/Plugins/Experimental/ChaosClothAssetEditor/Source/ChaosClothAssetTools/Private/ChaosClothAsset/ClothAssetBuilderEditor.cpp b/Engine/Plugins/Experimental/ChaosClothAssetEditor/Source/ChaosClothAssetTools/Private/ChaosClothAsset/ClothAssetBuilderEditor.cpp
--- a/Engine/Plugins/Experimental/ChaosClothAssetEditor/Source/ChaosClothAssetTools/Private/ChaosClothAsset/ClothAssetBuilderEditor.cpp
+++ b/Engine/Plugins/Experimental/ChaosClothAssetEditor/Source/ChaosClothAssetTools/Private/ChaosClothAsset/ClothAssetBuilderEditor.cpp
@@ -32,31 +32,37 @@ void UClothAssetBuilderEditor::BuildLod(FSkeletalMeshLODModel& LODModel, const U
 	TMap<int32, TArray<int32>> SectionFacesMap;
 	SectionFacesMap.Reserve(ClothAsset.GetMaterials().Num());
 
+	TArray<uint32> LodRenderIndexRemap;
+
 	const TSharedPtr<const UE::Chaos::ClothAsset::FClothCollection> ClothCollection = ClothAsset.GetClothCollection();
 	const int32 PatternStart = ClothCollection->PatternStart[LodIndex];
 	const int32 PatternEnd = ClothCollection->PatternEnd[LodIndex];
-
-	for (int32 PatternIndex = PatternStart; PatternIndex <= PatternEnd; ++PatternIndex)
+	if (PatternStart != INDEX_NONE && PatternEnd != INDEX_NONE)
 	{
-		const int32 RenderFacesStart = ClothCollection->RenderFacesStart[PatternIndex];
-		const int32 RenderFacesEnd = ClothCollection->RenderFacesEnd[PatternIndex];
-
-		for (int32 RenderFaceIndex = RenderFacesStart; RenderFaceIndex <= RenderFacesEnd; ++RenderFaceIndex)
+		for (int32 PatternIndex = PatternStart; PatternIndex <= PatternEnd; ++PatternIndex)
 		{
-			const int32 MaterialIndex = ClothCollection->RenderMaterialIndex[RenderFaceIndex];
+			const int32 RenderFacesStart = ClothCollection->RenderFacesStart[PatternIndex];
+			const int32 RenderFacesEnd = ClothCollection->RenderFacesEnd[PatternIndex];
+
+			if (RenderFacesStart != INDEX_NONE && RenderFacesEnd != INDEX_NONE)
+			{
+				for (int32 RenderFaceIndex = RenderFacesStart; RenderFaceIndex <= RenderFacesEnd; ++RenderFaceIndex)
+				{
+					const int32 MaterialIndex = ClothCollection->RenderMaterialIndex[RenderFaceIndex];
 
-			TArray<int32>& SectionFaces = SectionFacesMap.FindOrAdd(MaterialIndex);
-			SectionFaces.Add(RenderFaceIndex);
+					TArray<int32>& SectionFaces = SectionFacesMap.FindOrAdd(MaterialIndex);
+					SectionFaces.Add(RenderFaceIndex);
+				}
+			}
 		}
-	}
 
-	// Initialize the remapping array to start at LodRenderVerticesStart with the maximum number of vertices used in this LOD
-	const int32 LodRenderVerticesStart = ClothCollection->RenderVerticesStart[PatternStart];
-	const int32 LodRenderVerticesEnd = ClothCollection->RenderVerticesEnd[PatternEnd];
-	const int32 NumLodRenderVertices = LodRenderVerticesEnd - LodRenderVerticesStart + 1;
+		// Initialize the remapping array to start at LodRenderVerticesStart with the maximum number of vertices used in this LOD
+		const int32 LodRenderVerticesStart = ClothCollection->RenderVerticesStart[PatternStart];
+		const int32 LodRenderVerticesEnd = ClothCollection->RenderVerticesEnd[PatternEnd];
+		const int32 NumLodRenderVertices = LodRenderVerticesEnd - LodRenderVerticesStart + 1;
 
-	TArray<uint32> LodRenderIndexRemap;
-	LodRenderIndexRemap.SetNumUninitialized(NumLodRenderVertices);
+		LodRenderIndexRemap.SetNumUninitialized(NumLodRenderVertices);
+	}
 
 	// Keep track of the active bone indices for this LOD model
 	TSet<FBoneIndexType> ActiveBoneIndices;
