Generation/Release of the RenderData in ClothAsset now follows the sync compile flow parallel to the USkeletalMesh::Build().

diff --git a/Engine/Plugins/Experimental/ChaosClothAsset/Source/ChaosClothAssetEngine/Private/ChaosClothAsset/ClothAsset.cpp b/Engine/Plugins/Experimental/ChaosClothAsset/Source/ChaosClothAssetEngine/Private/ChaosClothAsset/ClothAsset.cpp
--- a/Engine/Plugins/Experimental/ChaosClothAsset/Source/ChaosClothAssetEngine/Private/ChaosClothAsset/ClothAsset.cpp
+++ b/Engine/Plugins/Experimental/ChaosClothAsset/Source/ChaosClothAssetEngine/Private/ChaosClothAsset/ClothAsset.cpp
@@ -326,15 +326,14 @@ void UChaosClothAsset::Build()
 {
 	using namespace UE::Chaos::ClothAsset;
 
-	// Release render resources
-	ReleaseResources();
-
-	// Flush the resource release commands to the rendering thread to ensure that the build doesn't occur while a resource is still
-	// allocated, and potentially accessing the UClothAsset
-	ReleaseResourcesFence.Wait();
+#if WITH_EDITOR
+	FSkinnedAssetAsyncBuildScope AsyncBuildScope(this);
 
-	// Clear SkeletalMeshRenderData so it can be rebuilt
-	SetResourceForRendering(nullptr);
+	FSkinnedAssetBuildContext Context;
+	BeginBuildInternal(Context);
+#else
+	ReleaseResources();
+#endif
 
 	// Set a new Guid to invalidate the DDC
 	AssetGuid = FGuid::NewGuid();
@@ -364,7 +363,8 @@ void UChaosClothAsset::Build()
 
 	// Load/save render data from/to DDC
 #if WITH_EDITOR
-	CacheDerivedData(nullptr);
+	ExecuteBuildInternal(Context);
+	FinishBuildInternal(Context);
 #endif
 
 	if (FApp::CanEverRender())
@@ -376,6 +376,48 @@ void UChaosClothAsset::Build()
 	ReregisterComponents();
 }
 
+#if WITH_EDITOR
+void UChaosClothAsset::ExecuteBuildInternal(FSkinnedAssetBuildContext& Context)
+{
+	TRACE_CPUPROFILER_EVENT_SCOPE(UChaosClothAsset::ExecuteBuildInternal);
+
+	// This scope allows us to use any locked properties without causing stalls
+	FSkinnedAssetAsyncBuildScope AsyncBuildScope(this);
+
+	// rebuild render data from imported model
+	CacheDerivedData(&Context);
+}
+
+void UChaosClothAsset::BeginBuildInternal(FSkinnedAssetBuildContext& Context)
+{
+	TRACE_CPUPROFILER_EVENT_SCOPE(UChaosClothAsset::BeginBuildInternal);
+
+	SetInternalFlags(EInternalObjectFlags::Async);
+	
+	// Unregister all instances of this component
+	Context.RecreateRenderStateContext = MakeUnique<FSkinnedMeshComponentRecreateRenderStateContext>(this, false);
+
+	// Release the render data resources.
+	ReleaseResources();
+
+	// Flush the resource release commands to the rendering thread to ensure that the build doesn't occur while a resource is still
+	// allocated, and potentially accessing the UChaosClothAsset.
+	ReleaseResourcesFence.Wait();
+
+	// Lock all properties that should not be modified/accessed during async post-load
+	AcquireAsyncProperty();
+}
+
+void UChaosClothAsset::FinishBuildInternal(FSkinnedAssetBuildContext& Context)
+{
+	TRACE_CPUPROFILER_EVENT_SCOPE(UChaosClothAsset::FinishBuildInternal);
+
+	ClearInternalFlags(EInternalObjectFlags::Async);
+
+	ReleaseAsyncProperty();
+}
+#endif // #if WITH_EDITOR
+
 #if WITH_EDITORONLY_DATA
 void UChaosClothAsset::BuildMeshModel()
 {
@@ -444,24 +486,21 @@ FString UChaosClothAsset::GetAsyncPropertyName(uint64 Property) const
 }
 
 #if WITH_EDITOR
-void UChaosClothAsset::CacheDerivedData(FSkinnedAssetPostLoadContext* Context)
+void UChaosClothAsset::CacheDerivedData(FSkinnedAssetCompilationContext* Context)
 {
 	TRACE_CPUPROFILER_EVENT_SCOPE(UChaosClothAsset::CacheDerivedData);
+	check(Context);
 
-	if (!GetOutermost()->bIsCookedForEditor)
-	{
-		if (GetResourceForRendering() == nullptr)
-		{
-			// Cache derived data for the running platform.
-			ITargetPlatform* RunningPlatform = GetTargetPlatformManagerRef().GetRunningTargetPlatform();
-			check(RunningPlatform);
+	// Cache derived data for the running platform.
+	ITargetPlatform* RunningPlatform = GetTargetPlatformManagerRef().GetRunningTargetPlatform();
+	check(RunningPlatform);
+
+	// Create the render data
+	SetResourceForRendering(MakeUnique<FSkeletalMeshRenderData>());
+
+	// Load render data from DDC, or generate it and save to DDC
+	GetResourceForRendering()->Cache(RunningPlatform, this, Context);
 
-			// Create the render data
-			SetResourceForRendering(MakeUnique<FSkeletalMeshRenderData>());
-			// Load render data from DDC, or generate it and save to DDC
-			GetResourceForRendering()->Cache(RunningPlatform, this, Context);
-		}
-	}
 }
 
 FString UChaosClothAsset::BuildDerivedDataKey(const ITargetPlatform* TargetPlatform)
diff --git a/Engine/Plugins/Experimental/ChaosClothAsset/Source/ChaosClothAssetEngine/Public/ChaosClothAsset/ClothAsset.h b/Engine/Plugins/Experimental/ChaosClothAsset/Source/ChaosClothAssetEngine/Public/ChaosClothAsset/ClothAsset.h
--- a/Engine/Plugins/Experimental/ChaosClothAsset/Source/ChaosClothAssetEngine/Public/ChaosClothAsset/ClothAsset.h
+++ b/Engine/Plugins/Experimental/ChaosClothAsset/Source/ChaosClothAssetEngine/Public/ChaosClothAsset/ClothAsset.h
@@ -13,12 +13,14 @@ class FSkeletalMeshModel;
 struct FChaosClothSimulationModel;
 struct FSkeletalMeshLODInfo;
 class UDataflow;
+class FSkinnedAssetCompilationContext;
 
 UENUM()
 enum class EClothAssetAsyncProperties : uint64
 {
 	None = 0,
 	RenderData = 1 << 0,
+	ThumbnailInfo = 1 << 1,
 	All = MAX_uint64
 };
 ENUM_CLASS_FLAGS(EClothAssetAsyncProperties);
@@ -179,7 +181,16 @@ private:
 
 #if WITH_EDITOR
 	/** Load render data from DDC if the data is cached, otherwise generate render data and save into DDC */
-	void CacheDerivedData(FSkinnedAssetPostLoadContext* Context);
+	void CacheDerivedData(FSkinnedAssetCompilationContext* Context);
+
+	/** Initial step for the building process - Can't be done in parallel. USkinnedAsset Interface. */
+	virtual void BeginBuildInternal(FSkinnedAssetBuildContext& Context) override;
+
+	/** Thread-safe part. USkinnedAsset Interface. */
+	virtual void ExecuteBuildInternal(FSkinnedAssetBuildContext& Context) override;
+
+	/** Complete the building process - Can't be done in parallel. USkinnedAsset Interface. */
+	virtual void FinishBuildInternal(FSkinnedAssetBuildContext& Context) override;
 #endif
 
 	/** Reregister all components using this asset to reset the simulation in case anything has changed. */
diff --git a/Engine/Source/Runtime/Engine/Classes/Engine/SkinnedAssetAsyncCompileUtils.h b/Engine/Source/Runtime/Engine/Classes/Engine/SkinnedAssetAsyncCompileUtils.h
--- a/Engine/Source/Runtime/Engine/Classes/Engine/SkinnedAssetAsyncCompileUtils.h
+++ b/Engine/Source/Runtime/Engine/Classes/Engine/SkinnedAssetAsyncCompileUtils.h
@@ -109,7 +109,7 @@ public:
 		RETURN_QUICK_DECLARE_CYCLE_STAT(FSkinnedAssetAsyncBuildWorker, STATGROUP_ThreadPoolAsyncTasks);
 	}
 
-	void DoWork();
+	void ENGINE_API DoWork();
 };
 
 struct FSkinnedAssetAsyncBuildTask : public FAsyncTask<FSkinnedAssetAsyncBuildWorker>
