CIS - Fix static analysis on redundant test.

#rnx

diff --git a/Engine/Plugins/Experimental/ChaosClothAsset/Source/ChaosClothAssetEngine/Private/ChaosClothAsset/ClothAsset.cpp b/Engine/Plugins/Experimental/ChaosClothAsset/Source/ChaosClothAssetEngine/Private/ChaosClothAsset/ClothAsset.cpp
--- a/Engine/Plugins/Experimental/ChaosClothAsset/Source/ChaosClothAssetEngine/Private/ChaosClothAsset/ClothAsset.cpp
+++ b/Engine/Plugins/Experimental/ChaosClothAsset/Source/ChaosClothAssetEngine/Private/ChaosClothAsset/ClothAsset.cpp
@@ -298,7 +298,7 @@ void UChaosClothAsset::CalculateBounds()
 
 			const int32 RenderVerticesStart = ClothCollection->RenderVerticesStart[PatternElementIndex];
 			const int32 RenderVerticesEnd = ClothCollection->RenderVerticesEnd[PatternElementIndex];
-			if (RenderVerticesStart != INDEX_NONE && RenderVerticesStart != INDEX_NONE)
+			if (RenderVerticesStart != INDEX_NONE && RenderVerticesEnd != INDEX_NONE)
 			{
 				for (int32 RenderVertexIndex = RenderVerticesStart; RenderVertexIndex <= RenderVerticesEnd; ++RenderVertexIndex)
 				{
