Chaos Cloth - Fixed the Cloth Component bounds calculation when the simulation isn't running.

#rnx

diff --git a/Engine/Plugins/Experimental/ChaosClothAsset/Source/ChaosClothAssetEngine/Private/ChaosClothAsset/ClothComponent.cpp b/Engine/Plugins/Experimental/ChaosClothAsset/Source/ChaosClothAssetEngine/Private/ChaosClothAsset/ClothComponent.cpp
--- a/Engine/Plugins/Experimental/ChaosClothAsset/Source/ChaosClothAssetEngine/Private/ChaosClothAsset/ClothComponent.cpp
+++ b/Engine/Plugins/Experimental/ChaosClothAsset/Source/ChaosClothAssetEngine/Private/ChaosClothAsset/ClothComponent.cpp
@@ -290,27 +290,41 @@ void UChaosClothComponent::OnPreEndOfFrameSync()
 
 FBoxSphereBounds UChaosClothComponent::CalcBounds(const FTransform& LocalToWorld) const
 {
-	const IConsoleVariable* const CVarCacheLocalSpaceBounds = IConsoleManager::Get().FindConsoleVariable(TEXT("a.CacheLocalSpaceBounds"));
-	const bool bCacheLocalSpaceBounds = CVarCacheLocalSpaceBounds ? CVarCacheLocalSpaceBounds->GetBool() : true;
+	FBoxSphereBounds NewBounds(ForceInitToZero);
 
-	const FTransform CachedBoundsTransform = bCacheLocalSpaceBounds ? FTransform::Identity : LocalToWorld;
-
-	FBoxSphereBounds NewBounds;
-	if (ClothSimulationProxy)
+	// Use cached local bounds if possible
+	if (bCachedWorldSpaceBoundsUpToDate || bCachedLocalBoundsUpToDate)
 	{
-		NewBounds = ClothSimulationProxy->CalculateBounds_AnyThread().TransformBy(CachedBoundsTransform);
+		NewBounds = bCachedLocalBoundsUpToDate ?
+			CachedWorldOrLocalSpaceBounds.TransformBy(LocalToWorld) :
+			CachedWorldOrLocalSpaceBounds.TransformBy(CachedWorldToLocalTransform * LocalToWorld.ToMatrixWithScale());
 	}
+	else  // Calculate new bounds
+	{
+		const IConsoleVariable* const CVarCacheLocalSpaceBounds = IConsoleManager::Get().FindConsoleVariable(TEXT("a.CacheLocalSpaceBounds"));
+		const bool bCacheLocalSpaceBounds = CVarCacheLocalSpaceBounds ? (CVarCacheLocalSpaceBounds->GetInt() != 0) : true;
 
-	CachedWorldOrLocalSpaceBounds = NewBounds;
-	bCachedLocalBoundsUpToDate = bCacheLocalSpaceBounds;
-	bCachedWorldSpaceBoundsUpToDate = !bCacheLocalSpaceBounds;
+		const FTransform CachedBoundsTransform = bCacheLocalSpaceBounds ? FTransform::Identity : LocalToWorld;
 
-	if (bCacheLocalSpaceBounds)
-	{
-		CachedWorldToLocalTransform.SetIdentity();
-		return NewBounds.TransformBy(LocalToWorld);
+		if (ClothSimulationProxy)
+		{
+			NewBounds = ClothSimulationProxy->CalculateBounds_AnyThread().TransformBy(CachedBoundsTransform);
+		}
+
+		CachedWorldOrLocalSpaceBounds = NewBounds;
+		bCachedLocalBoundsUpToDate = bCacheLocalSpaceBounds;
+		bCachedWorldSpaceBoundsUpToDate = !bCacheLocalSpaceBounds;
+
+		if (bCacheLocalSpaceBounds)
+		{
+			CachedWorldToLocalTransform.SetIdentity();
+			NewBounds = NewBounds.TransformBy(LocalToWorld);
+		}
+		else
+		{
+			CachedWorldToLocalTransform = LocalToWorld.ToInverseMatrixWithScale();
+		}
 	}
-	CachedWorldToLocalTransform = LocalToWorld.ToInverseMatrixWithScale();
 	return NewBounds;
 }
 
diff --git a/Engine/Plugins/Experimental/ChaosClothAsset/Source/ChaosClothAssetEngine/Private/ChaosClothAsset/ClothSimulationProxy.cpp b/Engine/Plugins/Experimental/ChaosClothAsset/Source/ChaosClothAssetEngine/Private/ChaosClothAsset/ClothSimulationProxy.cpp
--- a/Engine/Plugins/Experimental/ChaosClothAsset/Source/ChaosClothAssetEngine/Private/ChaosClothAsset/ClothSimulationProxy.cpp
+++ b/Engine/Plugins/Experimental/ChaosClothAsset/Source/ChaosClothAssetEngine/Private/ChaosClothAsset/ClothSimulationProxy.cpp
@@ -96,7 +96,13 @@ namespace UE::Chaos::ClothAsset
 
 		// Need a valid context to initialize the mesh
 		constexpr bool bIsInitialization = true;
-		ClothSimulationContext.Fill(ClothComponent, 0.f, MaxDeltaTime, bIsInitialization);
+		constexpr Softs::FSolverReal NoAdvanceDt = 0.f;
+		ClothSimulationContext.Fill(ClothComponent, NoAdvanceDt, MaxDeltaTime, bIsInitialization);
+
+		// Setup startup transforms
+		constexpr bool bNeedsReset = true;
+		Solver->SetLocalSpaceLocation((FVec3)ClothSimulationContext.ComponentTransform.GetLocation(), bNeedsReset);
+		Solver->SetLocalSpaceRotation((FQuat)ClothSimulationContext.ComponentTransform.GetRotation());
 
 		// Create mesh simulation thread object
 		const UChaosClothAsset* const ClothAsset = ClothComponent.GetClothAsset();
@@ -142,6 +148,10 @@ PRAGMA_ENABLE_DEPRECATION_WARNINGS
 		// Create solver config simulation thread object
 		const int32 SolverConfigIndex = Configs.Emplace(MakeUnique<FClothingSimulationConfig>(ClothComponent.GetPropertyCollection()));  // TODO: Use a separate solver config for outfits
 		Solver->SetConfig(Configs[SolverConfigIndex].Get());
+
+		// Set start pose (update the context, then the solver without advancing the simulation)
+		ClothSimulationContext.Fill(ClothComponent, NoAdvanceDt, MaxDeltaTime);
+		Solver->Update((Softs::FSolverReal)ClothSimulationContext.DeltaTime);
 	}
 
 	FClothSimulationProxy::~FClothSimulationProxy()
