Chaos Cloth:
- Refactor ClothAdapters into more consistent Geometry Collections type Facades.
- Rework Facade interfaces to prohibit direct access to the underlying collection.
- Added LODs management to Dataflow node.
- Added support for materials, skeletons, and physics assets to the cloth LOD adapter API.
- Added a new LOD copy functionality to the cloth adapter API for building the final Cloth Asset of the various input LODs.
- Added a new SetNumPatterns function to the cloth adapter API, and fixed the AddPattern to correctly deal with re-indexing when multiple LODs are used.

#rnx

diff --git a/Engine/Plugins/Enterprise/DatasmithImporter/Source/DatasmithImporter/Private/DatasmithImporter.cpp b/Engine/Plugins/Enterprise/DatasmithImporter/Source/DatasmithImporter/Private/DatasmithImporter.cpp
--- a/Engine/Plugins/Enterprise/DatasmithImporter/Source/DatasmithImporter/Private/DatasmithImporter.cpp
+++ b/Engine/Plugins/Enterprise/DatasmithImporter/Source/DatasmithImporter/Private/DatasmithImporter.cpp
@@ -83,9 +83,7 @@
 
 #include "Animation/Skeleton.h"
 #include "ChaosClothAsset/ClothAsset.h"
-#include "ChaosClothAsset/ClothAdapter.h"
-#include "ChaosClothAsset/ClothLodAdapter.h"
-#include "ChaosClothAsset/ClothPatternAdapter.h"
+#include "ChaosClothAsset/CollectionClothFacade.h"
 
 extern UNREALED_API UEditorEngine* GEditor;
 
@@ -488,14 +486,6 @@ void FDatasmithImporter::ImportClothes(FDatasmithImportContext& ImportContext)
 		else
 		{
 			ClothAsset = NewObject<UChaosClothAsset>(ImportPackage, *ClothName, ObjectFlags);
-
-			// Setup a default skeleton
-			if (ClothAsset && !ClothAsset->GetSkeleton())
-			{
-				USkeleton* const DefaultSkeleton = LoadObject<USkeleton>(nullptr, TEXT("/Engine/EditorMeshes/SkeletalMesh/DefaultSkeletalMesh_Skeleton.DefaultSkeletalMesh_Skeleton"), nullptr, LOAD_None, nullptr);
-				check(DefaultSkeleton);
-				ClothAsset->SetSkeleton(DefaultSkeleton);
-			}
 		}
 
 		if (!ensure(ClothAsset))
@@ -503,30 +493,30 @@ void FDatasmithImporter::ImportClothes(FDatasmithImportContext& ImportContext)
 			continue;
 		}
 
-		TSharedPtr<UE::Chaos::ClothAsset::FClothCollection> Collection = ClothAsset->GetClothCollection();
-		UE::Chaos::ClothAsset::FClothAdapter Adapter(Collection);
-		Adapter.Reset(); // on reimport, asset is non-empty
-		UE::Chaos::ClothAsset::FClothLodAdapter LODAdapter = Adapter.GetNumLods() ? Adapter.GetLod(0) : Adapter.AddGetLod();
+		TSharedPtr<FManagedArrayCollection> Collection = ClothAsset->GetClothCollection();
+		UE::Chaos::ClothAsset::FCollectionClothFacade Cloth(Collection);
+		Cloth.Reset(); // on reimport, asset is non-empty
+		UE::Chaos::ClothAsset::FCollectionClothLodFacade ClothLod = Cloth.GetNumLods() ? Cloth.GetLod(0) : Cloth.AddGetLod();
 
 		for (FDatasmithClothPattern& Pattern : DsCloth.Patterns)
 		{
 			if (Pattern.IsValid())
 			{
-				UE::Chaos::ClothAsset::FClothPatternAdapter ClothPattern = LODAdapter.AddGetPattern();
+				UE::Chaos::ClothAsset::FCollectionClothPatternFacade ClothPattern = ClothLod.AddGetPattern();
 				ClothPattern.Initialize(Pattern.SimPosition, Pattern.SimRestPosition, Pattern.SimTriangleIndices);
 			}
 		}
 
-		LODAdapter.SetNumSeams(DsCloth.Sewing.Num());
-		TArrayView<FIntVector2> SeamPatterns = LODAdapter.GetSeamPatterns();
-		TArrayView<TArray<FIntVector2>> SeamStitches = LODAdapter.GetSeamStitches();
+		ClothLod.SetNumSeams(DsCloth.Sewing.Num());
+		TArrayView<FIntVector2> SeamPatterns = ClothLod.GetSeamPatterns();
+		TArrayView<TArray<FIntVector2>> SeamStitches = ClothLod.GetSeamStitches();
 
 		uint32 SeamIndex = 0;
 		for (const FDatasmithClothSewingInfo& SeamInfo : DsCloth.Sewing)
 		{
 			SeamPatterns[SeamIndex] = FIntVector2(SeamInfo.Seam0PanelIndex, SeamInfo.Seam1PanelIndex);
 			TArray<FIntVector2>& Stitches = SeamStitches[SeamIndex];
-			uint32 StitchesCount = std::min(SeamInfo.Seam0MeshIndices.Num(), SeamInfo.Seam1MeshIndices.Num());
+			const uint32 StitchesCount = FMath::Min(SeamInfo.Seam0MeshIndices.Num(), SeamInfo.Seam1MeshIndices.Num());
 			Stitches.Reserve(StitchesCount);
 			for (uint32 StitchIndex = 0; StitchIndex < StitchesCount; ++StitchIndex)
 			{
@@ -536,23 +526,14 @@ void FDatasmithImporter::ImportClothes(FDatasmithImportContext& ImportContext)
 			++SeamIndex;
 		}
 
-		// Add a default material
-		ClothAsset->GetMaterials().Reset(1);
-		UMaterial* DefaultMaterial = LoadObject<UMaterial>(nullptr, TEXT("/Engine/EditorMaterials/Cloth/CameraLitDoubleSided.CameraLitDoubleSided"), nullptr, LOAD_None, nullptr);
-		const int32 MaterialId = ClothAsset->GetMaterials().Emplace(DefaultMaterial, true, false, DefaultMaterial->GetFName());
-
 		// Set a default skeleton
-		if (USkeleton* DefaultSkeleton = LoadObject<USkeleton>(nullptr, TEXT("/Engine/EditorMeshes/SkeletalMesh/DefaultSkeletalMesh_Skeleton.DefaultSkeletalMesh_Skeleton"), nullptr, LOAD_None, nullptr))
-		{
-			ClothAsset->SetReferenceSkeleton(DefaultSkeleton->GetReferenceSkeleton(), false);
-		}
-
-		// Bind the mesh to the root bone.
-		ClothAsset->BindSimMeshToRootBone(); 
+		constexpr bool bRebuildModels = false;
+		ClothAsset->SetSkeleton(nullptr, bRebuildModels);
 
 		// Set the render mesh to duplicate the sim mesh
-		ClothAsset->CopySimMeshToRenderMesh(MaterialId);
+		ClothAsset->CopySimMeshToRenderMesh();
 
+		// Rebuild static data
 		ClothAsset->Build();
 
 		ImportContext.ImportedClothes.Add(ClothElement, ClothAsset);
diff --git a/Engine/Plugins/Experimental/ChaosClothAsset/Source/ChaosClothAsset/Private/ChaosClothAsset/ClothAdapter.cpp b/Engine/Plugins/Experimental/ChaosClothAsset/Source/ChaosClothAsset/Private/ChaosClothAsset/ClothAdapter.cpp
deleted file mode 100644
diff --git a/Engine/Plugins/Experimental/ChaosClothAsset/Source/ChaosClothAsset/Private/ChaosClothAsset/ClothCollection.cpp b/Engine/Plugins/Experimental/ChaosClothAsset/Source/ChaosClothAsset/Private/ChaosClothAsset/ClothCollection.cpp
--- a/Engine/Plugins/Experimental/ChaosClothAsset/Source/ChaosClothAsset/Private/ChaosClothAsset/ClothCollection.cpp
+++ b/Engine/Plugins/Experimental/ChaosClothAsset/Source/ChaosClothAsset/Private/ChaosClothAsset/ClothCollection.cpp
@@ -1,128 +1,465 @@
 // Copyright Epic Games, Inc. All Rights Reserved.
 
 #include "ChaosClothAsset/ClothCollection.h"
+#include "GeometryCollection/ManagedArrayCollection.h"
+
+namespace UE::Chaos::ClothAsset::Private
+{
+	// Groups
+	static const FName LodsGroup(TEXT("Lods"));
+	static const FName MaterialsGroup(TEXT("Materials"));
+	static const FName TetherBatchesGroup(TEXT("TetherBatches"));
+	static const FName TethersGroup(TEXT("Tethers"));
+	static const FName SeamsGroup(TEXT("Seams"));
+	static const FName PatternsGroup(TEXT("Patterns"));
+	static const FName SimFacesGroup(TEXT("SimFaces"));
+	static const FName SimVerticesGroup(TEXT("SimVertices"));
+	static const FName RenderFacesGroup(TEXT("RenderFaces"));
+	static const FName RenderVerticesGroup(TEXT("RenderVertices"));
+
+	// LODs Group
+	static const FName MaterialStartAttribute(TEXT("MaterialStart"));
+	static const FName MaterialEndAttribute(TEXT("MaterialEnd"));
+	static const FName TetherBatchStartAttribute(TEXT("TetherBatchStart"));
+	static const FName TetherBatchEndAttribute(TEXT("TetherBatchEnd"));
+	static const FName SeamStartAttribute(TEXT("SeamStart"));
+	static const FName SeamEndAttribute(TEXT("SeamEnd"));
+	static const FName PatternStartAttribute(TEXT("PatternStart"));
+	static const FName PatternEndAttribute(TEXT("PatternEnd"));
+	static const FName PhysicsAssetPathNameAttribute(TEXT("PhysicsAssetPathName"));
+	static const FName SkeletonAssetPathNameAttribute(TEXT("SkeletonAssetPathName"));
+	static const TArray<FName> LodsGroupAttributes =
+	{
+		MaterialStartAttribute,
+		MaterialEndAttribute,
+		TetherBatchStartAttribute,
+		TetherBatchEndAttribute,
+		SeamStartAttribute,
+		SeamEndAttribute,
+		PatternStartAttribute,
+		PatternEndAttribute,
+		PhysicsAssetPathNameAttribute,
+		SkeletonAssetPathNameAttribute
+	};
+
+	// Materials Group
+	static const FName RenderMaterialPathNameAttribute(TEXT("RenderMaterialPathName"));
+	static const TArray<FName> MaterialsGroupAttributes =
+	{
+		RenderMaterialPathNameAttribute
+	};
+
+	// Tether Batches Group
+	static const FName TetherStartAttribute(TEXT("TetherStart"));
+	static const FName TetherEndAttribute(TEXT("TetherEnd"));
+	static const TArray<FName> TetherBatchesGroupAttributes =
+	{
+		TetherStartAttribute,
+		TetherEndAttribute
+	};
+
+	// Tethers Group
+	static const FName TetherKinematicIndexAttribute(TEXT("TetherKinematicIndex"));
+	static const FName TetherDynamicIndexAttribute(TEXT("TetherDynamicIndex"));
+	static const FName TetherReferenceLengthAttribute(TEXT("TetherReferenceLength"));
+	static const TArray<FName> TethersGroupAttributes =
+	{
+		TetherKinematicIndexAttribute,
+		TetherDynamicIndexAttribute,
+		TetherReferenceLengthAttribute
+	};
+
+	// Seam Group
+	static const FName SeamPatternsAttribute(TEXT("SeamPatterns"));
+	static const FName SeamStitchesAttribute(TEXT("SeamStitches"));
+	static const TArray<FName> SeamsGroupAttributes =
+	{
+		SeamPatternsAttribute,
+		SeamStitchesAttribute
+	};
+
+	// Patterns Group
+	static const FName SimVerticesStartAttribute(TEXT("SimVerticesStart"));
+	static const FName SimVerticesEndAttribute(TEXT("SimVerticesEnd"));
+	static const FName SimFacesStartAttribute(TEXT("SimFacesStart"));
+	static const FName SimFacesEndAttribute(TEXT("SimFacesEnd"));
+	static const FName RenderVerticesStartAttribute(TEXT("RenderVerticesStar"));
+	static const FName RenderVerticesEndAttribute(TEXT("RenderVerticesEnd"));
+	static const FName RenderFacesStartAttribute(TEXT("RenderFacesStart"));
+	static const FName RenderFacesEndAttribute(TEXT("RenderFacesEnd"));
+	static const FName StatusFlagsAttribute(TEXT("StatusFlags"));
+	static const TArray<FName> PatternsGroupAttributes =
+	{
+		SimVerticesStartAttribute,
+		SimVerticesEndAttribute,
+		SimFacesStartAttribute,
+		SimFacesEndAttribute,
+		RenderVerticesStartAttribute,
+		RenderVerticesEndAttribute,
+		RenderFacesStartAttribute,
+		RenderFacesEndAttribute,
+		StatusFlagsAttribute
+	};
+
+	// Sim Faces Group
+	static const FName SimIndicesAttribute(TEXT("SimIndices"));
+	static const TArray<FName> SimFacesGroupAttributes =
+	{
+		SimIndicesAttribute
+	};
+
+	// Sim Vertices Group
+	static const FName SimPositionAttribute(TEXT("SimPosition"));
+	static const FName SimRestPositionAttribute(TEXT("SimRestPosition"));
+	static const FName SimRestNormalAttribute(TEXT("SimRestNormal"));
+	static const FName SimNumBoneInfluencesAttribute(TEXT("SimNumBoneInfluences"));
+	static const FName SimBoneIndicesAttribute(TEXT("SimBoneIndices"));
+	static const FName SimBoneWeightsAttribute(TEXT("SimBoneWeights"));
+	static const TArray<FName> SimVerticesGroupAttributes =
+	{
+		SimPositionAttribute,
+		SimRestPositionAttribute,
+		SimRestNormalAttribute,
+		SimNumBoneInfluencesAttribute,
+		SimBoneIndicesAttribute,
+		SimBoneWeightsAttribute
+	};
+
+	// Render Faces Group
+	static const FName RenderIndicesAttribute(TEXT("RenderIndices"));
+	static const FName RenderMaterialIndexAttribute(TEXT("RenderMaterialIndex"));
+	static const TArray<FName> RenderFacesGroupAttributes =
+	{
+		RenderIndicesAttribute,
+		RenderMaterialIndexAttribute
+	};
+
+	// Render Vertices Group
+	static const FName RenderPositionAttribute(TEXT("RenderPosition"));
+	static const FName RenderNormalAttribute(TEXT("RenderNormal"));
+	static const FName RenderTangentUAttribute(TEXT("RenderTangentU"));
+	static const FName RenderTangentVAttribute(TEXT("RenderTangentV"));
+	static const FName RenderUVsAttribute(TEXT("RenderUVs"));
+	static const FName RenderColorAttribute(TEXT("RenderColor"));
+	static const FName RenderNumBoneInfluencesAttribute(TEXT("RenderNumBoneInfluences"));
+	static const FName RenderBoneIndicesAttribute(TEXT("RenderBoneIndices"));
+	static const FName RenderBoneWeightsAttribute(TEXT("RenderBoneWeights"));
+	static const TArray<FName> RenderVerticesGroupAttributes =
+	{
+		RenderPositionAttribute,
+		RenderNormalAttribute,
+		RenderTangentUAttribute,
+		RenderTangentVAttribute,
+		RenderUVsAttribute,
+		RenderColorAttribute,
+		RenderNumBoneInfluencesAttribute,
+		RenderBoneIndicesAttribute,
+		RenderBoneWeightsAttribute
+	};
+
+	static const TMap<FName, TArray<FName>> FixedAttributeNamesMap =
+	{
+		{ LodsGroup, LodsGroupAttributes },
+		{ MaterialsGroup, MaterialsGroupAttributes },
+		{ TetherBatchesGroup, TetherBatchesGroupAttributes },
+		{ TethersGroup, TethersGroupAttributes },
+		{ SeamsGroup, SeamsGroupAttributes },
+		{ PatternsGroup, PatternsGroupAttributes },
+		{ SimFacesGroup, SimFacesGroupAttributes },
+		{ SimVerticesGroup, SimVerticesGroupAttributes },
+		{ RenderFacesGroup, RenderFacesGroupAttributes },
+		{ RenderVerticesGroup, RenderVerticesGroupAttributes }
+	};
+}  // End namespace UE::Chaos::ClothAsset::Private
 
 namespace UE::Chaos::ClothAsset
 {
-	const FName UE::Chaos::ClothAsset::FClothCollection::SimVerticesGroup("SimVertices");
-	const FName UE::Chaos::ClothAsset::FClothCollection::SimFacesGroup("SimFaces");
-	const FName UE::Chaos::ClothAsset::FClothCollection::RenderVerticesGroup("RenderVertices");
-	const FName UE::Chaos::ClothAsset::FClothCollection::RenderFacesGroup("RenderFaces");
-	const FName UE::Chaos::ClothAsset::FClothCollection::WrapDeformersGroup("WrapDeformers");
-	const FName UE::Chaos::ClothAsset::FClothCollection::PatternsGroup("Patterns");
-	const FName UE::Chaos::ClothAsset::FClothCollection::SeamsGroup("Seams");
-	const FName UE::Chaos::ClothAsset::FClothCollection::TethersGroup("Tethers");
-	const FName UE::Chaos::ClothAsset::FClothCollection::TetherBatchesGroup("TetherBatches");
-	const FName UE::Chaos::ClothAsset::FClothCollection::LodsGroup("Lods");
-	const FName UE::Chaos::ClothAsset::FClothCollection::MaterialsGroup("Materials");
-	const FName UE::Chaos::ClothAsset::FClothCollection::CollisionsGroup("CollisionsGroup");
-	const FName UE::Chaos::ClothAsset::FClothCollection::SkeletonsGroup("SkeletonsGroup");
-
-	FClothCollection::FClothCollection()
-	{
-		Construct();
+	// Groups
+	const FName FClothCollection::LodsGroup = Private::LodsGroup;
+	const FName FClothCollection::MaterialsGroup = Private::MaterialsGroup;
+	const FName FClothCollection::TetherBatchesGroup = Private::TetherBatchesGroup;
+	const FName FClothCollection::TethersGroup = Private::TethersGroup;
+	const FName FClothCollection::SeamsGroup = Private::SeamsGroup;
+	const FName FClothCollection::PatternsGroup = Private::PatternsGroup;
+	const FName FClothCollection::SimFacesGroup = Private::SimFacesGroup;
+	const FName FClothCollection::SimVerticesGroup = Private::SimVerticesGroup;
+	const FName FClothCollection::RenderFacesGroup = Private::RenderFacesGroup;
+	const FName FClothCollection::RenderVerticesGroup = Private::RenderVerticesGroup;
+
+	FClothCollection::FClothCollection(const TSharedPtr<FManagedArrayCollection>& InManagedArrayCollection)
+		: ManagedArrayCollection(InManagedArrayCollection)
+	{
+		using namespace UE::Chaos::ClothAsset::Private;
+
+		check(ManagedArrayCollection.IsValid());
+
+		// LODs Group
+		MaterialStart = ManagedArrayCollection->FindAttribute<int32>(MaterialStartAttribute, LodsGroup);
+		MaterialEnd = ManagedArrayCollection->FindAttribute<int32>(MaterialEndAttribute, LodsGroup);
+		TetherBatchStart = ManagedArrayCollection->FindAttribute<int32>(TetherBatchStartAttribute, LodsGroup);
+		TetherBatchEnd = ManagedArrayCollection->FindAttribute<int32>(TetherBatchEndAttribute, LodsGroup);
+		SeamStart = ManagedArrayCollection->FindAttribute<int32>(SeamStartAttribute, LodsGroup);
+		SeamEnd = ManagedArrayCollection->FindAttribute<int32>(SeamEndAttribute, LodsGroup);
+		PatternStart = ManagedArrayCollection->FindAttribute<int32>(PatternStartAttribute, LodsGroup);
+		PatternEnd = ManagedArrayCollection->FindAttribute<int32>(PatternEndAttribute, LodsGroup);
+		PhysicsAssetPathName = ManagedArrayCollection->FindAttribute<FString>(PhysicsAssetPathNameAttribute, LodsGroup);
+		SkeletonAssetPathName = ManagedArrayCollection->FindAttribute<FString>(SkeletonAssetPathNameAttribute, LodsGroup);
+
+		// Materials Group
+		RenderMaterialPathName = ManagedArrayCollection->FindAttribute<FString>(RenderMaterialPathNameAttribute, MaterialsGroup);
+
+		// Tether Batches Group
+		TetherStart = ManagedArrayCollection->FindAttribute<int32>(TetherStartAttribute, TetherBatchesGroup);
+		TetherEnd = ManagedArrayCollection->FindAttribute<int32>(TetherEndAttribute, TetherBatchesGroup);
+
+		// Tethers Group
+		TetherKinematicIndex = ManagedArrayCollection->FindAttribute<int32>(TetherKinematicIndexAttribute, TethersGroup);
+		TetherDynamicIndex = ManagedArrayCollection->FindAttribute<int32>(TetherDynamicIndexAttribute, TethersGroup);
+		TetherReferenceLength = ManagedArrayCollection->FindAttribute<float>(TetherReferenceLengthAttribute, TethersGroup);
+
+		// Seam Group
+		SeamPatterns = ManagedArrayCollection->FindAttribute<FIntVector2>(SeamPatternsAttribute, SeamsGroup);
+		SeamStitches = ManagedArrayCollection->FindAttribute<TArray<FIntVector2>>(SeamStitchesAttribute, SeamsGroup);
+
+		// Patterns Group
+		SimVerticesStart = ManagedArrayCollection->FindAttribute<int32>(SimVerticesStartAttribute, PatternsGroup);
+		SimVerticesEnd = ManagedArrayCollection->FindAttribute<int32>(SimVerticesEndAttribute, PatternsGroup);
+		SimFacesStart = ManagedArrayCollection->FindAttribute<int32>(SimFacesStartAttribute, PatternsGroup);
+		SimFacesEnd = ManagedArrayCollection->FindAttribute<int32>(SimFacesEndAttribute, PatternsGroup);
+		RenderVerticesStart = ManagedArrayCollection->FindAttribute<int32>(RenderVerticesStartAttribute, PatternsGroup);
+		RenderVerticesEnd = ManagedArrayCollection->FindAttribute<int32>(RenderVerticesEndAttribute, PatternsGroup);
+		RenderFacesStart = ManagedArrayCollection->FindAttribute<int32>(RenderFacesStartAttribute, PatternsGroup);
+		RenderFacesEnd = ManagedArrayCollection->FindAttribute<int32>(RenderFacesEndAttribute, PatternsGroup);
+		StatusFlags = ManagedArrayCollection->FindAttribute<int32>(StatusFlagsAttribute, PatternsGroup);
+
+		// Sim Faces Group
+		SimIndices = ManagedArrayCollection->FindAttribute<FIntVector3>(SimIndicesAttribute, SimFacesGroup);
+
+		// Sim Vertices Group
+		SimPosition = ManagedArrayCollection->FindAttribute<FVector2f>(SimPositionAttribute, SimVerticesGroup);
+		SimRestPosition = ManagedArrayCollection->FindAttribute<FVector3f>(SimRestPositionAttribute, SimVerticesGroup);
+		SimRestNormal = ManagedArrayCollection->FindAttribute<FVector3f>(SimRestNormalAttribute, SimVerticesGroup);
+		SimNumBoneInfluences = ManagedArrayCollection->FindAttribute<int32>(SimNumBoneInfluencesAttribute, SimVerticesGroup);
+		SimBoneIndices = ManagedArrayCollection->FindAttribute<TArray<int32>>(SimBoneIndicesAttribute, SimVerticesGroup);
+		SimBoneWeights = ManagedArrayCollection->FindAttribute<TArray<float>>(SimBoneWeightsAttribute, SimVerticesGroup);
+
+		// Render Faces Group
+		RenderIndices = ManagedArrayCollection->FindAttribute<FIntVector3>(RenderIndicesAttribute, RenderFacesGroup);
+		RenderMaterialIndex = ManagedArrayCollection->FindAttribute<int32>(RenderMaterialIndexAttribute, RenderFacesGroup);
+
+		// Render Vertices Group
+		RenderPosition = ManagedArrayCollection->FindAttribute<FVector3f>(RenderPositionAttribute, RenderVerticesGroup);
+		RenderNormal = ManagedArrayCollection->FindAttribute<FVector3f>(RenderNormalAttribute, RenderVerticesGroup);
+		RenderTangentU = ManagedArrayCollection->FindAttribute<FVector3f>(RenderTangentUAttribute, RenderVerticesGroup);
+		RenderTangentV = ManagedArrayCollection->FindAttribute<FVector3f>(RenderTangentVAttribute, RenderVerticesGroup);
+		RenderUVs = ManagedArrayCollection->FindAttribute<TArray<FVector2f>>(RenderUVsAttribute, RenderVerticesGroup);
+		RenderColor = ManagedArrayCollection->FindAttribute<FLinearColor>(RenderColorAttribute, RenderVerticesGroup);
+		RenderNumBoneInfluences = ManagedArrayCollection->FindAttribute<int32>(RenderNumBoneInfluencesAttribute, RenderVerticesGroup);
+		RenderBoneIndices = ManagedArrayCollection->FindAttribute<TArray<int32>>(RenderBoneIndicesAttribute, RenderVerticesGroup);
+		RenderBoneWeights = ManagedArrayCollection->FindAttribute<TArray<float>>(RenderBoneWeightsAttribute, RenderVerticesGroup);
 	}
 
-	void FClothCollection::Serialize(FArchive& Ar)
+	bool FClothCollection::IsValid() const
 	{
-		::Chaos::FChaosArchive ChaosArchive(Ar);
-		Super::Serialize(ChaosArchive);
+		return 
+			// LODs Group
+			MaterialStart &&
+			MaterialEnd &&
+			TetherBatchStart &&
+			TetherBatchEnd &&
+			SeamStart &&
+			SeamEnd &&
+			PatternStart &&
+			PatternEnd &&
+			PhysicsAssetPathName &&
+			SkeletonAssetPathName &&
+
+			// Materials Group
+			RenderMaterialPathName &&
+
+			// Tether Batches Group
+			TetherStart &&
+			TetherEnd &&
+
+			// Tethers Group
+			TetherKinematicIndex &&
+			TetherDynamicIndex &&
+			TetherReferenceLength &&
+
+			// Seam Group
+			SeamPatterns &&
+			SeamStitches &&
+
+			// Patterns Group
+			SimVerticesStart &&
+			SimVerticesEnd &&
+			SimFacesStart &&
+			SimFacesEnd &&
+			RenderVerticesStart &&
+			RenderVerticesEnd &&
+			RenderFacesStart &&
+			RenderFacesEnd &&
+			StatusFlags &&
+
+			// Sim Faces Group
+			SimIndices &&
+
+			// Sim Vertices Group
+			SimPosition &&
+			SimRestPosition &&
+			SimRestNormal &&
+			SimNumBoneInfluences &&
+			SimBoneIndices &&
+			SimBoneWeights &&
+
+			// Render Faces Group
+			RenderIndices &&
+			RenderMaterialIndex &&
+
+			// Render Vertices Group
+			RenderPosition &&
+			RenderNormal &&
+			RenderTangentU &&
+			RenderTangentV &&
+			RenderUVs &&
+			RenderColor &&
+			RenderNumBoneInfluences &&
+			RenderBoneIndices &&
+			RenderBoneWeights;
 	}
 
-	void FClothCollection::Construct()
+	void FClothCollection::DefineSchema()
 	{
+		using namespace UE::Chaos::ClothAsset::Private;
+
 		// Dependencies
-		FManagedArrayCollection::FConstructionParameters SimVerticesDependency(SimVerticesGroup);
-		FManagedArrayCollection::FConstructionParameters SimFacesDependency(SimFacesGroup);
-		FManagedArrayCollection::FConstructionParameters RenderVerticesDependency(RenderVerticesGroup);
-		FManagedArrayCollection::FConstructionParameters RenderFacesDependency(RenderFacesGroup);
-		FManagedArrayCollection::FConstructionParameters WrapDeformersDependency(WrapDeformersGroup);
-		FManagedArrayCollection::FConstructionParameters PatternsDependency(PatternsGroup);
-		FManagedArrayCollection::FConstructionParameters StitchingsDependency(SeamsGroup);
-		FManagedArrayCollection::FConstructionParameters TethersDependency(TethersGroup);
+		FManagedArrayCollection::FConstructionParameters MaterialsDependency(MaterialsGroup);
 		FManagedArrayCollection::FConstructionParameters TetherBatchesDependency(TetherBatchesGroup);
+		FManagedArrayCollection::FConstructionParameters TethersDependency(TethersGroup);
+		FManagedArrayCollection::FConstructionParameters SeamsDependency(SeamsGroup);
+		FManagedArrayCollection::FConstructionParameters PatternsDependency(PatternsGroup);
+		FManagedArrayCollection::FConstructionParameters RenderFacesDependency(RenderFacesGroup);
+		FManagedArrayCollection::FConstructionParameters RenderVerticesDependency(RenderVerticesGroup);
+		FManagedArrayCollection::FConstructionParameters SimFacesDependency(SimFacesGroup);
+		FManagedArrayCollection::FConstructionParameters SimVerticesDependency(SimVerticesGroup);
 
-		// Sim Vertices Group
-		AddExternalAttribute<FVector2f>("SimPosition", SimVerticesGroup, SimPosition);
-		AddExternalAttribute<FVector3f>("SimRestPosition", SimVerticesGroup, SimRestPosition);
-		AddExternalAttribute<FVector3f>("SimRestNormal", SimVerticesGroup, SimRestNormal);
-		AddExternalAttribute<int32>("SimNumBoneInfluences", SimVerticesGroup, SimNumBoneInfluences);
-		AddExternalAttribute<TArray<int32>>("SimBoneIndices", SimVerticesGroup, SimBoneIndices);
-		AddExternalAttribute<TArray<float>>("SimBoneWeights", SimVerticesGroup, SimBoneWeights);
+		// LODs Group
+		PatternStart = &ManagedArrayCollection->AddAttribute<int32>(PatternStartAttribute, LodsGroup, PatternsDependency);
+		PatternEnd = &ManagedArrayCollection->AddAttribute<int32>(PatternEndAttribute, LodsGroup, PatternsDependency);
+		SeamStart = &ManagedArrayCollection->AddAttribute<int32>(SeamStartAttribute, LodsGroup, SeamsDependency);
+		SeamEnd = &ManagedArrayCollection->AddAttribute<int32>(SeamEndAttribute, LodsGroup, SeamsDependency);
+		TetherBatchStart = &ManagedArrayCollection->AddAttribute<int32>(TetherBatchStartAttribute, LodsGroup, TetherBatchesDependency);
+		TetherBatchEnd = &ManagedArrayCollection->AddAttribute<int32>(TetherBatchEndAttribute, LodsGroup, TetherBatchesDependency);
+		MaterialStart = &ManagedArrayCollection->AddAttribute<int32>(MaterialStartAttribute, LodsGroup, MaterialsDependency);
+		MaterialEnd = &ManagedArrayCollection->AddAttribute<int32>(MaterialEndAttribute, LodsGroup, MaterialsDependency);
+		PhysicsAssetPathName = &ManagedArrayCollection->AddAttribute<FString>(PhysicsAssetPathNameAttribute, LodsGroup);
+		SkeletonAssetPathName = &ManagedArrayCollection->AddAttribute<FString>(SkeletonAssetPathNameAttribute, LodsGroup);
 
-		// Sim Faces Group
-		AddExternalAttribute<FIntVector3>("SimIndices", SimFacesGroup, SimIndices, SimVerticesDependency);
+		// Materials Group
+		RenderMaterialPathName = &ManagedArrayCollection->AddAttribute<FString>(RenderMaterialPathNameAttribute, MaterialsGroup);
 
-		// Render Vertices Group
-		AddExternalAttribute<FVector3f>("RenderPosition", RenderVerticesGroup, RenderPosition);
-		AddExternalAttribute<FVector3f>("RenderNormal", RenderVerticesGroup, RenderNormal);
-		AddExternalAttribute<FVector3f>("RenderTangentU", RenderVerticesGroup, RenderTangentU);
-		AddExternalAttribute<FVector3f>("RenderTangentV", RenderVerticesGroup, RenderTangentV);
-		AddExternalAttribute<TArray<FVector2f>>("RenderUVs", RenderVerticesGroup, RenderUVs);
-		AddExternalAttribute<FLinearColor>("RenderColor", RenderVerticesGroup, RenderColor);
-		AddExternalAttribute<int32>("RenderNumBoneInfluences", RenderVerticesGroup, RenderNumBoneInfluences);
-		AddExternalAttribute<TArray<int32>>("RenderBoneIndices", RenderVerticesGroup, RenderBoneIndices);
-		AddExternalAttribute<TArray<float>>("RenderBoneWeights", RenderVerticesGroup, RenderBoneWeights);
+		// Tether Batches Group
+		TetherStart = &ManagedArrayCollection->AddAttribute<int32>(TetherStartAttribute, TetherBatchesGroup, TethersDependency);
+		TetherEnd = &ManagedArrayCollection->AddAttribute<int32>(TetherEndAttribute, TetherBatchesGroup, TethersDependency);
 
-		// Render Faces Group
-		AddExternalAttribute<FIntVector3>("RenderIndices", RenderFacesGroup, RenderIndices, RenderVerticesDependency);
-		AddExternalAttribute<int32>("RenderMaterialIndex", RenderFacesGroup, RenderMaterialIndex);
+		// Tethers Group
+		TetherKinematicIndex = &ManagedArrayCollection->AddAttribute<int32>(TetherKinematicIndexAttribute, TethersGroup, SimVerticesDependency);
+		TetherDynamicIndex = &ManagedArrayCollection->AddAttribute<int32>(TetherDynamicIndexAttribute, TethersGroup, SimVerticesDependency);
+		TetherReferenceLength = &ManagedArrayCollection->AddAttribute<float>(TetherReferenceLengthAttribute, TethersGroup);
 
-		// TODO: FMeshToMeshVertData
+		// Seams Group
+		SeamPatterns = &ManagedArrayCollection->AddAttribute<FIntVector2>(SeamPatternsAttribute, SeamsGroup, SimVerticesDependency);
+		SeamStitches = &ManagedArrayCollection->AddAttribute<TArray<FIntVector2>>(SeamStitchesAttribute, SeamsGroup, SimVerticesDependency);
 
 		// Patterns Group
-		AddExternalAttribute<int32>("SimVerticesStart", PatternsGroup, SimVerticesStart, SimVerticesDependency);
-		AddExternalAttribute<int32>("SimVerticesEnd", PatternsGroup, SimVerticesEnd, SimVerticesDependency);
-		AddExternalAttribute<int32>("SimFacesStart", PatternsGroup, SimFacesStart, SimFacesDependency);
-		AddExternalAttribute<int32>("SimFacesEnd", PatternsGroup, SimFacesEnd, SimFacesDependency);
-		AddExternalAttribute<int32>("RenderVerticesStart", PatternsGroup, RenderVerticesStart, RenderVerticesDependency);
-		AddExternalAttribute<int32>("RenderVerticesEnd", PatternsGroup, RenderVerticesEnd, RenderVerticesDependency);
-		AddExternalAttribute<int32>("RenderFacesStart", PatternsGroup, RenderFacesStart, RenderFacesDependency);
-		AddExternalAttribute<int32>("RenderFacesEnd", PatternsGroup, RenderFacesEnd, RenderFacesDependency);
-		AddExternalAttribute<int32>("WrapDeformerStart", PatternsGroup, WrapDeformerStart, WrapDeformersDependency);
-		AddExternalAttribute<int32>("WrapDeformerEnd", PatternsGroup, WrapDeformerEnd, WrapDeformersDependency);
-		AddExternalAttribute<int32>("NumWeights", PatternsGroup, NumWeights);
-		AddExternalAttribute<int32>("StatusFlags", PatternsGroup, StatusFlags);
-		AddExternalAttribute<int32>("SimMaterialIndex", PatternsGroup, SimMaterialIndex);
+		SimVerticesStart = &ManagedArrayCollection->AddAttribute<int32>(SimVerticesStartAttribute, PatternsGroup, SimVerticesDependency);
+		SimVerticesEnd = &ManagedArrayCollection->AddAttribute<int32>(SimVerticesEndAttribute, PatternsGroup, SimVerticesDependency);
+		SimFacesStart = &ManagedArrayCollection->AddAttribute<int32>(SimFacesStartAttribute, PatternsGroup, SimFacesDependency);
+		SimFacesEnd = &ManagedArrayCollection->AddAttribute<int32>(SimFacesEndAttribute, PatternsGroup, SimFacesDependency);
+		RenderVerticesStart = &ManagedArrayCollection->AddAttribute<int32>(RenderVerticesStartAttribute, PatternsGroup, RenderVerticesDependency);
+		RenderVerticesEnd = &ManagedArrayCollection->AddAttribute<int32>(RenderVerticesEndAttribute, PatternsGroup, RenderVerticesDependency);
+		RenderFacesStart = &ManagedArrayCollection->AddAttribute<int32>(RenderFacesStartAttribute, PatternsGroup, RenderFacesDependency);
+		RenderFacesEnd = &ManagedArrayCollection->AddAttribute<int32>(RenderFacesEndAttribute, PatternsGroup, RenderFacesDependency);
+		StatusFlags = &ManagedArrayCollection->AddAttribute<int32>(StatusFlagsAttribute, PatternsGroup);
 
-		// Seams Group
-		AddExternalAttribute<FIntVector2>("SeamPatterns", SeamsGroup, SeamPatterns, SimVerticesDependency);
-		AddExternalAttribute<TArray<FIntVector2>>("SeamStitches", SeamsGroup, SeamStitches, SimVerticesDependency);
+		// Sim Faces Group
+		SimIndices = &ManagedArrayCollection->AddAttribute<FIntVector3>(SimIndicesAttribute, SimFacesGroup, SimVerticesDependency);
 
-		// Tethers Group
-		AddExternalAttribute<int32>("TetherKinematicIndex", TethersGroup, TetherKinematicIndex, SimVerticesDependency);
-		AddExternalAttribute<int32>("TetherDynamicIndex", TethersGroup, TetherDynamicIndex, SimVerticesDependency);
-		AddExternalAttribute<float>("TetherReferenceLength", TethersGroup, TetherReferenceLength);
+		// Sim Vertices Group
+		SimPosition = &ManagedArrayCollection->AddAttribute<FVector2f>(SimPositionAttribute, SimVerticesGroup);
+		SimRestPosition = &ManagedArrayCollection->AddAttribute<FVector3f>(SimRestPositionAttribute, SimVerticesGroup);
+		SimRestNormal = &ManagedArrayCollection->AddAttribute<FVector3f>(SimRestNormalAttribute, SimVerticesGroup);
+		SimNumBoneInfluences = &ManagedArrayCollection->AddAttribute<int32>(SimNumBoneInfluencesAttribute, SimVerticesGroup);
+		SimBoneIndices = &ManagedArrayCollection->AddAttribute<TArray<int32>>(SimBoneIndicesAttribute, SimVerticesGroup);
+		SimBoneWeights = &ManagedArrayCollection->AddAttribute<TArray<float>>(SimBoneWeightsAttribute, SimVerticesGroup);
 
-		// Tether Batches Group
-		AddExternalAttribute<int32>("TetherStart", TetherBatchesGroup, TetherStart, TethersDependency);
-		AddExternalAttribute<int32>("TetherEnd", TetherBatchesGroup, TetherEnd, TethersDependency);
+		// Render Faces Group
+		RenderIndices = &ManagedArrayCollection->AddAttribute<FIntVector3>(RenderIndicesAttribute, RenderFacesGroup, RenderVerticesDependency);
+		RenderMaterialIndex = &ManagedArrayCollection->AddAttribute<int32>(RenderMaterialIndexAttribute, RenderFacesGroup);
 
-		// LODs Group
-		AddExternalAttribute<int32>("PatternStart", LodsGroup, PatternStart, PatternsDependency);
-		AddExternalAttribute<int32>("PatternEnd", LodsGroup, PatternEnd, PatternsDependency);
-		AddExternalAttribute<int32>("StitchingStart", LodsGroup, SeamStart, StitchingsDependency);
-		AddExternalAttribute<int32>("StitchingEnd", LodsGroup, SeamEnd, StitchingsDependency);
-		AddExternalAttribute<int32>("TetherBatchStart", LodsGroup, TetherBatchStart, TetherBatchesDependency);
-		AddExternalAttribute<int32>("TetherBatchEnd", LodsGroup, TetherBatchEnd, TetherBatchesDependency);
-		AddExternalAttribute<int32>("LodBiasDepth", LodsGroup, LodBiasDepth);
+		// Render Vertices Group
+		RenderPosition = &ManagedArrayCollection->AddAttribute<FVector3f>(RenderPositionAttribute, RenderVerticesGroup);
+		RenderNormal = &ManagedArrayCollection->AddAttribute<FVector3f>(RenderNormalAttribute, RenderVerticesGroup);
+		RenderTangentU = &ManagedArrayCollection->AddAttribute<FVector3f>(RenderTangentUAttribute, RenderVerticesGroup);
+		RenderTangentV = &ManagedArrayCollection->AddAttribute<FVector3f>(RenderTangentVAttribute, RenderVerticesGroup);
+		RenderUVs = &ManagedArrayCollection->AddAttribute<TArray<FVector2f>>(RenderUVsAttribute, RenderVerticesGroup);
+		RenderColor = &ManagedArrayCollection->AddAttribute<FLinearColor>(RenderColorAttribute, RenderVerticesGroup);
+		RenderNumBoneInfluences = &ManagedArrayCollection->AddAttribute<int32>(RenderNumBoneInfluencesAttribute, RenderVerticesGroup);
+		RenderBoneIndices = &ManagedArrayCollection->AddAttribute<TArray<int32>>(RenderBoneIndicesAttribute, RenderVerticesGroup);
+		RenderBoneWeights = &ManagedArrayCollection->AddAttribute<TArray<float>>(RenderBoneWeightsAttribute, RenderVerticesGroup);
+	}
 
-		// Materials Group
-		AddExternalAttribute<FString>("MaterialPathName", MaterialsGroup, MaterialPathName);
+	int32 FClothCollection::GetNumElements(const FName& GroupName) const
+	{
+		return ManagedArrayCollection->NumElements(GroupName);
+	}
+
+	int32 FClothCollection::GetNumElements(const TManagedArray<int32>* StartArray, const TManagedArray<int32>* EndArray, int32 ArrayIndex) const
+	{
+		if (StartArray && StartArray->GetConstArray().IsValidIndex(ArrayIndex) &&
+			EndArray && EndArray->GetConstArray().IsValidIndex(ArrayIndex))
+		{
+			const int32 Start = (*StartArray)[ArrayIndex];
+			const int32 End = (*EndArray)[ArrayIndex];
+			if (Start != INDEX_NONE && End != INDEX_NONE)
+			{
+				return End - Start + 1;
+			}
+			checkf(Start == End, TEXT("Only one boundary of the range is set to INDEX_NONE, when both should."));
+		}
+		return 0;
+	}
 
-		// Collisions Group
-		AddExternalAttribute<FString>("PhysicsAssetPathName", CollisionsGroup, PhysicsAssetPathName);
+	void FClothCollection::SetNumElements(int32 InNumElements, const FName& GroupName)
+	{
+		check(IsValid());
+		check(InNumElements >= 0);
 		
-		// Skeletons Group
-		AddExternalAttribute<FString>("SkeletonAssetPathName", SkeletonsGroup, SkeletonAssetPathName);
+		const int32 NumElements = ManagedArrayCollection->NumElements(GroupName);
+
+		if (const int32 Delta = InNumElements - NumElements)
+		{
+			if (Delta > 0)
+			{
+				ManagedArrayCollection->AddElements(Delta, GroupName);
+			}
+			else
+			{
+				ManagedArrayCollection->RemoveElements(GroupName, -Delta, InNumElements);
+			}
+		}
 	}
 
-	int32 FClothCollection::SetNumElements(int32 InNumElements, const FName& GroupName, TManagedArray<int32>& StartArray, TManagedArray<int32>& EndArray, int32 StartEndIndex)
+	int32 FClothCollection::SetNumElements(int32 InNumElements, const FName& GroupName, TManagedArray<int32>* StartArray, TManagedArray<int32>* EndArray, int32 ArrayIndex)
 	{
+		check(IsValid());
 		check(InNumElements >= 0);
 
-		int32& Start = StartArray[StartEndIndex];
-		int32& End = EndArray[StartEndIndex];
+		check(StartArray && StartArray->GetConstArray().IsValidIndex(ArrayIndex));
+		check(EndArray && EndArray->GetConstArray().IsValidIndex(ArrayIndex));
+
+		int32& Start = (*StartArray)[ArrayIndex];
+		int32& End = (*EndArray)[ArrayIndex];
 		check(Start != INDEX_NONE || End == INDEX_NONE);  // Best to avoid situations where only one boundary of the range is set to INDEX_NONE
 
 		const int32 NumElements = (Start == INDEX_NONE) ? 0 : End - Start + 1;
@@ -132,21 +469,21 @@ namespace UE::Chaos::ClothAsset
 			if (Delta > 0)
 			{
 				// Find a previous valid index range to insert after when the range is empty
-				auto ComputeEnd = [&EndArray](int32 StartEndIndex)->int32
+				auto ComputeEnd = [&EndArray, ArrayIndex]()->int32
 				{
-					for (int32 Index = StartEndIndex; Index >= 0; --Index)
+					for (int32 Index = ArrayIndex; Index >= 0; --Index)
 					{
-						if (EndArray[Index] != INDEX_NONE)
+						if ((*EndArray)[Index] != INDEX_NONE)
 						{
-							return EndArray[Index];
+							return (*EndArray)[Index];
 						}
 					}
 					return INDEX_NONE;
 				};
 
 				// Grow the array
-				const int32 Position = ComputeEnd(StartEndIndex) + 1;
-				InsertElements(Delta, Position, GroupName);
+				const int32 Position = ComputeEnd() + 1;
+				ManagedArrayCollection->InsertElements(Delta, Position, GroupName);
 
 				// Update Start/End
 				if (!NumElements)
@@ -159,7 +496,7 @@ namespace UE::Chaos::ClothAsset
 			{
 				// Shrink the array
 				const int32 Position = Start + InNumElements;
-				RemoveElements(GroupName, -Delta, Position);
+				ManagedArrayCollection->RemoveElements(GroupName, -Delta, Position);
 
 				// Update Start/End
 				if (InNumElements)
@@ -175,51 +512,138 @@ namespace UE::Chaos::ClothAsset
 		return Start;
 	}
 
-	int32 FClothCollection::GetNumElements(const TManagedArray<int32>& StartArray, const TManagedArray<int32>& EndArray, int32 StartEndIndex) const
+	int32 FClothCollection::GetNumSubElements(
+		const TManagedArray<int32>* StartArray,
+		const TManagedArray<int32>* EndArray,
+		const TManagedArray<int32>* StartSubArray,
+		const TManagedArray<int32>* EndSubArray,
+		int32 ArrayIndex) const
 	{
-		const int32 Start = StartArray[StartEndIndex];
-		const int32 End = EndArray[StartEndIndex];
-		check(Start != INDEX_NONE || End == INDEX_NONE);  // Best to avoid situations where only one boundary of the range is set to INDEX_NONE
-
-		return Start == INDEX_NONE ? 0 : End - Start + 1;
-	}
-
-	int32 FClothCollection::GetPatternsNumElements(const TManagedArray<int32>& StartArray, const TManagedArray<int32>& EndArray, int32 LodIndex) const
-	{
-		const TTuple<int32, int32> StartEnd = GetPatternsElementsStartEnd(StartArray, EndArray, LodIndex);
+		const TTuple<int32, int32> StartEnd = GetSubElementsStartEnd(StartArray, EndArray, StartSubArray, EndSubArray, ArrayIndex);
 		const int32 Start = StartEnd.Get<0>();
 		const int32 End = StartEnd.Get<1>();
-		check(Start != INDEX_NONE || End == INDEX_NONE);  // Best to avoid situations where only one boundary of the range is set to INDEX_NONE
-
-		return Start == INDEX_NONE ? 0 : End - Start + 1;
+		if (Start != INDEX_NONE && End != INDEX_NONE)
+		{
+			return End - Start + 1;
+		}
+		checkf(Start == End, TEXT("Only one boundary of the range is set to INDEX_NONE, when both should."));
+		return 0;
 	}
 
 	template<bool bStart, bool bEnd>
-	TTuple<int32, int32> FClothCollection::GetPatternsElementsStartEnd(const TManagedArray<int32>& StartArray, const TManagedArray<int32>& EndArray, int32 LodIndex) const
+	TTuple<int32, int32> FClothCollection::GetSubElementsStartEnd(
+		const TManagedArray<int32>* StartArray,
+		const TManagedArray<int32>* EndArray,
+		const TManagedArray<int32>* StartSubArray,
+		const TManagedArray<int32>* EndSubArray,
+		int32 ArrayIndex) const
 	{
-		const int32 LodPatternStart = PatternStart[LodIndex];
-		const int32 LodPatternEnd = PatternEnd[LodIndex];
-
 		int32 Start = INDEX_NONE;  // Find Start and End indices for the entire LOD minding empty patterns on the way
 		int32 End = INDEX_NONE;
 
-		if (LodPatternStart != INDEX_NONE && LodPatternEnd != INDEX_NONE)
+		if (StartArray && StartArray->GetConstArray().IsValidIndex(ArrayIndex) &&
+			EndArray && EndArray->GetConstArray().IsValidIndex(ArrayIndex))
 		{
-			for (int32 PatternIndex = LodPatternStart; PatternIndex <= LodPatternEnd; ++PatternIndex)
+			const int32 SubStart = (*StartArray)[ArrayIndex];
+			const int32 SubEnd = (*EndArray)[ArrayIndex];
+
+			if (SubStart != INDEX_NONE && SubEnd != INDEX_NONE)
 			{
-				if (bStart && StartArray[PatternIndex] != INDEX_NONE)
+				for (int32 SubIndex = SubStart; SubIndex <= SubEnd; ++SubIndex)
 				{
-					Start = (Start == INDEX_NONE) ? StartArray[PatternIndex] : FMath::Min(Start, StartArray[PatternIndex]);
+					if (bStart && (*StartSubArray)[SubIndex] != INDEX_NONE)
+					{
+						Start = (Start == INDEX_NONE) ? (*StartSubArray)[SubIndex] : FMath::Min(Start, (*StartSubArray)[SubIndex]);
+					}
+					if (bEnd && (*EndSubArray)[SubIndex] != INDEX_NONE)
+					{
+						End = (End == INDEX_NONE) ? (*EndSubArray)[SubIndex] : FMath::Max(End, (*EndSubArray)[SubIndex]);
+					}
 				}
-				if (bEnd && EndArray[PatternIndex] != INDEX_NONE)
+			}
+			else
+			{
+				checkf(SubStart == SubEnd, TEXT("Only one boundary of the range is set to INDEX_NONE, when both should."));
+			}
+		}
+		return TTuple<int32, int32>(Start, End);
+	}
+	template CHAOSCLOTHASSET_API TTuple<int32, int32> FClothCollection::GetSubElementsStartEnd<true, false>(const TManagedArray<int32>* StartArray, const TManagedArray<int32>* EndArray, const TManagedArray<int32>* StartSubArray, const TManagedArray<int32>* EndSubArray, int32 ArrayIndex) const;
+	template CHAOSCLOTHASSET_API TTuple<int32, int32> FClothCollection::GetSubElementsStartEnd<false, true>(const TManagedArray<int32>* StartArray, const TManagedArray<int32>* EndArray, const TManagedArray<int32>* StartSubArray, const TManagedArray<int32>* EndSubArray, int32 ArrayIndex) const;
+	template CHAOSCLOTHASSET_API TTuple<int32, int32> FClothCollection::GetSubElementsStartEnd<true, true>(const TManagedArray<int32>* StartArray, const TManagedArray<int32>* EndArray, const TManagedArray<int32>* StartSubArray, const TManagedArray<int32>* EndSubArray, int32 ArrayIndex) const;
+
+	template<typename T, typename TEnableIf<TIsUserAttributeType<T>::Value, int>::type>
+	TArray<FName> FClothCollection::GetUserDefinedAttributeNames(const FName& GroupName) const
+	{
+		using namespace UE::Chaos::ClothAsset::Private;
+
+		TArray<FName> UserDefinedAttributeNames;
+
+		const TArray<FName>& FixedAttributeNames = FixedAttributeNamesMap.FindChecked(GroupName);  // Also checks that the group name is a recognized group name 
+		const TArray<FName> AttributeNames = ManagedArrayCollection->AttributeNames(GroupName);
+
+		const int32 MaxUserDefinedAttributes = AttributeNames.Num() - FixedAttributeNames.Num();
+		if (MaxUserDefinedAttributes > 0)
+		{
+			UserDefinedAttributeNames.Reserve(MaxUserDefinedAttributes);
+
+			for (const FName& AttributeName : AttributeNames)
+			{
+				if (!FixedAttributeNames.Contains(AttributeName) && ManagedArrayCollection->FindAttributeTyped<T>(AttributeName, GroupName))
 				{
-					End = (End == INDEX_NONE) ? EndArray[PatternIndex] : FMath::Max(End, EndArray[PatternIndex]);
+					UserDefinedAttributeNames.Add(AttributeName);
 				}
 			}
 		}
-		return TTuple<int32, int32>(Start, End);
+		return UserDefinedAttributeNames;
+	}
+	template CHAOSCLOTHASSET_API TArray<FName> FClothCollection::GetUserDefinedAttributeNames<bool>(const FName& GroupName) const;
+	template CHAOSCLOTHASSET_API TArray<FName> FClothCollection::GetUserDefinedAttributeNames<int32>(const FName& GroupName) const;
+	template CHAOSCLOTHASSET_API TArray<FName> FClothCollection::GetUserDefinedAttributeNames<float>(const FName& GroupName) const;
+	template CHAOSCLOTHASSET_API TArray<FName> FClothCollection::GetUserDefinedAttributeNames<FVector3f>(const FName& GroupName) const;
+
+	template<typename T, typename TEnableIf<TIsUserAttributeType<T>::Value, int>::type>
+	void FClothCollection::AddUserDefinedAttribute(const FName& Name, const FName& GroupName)
+	{
+		ManagedArrayCollection->AddAttribute<T>(Name, GroupName);
+	}
+	template CHAOSCLOTHASSET_API void FClothCollection::AddUserDefinedAttribute<bool>(const FName& Name, const FName& GroupName);
+	template CHAOSCLOTHASSET_API void FClothCollection::AddUserDefinedAttribute<int32>(const FName& Name, const FName& GroupName);
+	template CHAOSCLOTHASSET_API void FClothCollection::AddUserDefinedAttribute<float>(const FName& Name, const FName& GroupName);
+	template CHAOSCLOTHASSET_API void FClothCollection::AddUserDefinedAttribute<FVector3f>(const FName& Name, const FName& GroupName);
+
+	void FClothCollection::RemoveUserDefinedAttribute(const FName& Name, const FName& GroupName)
+	{
+		ManagedArrayCollection->RemoveAttribute(Name, GroupName);
+	}
+
+	template<typename T, typename TEnableIf<TIsUserAttributeType<T>::Value, int>::type>
+	bool FClothCollection::HasUserDefinedAttribute(const FName& Name, const FName& GroupName) const
+	{
+		return ManagedArrayCollection->FindAttributeTyped<T>(Name, GroupName) != nullptr;
+	}
+	template CHAOSCLOTHASSET_API bool FClothCollection::HasUserDefinedAttribute<bool>(const FName& Name, const FName& GroupName) const;
+	template CHAOSCLOTHASSET_API bool FClothCollection::HasUserDefinedAttribute<int32>(const FName& Name, const FName& GroupName) const;
+	template CHAOSCLOTHASSET_API bool FClothCollection::HasUserDefinedAttribute<float>(const FName& Name, const FName& GroupName) const;
+	template CHAOSCLOTHASSET_API bool FClothCollection::HasUserDefinedAttribute<FVector3f>(const FName& Name, const FName& GroupName) const;
+
+	template<typename T, typename TEnableIf<TIsUserAttributeType<T>::Value, int>::type>
+	const TManagedArray<T>* FClothCollection::GetUserDefinedAttribute(const FName& Name, const FName& GroupName) const
+	{
+		return ManagedArrayCollection->FindAttribute<T>(Name, GroupName);
+	}
+	template CHAOSCLOTHASSET_API const TManagedArray<bool>* FClothCollection::GetUserDefinedAttribute<bool>(const FName& Name, const FName& GroupName) const;
+	template CHAOSCLOTHASSET_API const TManagedArray<int32>* FClothCollection::GetUserDefinedAttribute<int32>(const FName& Name, const FName& GroupName) const;
+	template CHAOSCLOTHASSET_API const TManagedArray<float>* FClothCollection::GetUserDefinedAttribute<float>(const FName& Name, const FName& GroupName) const;
+	template CHAOSCLOTHASSET_API const TManagedArray<FVector3f>* FClothCollection::GetUserDefinedAttribute<FVector3f>(const FName& Name, const FName& GroupName) const;
+
+	template<typename T, typename TEnableIf<TIsUserAttributeType<T>::Value, int>::type>
+	TManagedArray<T>* FClothCollection::GetUserDefinedAttribute(const FName& Name, const FName& GroupName)
+	{
+		return ManagedArrayCollection->FindAttribute<T>(Name, GroupName);
 	}
-	template CHAOSCLOTHASSET_API TTuple<int32, int32> FClothCollection::GetPatternsElementsStartEnd<true, false>(const TManagedArray<int32>& StartArray, const TManagedArray<int32>& EndArray, int32 LodIndex) const;
-	template CHAOSCLOTHASSET_API TTuple<int32, int32> FClothCollection::GetPatternsElementsStartEnd<false, true>(const TManagedArray<int32>& StartArray, const TManagedArray<int32>& EndArray, int32 LodIndex) const;
-	template CHAOSCLOTHASSET_API TTuple<int32, int32> FClothCollection::GetPatternsElementsStartEnd<true, true>(const TManagedArray<int32>& StartArray, const TManagedArray<int32>& EndArray, int32 LodIndex) const;
+	template CHAOSCLOTHASSET_API TManagedArray<bool>* FClothCollection::GetUserDefinedAttribute<bool>(const FName& Name, const FName& GroupName);
+	template CHAOSCLOTHASSET_API TManagedArray<int32>* FClothCollection::GetUserDefinedAttribute<int32>(const FName& Name, const FName& GroupName);
+	template CHAOSCLOTHASSET_API TManagedArray<float>* FClothCollection::GetUserDefinedAttribute<float>(const FName& Name, const FName& GroupName);
+	template CHAOSCLOTHASSET_API TManagedArray<FVector3f>* FClothCollection::GetUserDefinedAttribute<FVector3f>(const FName& Name, const FName& GroupName);
 } // End namespace UE::Chaos::ClothAsset
diff --git a/Engine/Plugins/Experimental/ChaosClothAsset/Source/ChaosClothAsset/Private/ChaosClothAsset/ClothGeometryTools.cpp b/Engine/Plugins/Experimental/ChaosClothAsset/Source/ChaosClothAsset/Private/ChaosClothAsset/ClothGeometryTools.cpp
--- a/Engine/Plugins/Experimental/ChaosClothAsset/Source/ChaosClothAsset/Private/ChaosClothAsset/ClothGeometryTools.cpp
+++ b/Engine/Plugins/Experimental/ChaosClothAsset/Source/ChaosClothAsset/Private/ChaosClothAsset/ClothGeometryTools.cpp
@@ -1,63 +1,143 @@
 // Copyright Epic Games, Inc. All Rights Reserved.
 
 #include "ChaosClothAsset/ClothGeometryTools.h"
-#include "ChaosClothAsset/ClothCollection.h"
-#include "ChaosClothAsset/ClothAdapter.h"
+#include "ChaosClothAsset/CollectionClothFacade.h"
+#include "Math/Vector.h"
 
 namespace UE::Chaos::ClothAsset
 {
-	void FClothGeometryTools::DeleteRenderMesh(const TSharedPtr<FClothCollection>& ClothCollection)
+	// TODO: Move these functions to the cloth facade?
+	bool FClothGeometryTools::HasSimMesh(const TSharedPtr<const FManagedArrayCollection>& ClothCollection)
 	{
-		FClothAdapter ClothAdapter(ClothCollection);
+		bool bHasSimMesh = false;
 
-		for (int32 LodIndex = 0; LodIndex < ClothAdapter.GetNumLods(); ++LodIndex)
+		FCollectionClothConstFacade ClothFacade(ClothCollection);
+
+		for (int32 LodIndex = 0; LodIndex < ClothFacade.GetNumLods(); ++LodIndex)
 		{
-			FClothLodAdapter ClothLodAdapter = ClothAdapter.GetLod(LodIndex);
+			FCollectionClothLodConstFacade ClothLodFacade = ClothFacade.GetLod(LodIndex);
 
-			for (int32 PatternIndex = ClothLodAdapter.GetNumPatterns() - 1; PatternIndex >= 0 ; --PatternIndex)  // Use a reverse order to avoid having to move previous elements
+			for (int32 PatternIndex = ClothLodFacade.GetNumPatterns() - 1; PatternIndex >= 0 ; --PatternIndex)  // Use a reverse order to avoid having to move previous elements
 			{
-				FClothPatternAdapter ClothPatternAdapter = ClothLodAdapter.GetPattern(PatternIndex);
-				ClothPatternAdapter.SetNumRenderVertices(0);
-				ClothPatternAdapter.SetNumRenderFaces(0);
+				FCollectionClothPatternConstFacade ClothPatternFacade = ClothLodFacade.GetPattern(PatternIndex);
+				if (ClothPatternFacade.GetNumSimVertices() && ClothPatternFacade.GetNumSimFaces())
+				{
+					bHasSimMesh = true;
+					break;
+				}
 			}
 		}
-		// TODO: Add Materials functions to cloth adapter
-		ClothCollection->EmptyGroup(FClothCollection::MaterialsGroup);
+		return bHasSimMesh;
 	}
 
-	void FClothGeometryTools::CopySimMeshToRenderMesh(const TSharedPtr<FClothCollection>& ClothCollection, int32 MaterialIndex)
+	bool FClothGeometryTools::HasRenderMesh(const TSharedPtr<const FManagedArrayCollection>& ClothCollection)
 	{
-		FClothAdapter ClothAdapter(ClothCollection);
+		bool bHasRenderMesh = false;
 
-		for (int32 LodIndex = 0; LodIndex < ClothAdapter.GetNumLods(); ++LodIndex)
+		FCollectionClothConstFacade ClothFacade(ClothCollection);
+
+		for (int32 LodIndex = 0; LodIndex < ClothFacade.GetNumLods(); ++LodIndex)
 		{
-			FClothLodAdapter ClothLodAdapter = ClothAdapter.GetLod(LodIndex);
+			FCollectionClothLodConstFacade ClothLodFacade = ClothFacade.GetLod(LodIndex);
 
-			const TConstArrayView<FVector2f> LodSimPosition = ClothLodAdapter.GetPatternsSimPosition();
-			const TConstArrayView<FVector3f> LodSimRestPosition = ClothLodAdapter.GetPatternsSimRestPosition();
+			for (int32 PatternIndex = ClothLodFacade.GetNumPatterns() - 1; PatternIndex >= 0 ; --PatternIndex)  // Use a reverse order to avoid having to move previous elements
+			{
+				FCollectionClothPatternConstFacade ClothPatternFacade = ClothLodFacade.GetPattern(PatternIndex);
+				if (ClothPatternFacade.GetNumRenderVertices() && ClothPatternFacade.GetNumRenderFaces())
+				{
+					bHasRenderMesh = true;
+					break;
+				}
+			}
+		}
+		return bHasRenderMesh;
+	}
 
-			for (int32 PatternIndex = 0; PatternIndex < ClothLodAdapter.GetNumPatterns(); ++PatternIndex)
+	void FClothGeometryTools::SetSkeletonAssetPathName(const TSharedPtr<FManagedArrayCollection>& ClothCollection, const FString& SkeletonAssetPathName)
+	{
+		FCollectionClothFacade ClothFacade(ClothCollection);
+
+		for (int32 LodIndex = 0; LodIndex < ClothFacade.GetNumLods(); ++LodIndex)
+		{
+			FCollectionClothLodFacade ClothLodFacade = ClothFacade.GetLod(LodIndex);
+			ClothLodFacade.SetSkeletonAssetPathName(SkeletonAssetPathName);
+		}
+	}
+
+	void FClothGeometryTools::SetPhysicsAssetPathName(const TSharedPtr<FManagedArrayCollection>& ClothCollection, const FString& PhysicsAssetPathName)
+	{
+		FCollectionClothFacade ClothFacade(ClothCollection);
+
+		for (int32 LodIndex = 0; LodIndex < ClothFacade.GetNumLods(); ++LodIndex)
+		{
+			FCollectionClothLodFacade ClothLodFacade = ClothFacade.GetLod(LodIndex);
+			ClothLodFacade.SetPhysicsAssetPathName(PhysicsAssetPathName);
+		}
+	}
+
+	void FClothGeometryTools::DeleteRenderMesh(const TSharedPtr<FManagedArrayCollection>& ClothCollection)
+	{
+		FCollectionClothFacade ClothFacade(ClothCollection);
+
+		for (int32 LodIndex = 0; LodIndex < ClothFacade.GetNumLods(); ++LodIndex)
+		{
+			FCollectionClothLodFacade ClothLodFacade = ClothFacade.GetLod(LodIndex);
+
+			for (int32 PatternIndex = ClothLodFacade.GetNumPatterns() - 1; PatternIndex >= 0 ; --PatternIndex)  // Use a reverse order to avoid having to move previous elements
+			{
+				FCollectionClothPatternFacade ClothPatternFacade = ClothLodFacade.GetPattern(PatternIndex);
+				ClothPatternFacade.SetNumRenderVertices(0);
+				ClothPatternFacade.SetNumRenderFaces(0);
+			}
+
+			ClothLodFacade.SetNumMaterials(0);
+		}
+	}
+
+	void FClothGeometryTools::CopySimMeshToRenderMesh(const TSharedPtr<FManagedArrayCollection>& ClothCollection, const FString& RenderMaterialPathName)
+	{
+		FCollectionClothFacade ClothFacade(ClothCollection);
+
+		for (int32 LodIndex = 0; LodIndex < ClothFacade.GetNumLods(); ++LodIndex)
+		{
+			FCollectionClothLodFacade ClothLodFacade = ClothFacade.GetLod(LodIndex);
+
+			// Add the material if it doesn't already exists in this LOD
+			const TConstArrayView<FString> RenderMaterialPathNameArray = ClothLodFacade.GetRenderMaterialPathName();
+			int32 MaterialIndex = RenderMaterialPathNameArray.Find(RenderMaterialPathName);
+			if (MaterialIndex == INDEX_NONE)
 			{
-				FClothPatternAdapter ClothPatternAdapter = ClothLodAdapter.GetPattern(PatternIndex);
+				MaterialIndex = RenderMaterialPathNameArray.Num();
+				ClothLodFacade.SetNumMaterials(MaterialIndex + 1);
+				ClothLodFacade.GetRenderMaterialPathName()[MaterialIndex] = RenderMaterialPathName;
+			}
+
+			// Copy sim data to render data
+			const TConstArrayView<FVector2f> LodSimPosition = ClothLodFacade.GetSimPosition();
+			const TConstArrayView<FVector3f> LodSimRestPosition = ClothLodFacade.GetSimRestPosition();
 
-				const int32 NumVertices = ClothPatternAdapter.GetNumSimVertices();
-				const int32 NumFaces = ClothPatternAdapter.GetNumSimFaces();
+			for (int32 PatternIndex = 0; PatternIndex < ClothLodFacade.GetNumPatterns(); ++PatternIndex)
+			{
+				FCollectionClothPatternFacade ClothPatternFacade = ClothLodFacade.GetPattern(PatternIndex);
+
+				const int32 NumVertices = ClothPatternFacade.GetNumSimVertices();
+				const int32 NumFaces = ClothPatternFacade.GetNumSimFaces();
 
 				if (!NumVertices || !NumFaces)
 				{
-					ClothPatternAdapter.SetNumRenderVertices(0);
-					ClothPatternAdapter.SetNumRenderFaces(0);
+					ClothPatternFacade.SetNumRenderVertices(0);
+					ClothPatternFacade.SetNumRenderFaces(0);
 					continue;
 				}
 
-				ClothPatternAdapter.SetNumRenderVertices(NumVertices);
-				ClothPatternAdapter.SetNumRenderFaces(NumFaces);
+				ClothPatternFacade.SetNumRenderVertices(NumVertices);
+				ClothPatternFacade.SetNumRenderFaces(NumFaces);
 
-				const TConstArrayView<FVector2f> SimPosition = ClothPatternAdapter.GetSimPosition();
-				const TConstArrayView<FVector3f> SimRestPosition = ClothPatternAdapter.GetSimRestPosition();
-				const TConstArrayView<FVector3f> SimRestNormal = ClothPatternAdapter.GetSimRestNormal();
-				const TArrayView<FVector3f> RenderTangentU = ClothPatternAdapter.GetRenderTangentU();
-				const TArrayView<FVector3f> RenderTangentV = ClothPatternAdapter.GetRenderTangentV();
+				const TConstArrayView<FVector2f> SimPosition = ClothPatternFacade.GetSimPosition();
+				const TConstArrayView<FVector3f> SimRestPosition = ClothPatternFacade.GetSimRestPosition();
+				const TConstArrayView<FVector3f> SimRestNormal = ClothPatternFacade.GetSimRestNormal();
+				const TArrayView<FVector3f> RenderTangentU = ClothPatternFacade.GetRenderTangentU();
+				const TArrayView<FVector3f> RenderTangentV = ClothPatternFacade.GetRenderTangentV();
 
 				FVector2f MinPosition(TNumericLimits<float>::Max());
 				FVector2f MaxPosition(TNumericLimits<float>::Lowest());
@@ -77,11 +157,11 @@ namespace UE::Chaos::ClothAsset
 					UVScale.Y < SMALL_NUMBER ? 0.f : 1.f / UVScale.Y);
 
 				// Face group
-				const TConstArrayView<FIntVector3> SimIndices = ClothPatternAdapter.GetSimIndices();
-				const TArrayView<FIntVector3> RenderIndices = ClothPatternAdapter.GetRenderIndices();
-				const TArrayView<int32> RenderMaterialIndex = ClothPatternAdapter.GetRenderMaterialIndex();
-				const TArrayView<FVector3f> LodRenderTangentU = ClothLodAdapter.GetPatternsRenderTangentU();
-				const TArrayView<FVector3f> LodRenderTangentV = ClothLodAdapter.GetPatternsRenderTangentV();
+				const TConstArrayView<FIntVector3> SimIndices = ClothPatternFacade.GetSimIndices();
+				const TArrayView<FIntVector3> RenderIndices = ClothPatternFacade.GetRenderIndices();
+				const TArrayView<int32> RenderMaterialIndex = ClothPatternFacade.GetRenderMaterialIndex();
+				const TArrayView<FVector3f> LodRenderTangentU = ClothLodFacade.GetRenderTangentU();
+				const TArrayView<FVector3f> LodRenderTangentV = ClothLodFacade.GetRenderTangentV();
 				for (int32 FaceIndex = 0; FaceIndex < NumFaces; ++FaceIndex)
 				{
 					const FIntVector3& Face = SimIndices[FaceIndex];
@@ -89,19 +169,15 @@ namespace UE::Chaos::ClothAsset
 					RenderIndices[FaceIndex] = Face;
 					RenderMaterialIndex[FaceIndex] = MaterialIndex;
 
-					// Edges of the triangle : position delta
-					const FVector3f DeltaPos1 = LodSimRestPosition[Face[1]] - LodSimRestPosition[Face[0]];
-					const FVector3f DeltaPos2 = LodSimRestPosition[Face[2]] - LodSimRestPosition[Face[0]];
-
-					// UV delta
-					const FVector2f DeltaUV1 = LodSimPosition[Face[1]] - LodSimPosition[Face[0]];
-					const FVector2f DeltaUV2 = LodSimPosition[Face[2]] - LodSimPosition[Face[0]];
+					const FVector3f Pos01 = LodSimRestPosition[Face[1]] - LodSimRestPosition[Face[0]];
+					const FVector3f Pos02 = LodSimRestPosition[Face[2]] - LodSimRestPosition[Face[0]];
+					const FVector2f UV01 = LodSimPosition[Face[1]] - LodSimPosition[Face[0]];
+					const FVector2f UV02 = LodSimPosition[Face[2]] - LodSimPosition[Face[0]];
 					
-					// We can now use our formula to compute the tangent and the bitangent :
-					const float Denom = DeltaUV1.X * DeltaUV2.Y - DeltaUV1.Y * DeltaUV2.X;
-					const float R = (Denom < SMALL_NUMBER) ? 0.f : 1.f / Denom;
-					const FVector3f TangentU = (DeltaPos1 * DeltaUV2.Y - DeltaPos2 * DeltaUV1.Y) * R;
-					const FVector3f TangentV = (DeltaPos2 * DeltaUV1.X - DeltaPos1 * DeltaUV2.X) * R;
+					const float Denom = UV01.X * UV02.Y - UV01.Y * UV02.X;
+					const float InvDenom = (Denom < SMALL_NUMBER) ? 0.f : 1.f / Denom;
+					const FVector3f TangentU = (Pos01 * UV02.Y - Pos02 * UV01.Y) * InvDenom;
+					const FVector3f TangentV = (Pos02 * UV01.X - Pos01 * UV02.X) * InvDenom;
 
 					for (int32 PointIndex = 0; PointIndex < 3; ++PointIndex)
 					{
@@ -111,16 +187,16 @@ namespace UE::Chaos::ClothAsset
 				}
 
 				// Vertex group
-				const TArrayView<FVector3f> RenderPosition = ClothPatternAdapter.GetRenderPosition();
-				const TArrayView<FVector3f> RenderNormal = ClothPatternAdapter.GetRenderNormal();
-				const TArrayView<TArray<FVector2f>> RenderUVs = ClothPatternAdapter.GetRenderUVs();
-				const TArrayView<FLinearColor> RenderColor = ClothPatternAdapter.GetRenderColor();
-				const TArrayView<int32> RenderNumBoneInfluences = ClothPatternAdapter.GetRenderNumBoneInfluences();
-				const TArrayView<TArray<int32>> RenderBoneIndices = ClothPatternAdapter.GetRenderBoneIndices();
-				const TArrayView<TArray<float>> RenderBoneWeights = ClothPatternAdapter.GetRenderBoneWeights();
-				const TArrayView<int32> SimNumBoneInfluences = ClothPatternAdapter.GetSimNumBoneInfluences();
-				const TArrayView<TArray<int32>> SimBoneIndices = ClothPatternAdapter.GetSimBoneIndices();
-				const TArrayView<TArray<float>> SimBoneWeights = ClothPatternAdapter.GetSimBoneWeights();
+				const TArrayView<FVector3f> RenderPosition = ClothPatternFacade.GetRenderPosition();
+				const TArrayView<FVector3f> RenderNormal = ClothPatternFacade.GetRenderNormal();
+				const TArrayView<TArray<FVector2f>> RenderUVs = ClothPatternFacade.GetRenderUVs();
+				const TArrayView<FLinearColor> RenderColor = ClothPatternFacade.GetRenderColor();
+				const TArrayView<int32> RenderNumBoneInfluences = ClothPatternFacade.GetRenderNumBoneInfluences();
+				const TArrayView<TArray<int32>> RenderBoneIndices = ClothPatternFacade.GetRenderBoneIndices();
+				const TArrayView<TArray<float>> RenderBoneWeights = ClothPatternFacade.GetRenderBoneWeights();
+				const TArrayView<int32> SimNumBoneInfluences = ClothPatternFacade.GetSimNumBoneInfluences();
+				const TArrayView<TArray<int32>> SimBoneIndices = ClothPatternFacade.GetSimBoneIndices();
+				const TArrayView<TArray<float>> SimBoneWeights = ClothPatternFacade.GetSimBoneWeights();
 
 				for (int32 VertexIndex = 0; VertexIndex < NumVertices; ++VertexIndex)
 				{
@@ -139,7 +215,7 @@ namespace UE::Chaos::ClothAsset
 	}
 
 	void FClothGeometryTools::ReverseNormals(
-		const TSharedPtr<FClothCollection>& ClothCollection,
+		const TSharedPtr<FManagedArrayCollection>& ClothCollection,
 		bool bReverseSimMeshNormals,
 		bool bReverseRenderMeshNormals,
 		const TArray<int32>& PatternSelection)
@@ -161,38 +237,38 @@ namespace UE::Chaos::ClothAsset
 			}
 		};
 
-		FClothAdapter ClothAdapter(ClothCollection);
+		FCollectionClothFacade ClothFacade(ClothCollection);
 
-		for (int32 LodIndex = 0; LodIndex < ClothAdapter.GetNumLods(); ++LodIndex)
+		for (int32 LodIndex = 0; LodIndex < ClothFacade.GetNumLods(); ++LodIndex)
 		{
-			FClothLodAdapter ClothLodAdapter = ClothAdapter.GetLod(LodIndex);
+			FCollectionClothLodFacade ClothLodFacade = ClothFacade.GetLod(LodIndex);
 
 			if (PatternSelection.IsEmpty())
 			{
 				if (bReverseSimMeshNormals)
 				{
-					ReverseSimNormals(ClothLodAdapter.GetPatternsSimRestNormal());
+					ReverseSimNormals(ClothLodFacade.GetSimRestNormal());
 				}
 				if (bReverseRenderMeshNormals)
 				{
-					ReverseRenderNormals(ClothLodAdapter.GetPatternsRenderNormal(), ClothLodAdapter.GetPatternsRenderTangentU());
+					ReverseRenderNormals(ClothLodFacade.GetRenderNormal(), ClothLodFacade.GetRenderTangentU());
 				}
 			}
 			else
 			{
-				for (int32 PatternIndex = 0; PatternIndex < ClothLodAdapter.GetNumPatterns(); ++PatternIndex)
+				for (int32 PatternIndex = 0; PatternIndex < ClothLodFacade.GetNumPatterns(); ++PatternIndex)
 				{
 					if (PatternSelection.Find(PatternIndex) != INDEX_NONE)
 					{
-						FClothPatternAdapter ClothPatternAdapter = ClothLodAdapter.GetPattern(PatternIndex);
+						FCollectionClothPatternFacade ClothPatternFacade = ClothLodFacade.GetPattern(PatternIndex);
 
 						if (bReverseSimMeshNormals)
 						{
-							ReverseSimNormals(ClothPatternAdapter.GetSimRestNormal());
+							ReverseSimNormals(ClothPatternFacade.GetSimRestNormal());
 						}
 						if (bReverseRenderMeshNormals)
 						{
-							ReverseRenderNormals(ClothPatternAdapter.GetRenderNormal(), ClothPatternAdapter.GetRenderTangentU());
+							ReverseRenderNormals(ClothPatternFacade.GetRenderNormal(), ClothPatternFacade.GetRenderTangentU());
 						}
 					}
 				}
@@ -200,24 +276,25 @@ namespace UE::Chaos::ClothAsset
 		}
 	}
 
-	void FClothGeometryTools::BindMeshToRootBone(const TSharedPtr<FClothCollection>& ClothCollection,
-												 bool bBindSimMesh,
-												 bool bBindRenderMesh,
-												 const TArray<int32> Lods)
+	void FClothGeometryTools::BindMeshToRootBone(
+		const TSharedPtr<FManagedArrayCollection>& ClothCollection,
+		bool bBindSimMesh,
+		bool bBindRenderMesh,
+		const TArray<int32> Lods)
 	{
 		if (!bBindSimMesh && !bBindRenderMesh)
 		{
 			return;
 		}
 		
-		FClothAdapter ClothAdapter(ClothCollection);
+		FCollectionClothFacade ClothFacade(ClothCollection);
 
 		TArray<int32> LodsToBind;
 		if (Lods.IsEmpty())
 		{
-			LodsToBind.Reserve(ClothAdapter.GetNumLods());
+			LodsToBind.Reserve(ClothFacade.GetNumLods());
 			
-			for (int32 LodIndex = 0; LodIndex < ClothAdapter.GetNumLods(); ++LodIndex)
+			for (int32 LodIndex = 0; LodIndex < ClothFacade.GetNumLods(); ++LodIndex)
 			{
 				LodsToBind.Add(LodIndex);
 			}
@@ -229,7 +306,7 @@ namespace UE::Chaos::ClothAsset
 			for (const int32 Lod : Lods)
 			{	
 				// Make sure the Lod indices are valid.
-				if (Lod < ClothAdapter.GetNumLods())
+				if (Lod < ClothFacade.GetNumLods())
 				{
 					LodsToBind.Add(Lod);
 				}
@@ -238,14 +315,14 @@ namespace UE::Chaos::ClothAsset
 
 		for (const int32 LodIndex : LodsToBind)
 		{
-			FClothLodAdapter ClothLodAdapter = ClothAdapter.GetLod(LodIndex);
+			FCollectionClothLodFacade ClothLodFacade = ClothFacade.GetLod(LodIndex);
 			
 			if (bBindSimMesh)
 			{	
-				const int32 NumVertices = ClothLodAdapter.GetPatternsNumSimVertices();
-				TArrayView<int32> NumBoneInfluences = ClothLodAdapter.GetPatternsSimNumBoneInfluences();
-				TArrayView<TArray<int32>> BoneIndices = ClothLodAdapter.GetPatternsSimBoneIndices();
-				TArrayView<TArray<float>> BoneWeights = ClothLodAdapter.GetPatternsSimBoneWeights();
+				const int32 NumVertices = ClothLodFacade.GetNumSimVertices();
+				TArrayView<int32> NumBoneInfluences = ClothLodFacade.GetSimNumBoneInfluences();
+				TArrayView<TArray<int32>> BoneIndices = ClothLodFacade.GetSimBoneIndices();
+				TArrayView<TArray<float>> BoneWeights = ClothLodFacade.GetSimBoneWeights();
 
 				for (int32 VertexIndex = 0; VertexIndex < NumVertices; ++VertexIndex)
 				{
@@ -257,10 +334,10 @@ namespace UE::Chaos::ClothAsset
 
 			if (bBindRenderMesh)
 			{
-				const int32 NumVertices = ClothLodAdapter.GetPatternsNumRenderVertices();
-				TArrayView<int32> NumBoneInfluences = ClothLodAdapter.GetPatternsRenderNumBoneInfluences();
-				TArrayView<TArray<int32>> BoneIndices = ClothLodAdapter.GetPatternsRenderBoneIndices();
-				TArrayView<TArray<float>> BoneWeights = ClothLodAdapter.GetPatternsRenderBoneWeights();
+				const int32 NumVertices = ClothLodFacade.GetNumRenderVertices();
+				TArrayView<int32> NumBoneInfluences = ClothLodFacade.GetRenderNumBoneInfluences();
+				TArrayView<TArray<int32>> BoneIndices = ClothLodFacade.GetRenderBoneIndices();
+				TArrayView<TArray<float>> BoneWeights = ClothLodFacade.GetRenderBoneWeights();
 
 				for (int32 VertexIndex = 0; VertexIndex < NumVertices; ++VertexIndex)
 				{
diff --git a/Engine/Plugins/Experimental/ChaosClothAsset/Source/ChaosClothAsset/Private/ChaosClothAsset/ClothLodAdapter.cpp b/Engine/Plugins/Experimental/ChaosClothAsset/Source/ChaosClothAsset/Private/ChaosClothAsset/ClothLodAdapter.cpp
deleted file mode 100644
diff --git a/Engine/Plugins/Experimental/ChaosClothAsset/Source/ChaosClothAsset/Private/ChaosClothAsset/ClothPatternAdapter.cpp b/Engine/Plugins/Experimental/ChaosClothAsset/Source/ChaosClothAsset/Private/ChaosClothAsset/ClothPatternAdapter.cpp
deleted file mode 100644
diff --git a/Engine/Plugins/Experimental/ChaosClothAsset/Source/ChaosClothAsset/Private/ChaosClothAsset/CollectionClothFacade.cpp b/Engine/Plugins/Experimental/ChaosClothAsset/Source/ChaosClothAsset/Private/ChaosClothAsset/CollectionClothFacade.cpp
new file mode 100644
--- /dev/null
+++ b/Engine/Plugins/Experimental/ChaosClothAsset/Source/ChaosClothAsset/Private/ChaosClothAsset/CollectionClothFacade.cpp
@@ -0,0 +1,112 @@
+// Copyright Epic Games, Inc. All Rights Reserved.
+
+#include "ChaosClothAsset/CollectionClothFacade.h"
+#include "ChaosClothAsset/CollectionClothLodFacade.h"
+#include "ChaosClothAsset/ClothCollection.h"
+
+namespace UE::Chaos::ClothAsset
+{
+	FCollectionClothConstFacade::FCollectionClothConstFacade(const TSharedPtr<const FManagedArrayCollection>& ManagedArrayCollection)
+		: ClothCollection(MakeShared<const FClothCollection>(ConstCastSharedPtr<FManagedArrayCollection>(ManagedArrayCollection)))
+	{
+	}
+
+	bool FCollectionClothConstFacade::IsValid() const
+	{
+		return ClothCollection->IsValid();
+	}
+
+	FCollectionClothLodConstFacade FCollectionClothConstFacade::GetLod(int32 LodIndex) const
+	{
+		check(LodIndex >= 0 && LodIndex < GetNumLods());
+		return FCollectionClothLodConstFacade(ClothCollection, LodIndex);
+	}
+
+	int32 FCollectionClothConstFacade::GetNumLods() const
+	{
+		return ClothCollection->GetNumElements(FClothCollection::LodsGroup);
+	}
+
+	bool FCollectionClothConstFacade::HasWeightMap(const FName& Name) const
+	{
+		return ClothCollection->HasUserDefinedAttribute<float>(Name, FClothCollection::SimVerticesGroup);
+	}
+
+	TArray<FName> FCollectionClothConstFacade::GetWeightMapNames() const
+	{
+		return ClothCollection->GetUserDefinedAttributeNames<float>(FClothCollection::SimVerticesGroup);
+	}
+
+	FCollectionClothFacade::FCollectionClothFacade(const TSharedPtr<FManagedArrayCollection>& ManagedArrayCollection)
+		: FCollectionClothConstFacade(ManagedArrayCollection)
+	{
+	}
+
+	void FCollectionClothFacade::DefineSchema()
+	{
+		GetClothCollection()->DefineSchema();
+	}
+
+	void FCollectionClothFacade::Reset()
+	{
+		check(IsValid());
+
+		const int32 NumLods = GetNumLods();
+		for (int32 LodIndex = 0; LodIndex < NumLods; ++LodIndex)
+		{
+			GetLod(LodIndex).Reset();
+		}
+
+		GetClothCollection()->SetNumElements(0, FClothCollection::LodsGroup);
+	}
+
+	int32 FCollectionClothFacade::AddLod()
+	{
+		check(IsValid());
+
+		const int32 LodIndex = GetClothCollection()->GetNumElements(FClothCollection::LodsGroup);
+		GetClothCollection()->SetNumElements(LodIndex + 1, FClothCollection::LodsGroup);
+
+		FCollectionClothLodFacade(GetClothCollection(), LodIndex).SetDefaults();
+
+		return LodIndex;
+	}
+
+	FCollectionClothLodFacade FCollectionClothFacade::GetLod(int32 LodIndex)
+	{
+		check(IsValid());
+		check(LodIndex < GetNumLods());
+		return FCollectionClothLodFacade(GetClothCollection(), LodIndex);
+	}
+
+	void FCollectionClothFacade::SetNumLod(int32 InNumLods)
+	{
+		check(IsValid());
+
+		const int32 NumLods = GetNumLods();
+
+		for (int32 LodIndex = InNumLods; LodIndex < NumLods; ++LodIndex)
+		{
+			GetLod(LodIndex).Reset();
+		}
+
+		GetClothCollection()->SetNumElements(InNumLods, FClothCollection::LodsGroup);
+
+		for (int32 LodIndex = NumLods; LodIndex < InNumLods; ++LodIndex)
+		{
+			GetLod(LodIndex).SetDefaults();
+		}
+	}
+
+	void FCollectionClothFacade::AddWeightMap(const FName& Name)
+	{
+		check(IsValid());
+		GetClothCollection()->AddUserDefinedAttribute<float>(Name, FClothCollection::SimVerticesGroup);
+	}
+
+	void FCollectionClothFacade::RemoveWeightMap(const FName& Name)
+	{
+		check(IsValid());
+		GetClothCollection()->RemoveUserDefinedAttribute(Name, FClothCollection::SimVerticesGroup);
+	}
+} // End namespace UE::Chaos::ClothAsset
diff --git a/Engine/Plugins/Experimental/ChaosClothAsset/Source/ChaosClothAsset/Private/ChaosClothAsset/CollectionClothLodFacade.cpp b/Engine/Plugins/Experimental/ChaosClothAsset/Source/ChaosClothAsset/Private/ChaosClothAsset/CollectionClothLodFacade.cpp
new file mode 100644
--- /dev/null
+++ b/Engine/Plugins/Experimental/ChaosClothAsset/Source/ChaosClothAsset/Private/ChaosClothAsset/CollectionClothLodFacade.cpp
@@ -0,0 +1,1255 @@
+// Copyright Epic Games, Inc. All Rights Reserved.
+
+#include "ChaosClothAsset/CollectionClothLodFacade.h"
+#include "ChaosClothAsset/CollectionClothFacade.h"
+#include "ChaosClothAsset/ClothCollection.h"
+
+// Utility functions to unwrap a 3d sim mesh into a tailored cloth
+namespace UE::Chaos::ClothAsset::Private
+{
+	// Triangle islands to become patterns, although in this case all the seams are internal (same pattern)
+	struct FIsland
+	{
+		TArray<uint32> Indices;  // 3x number of triangles
+		TArray<FVector2f> Positions;
+		TArray<FVector3f> RestPositions;  // Same size as Positions
+		TArray<uint32> SourceIndices;  // Index in the original welded position array
+	};
+
+	enum class EIntersectCirclesResult
+	{
+		SingleIntersect,
+		DoubleIntersect,
+		Coincident,
+		Separate,
+		Contained
+	};
+
+	static EIntersectCirclesResult IntersectCircles(const FVector2f& C0, float R0, const FVector2f& C1, float R1, FVector2f& OutI0, FVector2f& OutI1)
+	{
+		const FVector2f C0C1 = C0 - C1;
+		const float D = C0C1.Length();
+		if (D < SMALL_NUMBER)
+		{
+			return EIntersectCirclesResult::Coincident;
+		}
+		else if (D > R0 + R1)
+		{
+			return EIntersectCirclesResult::Separate;
+		}
+		else if (D < FMath::Abs(R0 - R1))
+		{
+			return EIntersectCirclesResult::Contained;
+		}
+		const float SquareR0 = R0 * R0;
+		const float SquareR1 = R1 * R1;
+		const float SquareD = D * D;
+		const float A = (SquareD - SquareR1 + SquareR0) / (2.f * D);
+
+		OutI0 = OutI1 = C0 + A * (C1 - C0) / D;
+
+		if (FMath::Abs(A - R0) < SMALL_NUMBER)
+		{
+			return EIntersectCirclesResult::SingleIntersect;
+		}
+
+		const float SquareA = A * A;
+		const float H = FMath::Sqrt(SquareR0 - SquareA);
+
+		const FVector2f N(C0C1.Y, -C0C1.X);
+
+		OutI0 += N * H / D;
+		OutI1 -= N * H / D;
+
+		return EIntersectCirclesResult::DoubleIntersect;
+	}
+	;
+	static FUintVector2 MakeSortedUintVector2(uint32 Index0, uint32 Index1)
+	{
+		return Index0 < Index1 ? FUintVector2(Index0, Index1) : FUintVector2(Index1, Index0);
+	}
+
+	static void BuildEdgeMap(const TArray<uint32>& Indices, TMap<FUintVector2, TArray<uint32>>& OutEdgeToTriangles)
+	{
+		ensure(Indices.Num() % 3 == 0);
+		const uint32 NumTriangles = (uint32)Indices.Num() / 3;
+		OutEdgeToTriangles.Empty(NumTriangles * 2);  // Rough estimate for the number of edges
+		if (!NumTriangles)
+		{
+			return;
+		}
+
+		for (uint32 Triangle = 0; Triangle < NumTriangles; ++Triangle)
+		{
+			const uint32 Index0 = Indices[Triangle * 3 + 0];
+			const uint32 Index1 = Indices[Triangle * 3 + 1];
+			const uint32 Index2 = Indices[Triangle * 3 + 2];
+
+			const FUintVector2 Edge0 = MakeSortedUintVector2(Index0, Index1);
+			const FUintVector2 Edge1 = MakeSortedUintVector2(Index1, Index2);
+			const FUintVector2 Edge2 = MakeSortedUintVector2(Index2, Index0);
+
+			OutEdgeToTriangles.FindOrAdd(Edge0).Add(Triangle);
+			OutEdgeToTriangles.FindOrAdd(Edge1).Add(Triangle);
+			OutEdgeToTriangles.FindOrAdd(Edge2).Add(Triangle);
+		}
+	}
+
+	static void UnwrapMesh(const TArray<FVector3f>& Positions, const TArray<uint32>& Indices, TArray<FIsland>& OutIslands)
+	{
+		OutIslands.Reset();
+
+		ensure(Indices.Num() % 3 == 0);
+		const uint32 NumTriangles = (uint32)Indices.Num() / 3;
+		if (!NumTriangles)
+		{
+			return;
+		}
+
+		// Gather edge information
+		TMap<FUintVector2, TArray<uint32>> EdgeToTrianglesMap;
+		BuildEdgeMap(Indices, EdgeToTrianglesMap);
+
+		// Build pattern islands
+		TSet<uint32> VisitedTriangles;
+		VisitedTriangles.Reserve(NumTriangles);
+
+		constexpr float SquaredWeldingDistance = FMath::Square(0.01f);  // 0.1 mm
+
+		for (uint32 SeedTriangle = 0; SeedTriangle < NumTriangles; ++SeedTriangle)
+		{
+			if (VisitedTriangles.Contains(SeedTriangle))
+			{
+				continue;
+			}
+
+			const uint32 SeedIndex0 = Indices[SeedTriangle * 3 + 0];
+			const uint32 SeedIndex1 = Indices[SeedTriangle * 3 + 1];
+
+			if (FVector3f::DistSquared(Positions[SeedIndex0], Positions[SeedIndex1]) <= SquaredWeldingDistance)
+			{
+				continue;  // A degenerated triangle edge is not a good start
+			}
+
+			// Setup first visitor from seed, and add the first two points
+			FIsland& Island = OutIslands.AddDefaulted_GetRef();
+
+			Island.RestPositions.Add(Positions[SeedIndex0]);
+			Island.RestPositions.Add(Positions[SeedIndex1]);
+
+			const uint32 SeedIndex2D0 = Island.Positions.Add(FVector2f::ZeroVector);
+			const uint32 SeedIndex2D1 = Island.Positions.Add(FVector2f(FVector3f::Dist(Positions[SeedIndex0], Positions[SeedIndex1]), 0.f));
+
+			struct FVisitor
+			{
+				uint32 Triangle;
+				FUintVector2 OldEdge;
+				FUintVector2 NewEdge;
+				uint32 CrossEdgePoint;  // Keep the opposite point to orientate degenerate cases
+			} Visitor =
+			{
+				SeedTriangle,
+				FUintVector2(SeedIndex0, SeedIndex1),
+				FUintVector2(SeedIndex2D0, SeedIndex2D1),
+				(uint32)INDEX_NONE
+			};
+
+			VisitedTriangles.Add(SeedTriangle);
+
+			TQueue<FVisitor> Visitors;
+			do
+			{
+				const uint32 Triangle = Visitor.Triangle;
+				const uint32 CrossEdgePoint = Visitor.CrossEdgePoint;
+				const uint32 OldIndex0 = Visitor.OldEdge.X;
+				const uint32 OldIndex1 = Visitor.OldEdge.Y;
+				const uint32 NewIndex0 = Visitor.NewEdge.X;
+				const uint32 NewIndex1 = Visitor.NewEdge.Y;
+
+				// Find opposite index from this triangle edge
+				const uint32 TriangleIndex0 = Indices[Triangle * 3 + 0];
+				const uint32 TriangleIndex1 = Indices[Triangle * 3 + 1];
+				const uint32 TriangleIndex2 = Indices[Triangle * 3 + 2];
+
+				const uint32 OldIndex2 =
+					(OldIndex0 != TriangleIndex0 && OldIndex1 != TriangleIndex0) ? TriangleIndex0 :
+					(OldIndex0 != TriangleIndex1 && OldIndex1 != TriangleIndex1) ? TriangleIndex1 : TriangleIndex2;
+
+				// Find the 2D intersection of the two connecting adjacent edges using the 3D reference length
+				const FVector3f& P0 = Positions[OldIndex0];
+				const FVector3f& P1 = Positions[OldIndex1];
+				const FVector3f& P2 = Positions[OldIndex2];
+
+				const float R0 = FVector3f::Dist(P0, P2);
+				const float R1 = FVector3f::Dist(P1, P2);
+				const FVector2f C0 = Island.Positions[NewIndex0];
+				const FVector2f C1 = Island.Positions[NewIndex1];
+
+				FVector2f I0, I1;
+				const EIntersectCirclesResult IntersectCirclesResult = IntersectCircles(C0, R0, C1, R1, I0, I1);
+
+				FVector2f C2;
+				switch (IntersectCirclesResult)
+				{
+				case EIntersectCirclesResult::SingleIntersect:
+					C2 = I0;  // Degenerated C2 is on (C0C1)
+					break;
+				case EIntersectCirclesResult::DoubleIntersect:
+					C2 = (FVector2f::CrossProduct(C0 - C1, C0 - I0) > 0) ? I0 : I1;  // Keep correct winding order
+					break;
+				case EIntersectCirclesResult::Coincident:
+					check(CrossEdgePoint != INDEX_NONE); // We can't start on a degenerated triangle
+					C2 = C0 - (Island.Positions[CrossEdgePoint] - C0).GetSafeNormal() * R0;  // Degenerated C0 == C1, choose C2 on the opposite of the visitor opposite point
+					break;
+				case EIntersectCirclesResult::Separate: [[fallthrough]];
+				case EIntersectCirclesResult::Contained:
+					C2 = C0 - (C1 - C0).GetSafeNormal() * R0;  // Degenerated + some tolerance, C2 is on (C0C1)
+					break;
+				}
+
+				// Add the new position found for the opposite point
+				uint32 NewIndex2 = INDEX_NONE;
+				for (uint32 UsedIndex = 0; UsedIndex < (uint32)Island.Positions.Num(); ++UsedIndex)
+				{
+					if (FVector2f::DistSquared(Island.Positions[UsedIndex], C2) <= SquaredWeldingDistance &&
+						FVector3f::DistSquared(Island.RestPositions[UsedIndex], Positions[OldIndex2]) <= SquaredWeldingDistance)
+					{
+						NewIndex2 = UsedIndex;  // Both Rest and 2D positions match, reuse this index
+						break;
+					}
+				}
+				if (NewIndex2 == INDEX_NONE)
+				{
+					NewIndex2 = Island.Positions.Add(C2);
+					Island.RestPositions.Add(Positions[OldIndex2]);
+				}
+
+				// Add triangle to list of indices, unless it is degenerated to a segment
+				if (NewIndex0 != NewIndex1 && NewIndex1 != NewIndex2 && NewIndex2 != NewIndex0)
+				{
+					Island.Indices.Add(NewIndex0);
+					Island.Indices.Add(NewIndex1);
+					Island.Indices.Add(NewIndex2);
+					Island.SourceIndices.Add(OldIndex0);
+					Island.SourceIndices.Add(OldIndex1);
+					Island.SourceIndices.Add(OldIndex2);
+				}
+
+				// Add neighbor triangles to the queue
+				const FUintVector2 OldEdgeList[3] =
+				{
+					FUintVector2(OldIndex1, OldIndex0),  // Reversed as to keep the correct winding order
+					FUintVector2(OldIndex2, OldIndex1),
+					FUintVector2(OldIndex0, OldIndex2)
+				};
+				const FUintVector3 NewEdgeList[3] =
+				{
+					FUintVector3(NewIndex1, NewIndex0, NewIndex2),  // Adds opposite point index
+					FUintVector3(NewIndex2, NewIndex1, NewIndex0),
+					FUintVector3(NewIndex0, NewIndex2, NewIndex1)
+				};
+				for (int32 Edge = 0; Edge < 3; ++Edge)
+				{
+					const uint32 EdgeIndex0 = OldEdgeList[Edge].X;
+					const uint32 EdgeIndex1 = OldEdgeList[Edge].Y;
+
+					const TArray<uint32>& NeighborTriangles = EdgeToTrianglesMap.FindChecked(MakeSortedUintVector2(EdgeIndex0, EdgeIndex1));
+
+					for (const uint32 NeighborTriangle : NeighborTriangles)
+					{
+						if (!VisitedTriangles.Contains(NeighborTriangle))
+						{
+							// Mark neighboring triangle as visited
+							VisitedTriangles.Add(NeighborTriangle);
+
+							// Enqueue next triangle
+							Visitors.Enqueue(FVisitor
+								{
+									NeighborTriangle,
+									FUintVector2(EdgeIndex0, EdgeIndex1),
+									FUintVector2(NewEdgeList[Edge].X, NewEdgeList[Edge].Y),
+									NewEdgeList[Edge].Z,  // Pass the cross edge 2D opposite point to help define orientation of any degenerated triangles
+								});
+						}
+					}
+				}
+			} while (Visitors.Dequeue(Visitor));
+		}
+	}
+
+	struct FSeam
+	{
+		TSet<FIntVector2> Stitches;
+		FIntVector2 Patterns;
+	};
+
+	// Rebuild the seam information from the torn/unwrapped mesh islands data
+	// Note that the isolated mesh islands are not technically patterns despite being considered so, 
+	// since they aren't sewed together in the source welded mesh.
+	// The algorithm will have to be slightly modified to be used with provided UV panels.
+	static void BuildSeams(const TArray<FIsland>& Islands, TArray<FSeam>& OutSeams)
+	{
+		OutSeams.Reset();
+
+		for (int32 IslandIndex = 0; IslandIndex < Islands.Num(); ++IslandIndex)
+		{
+			const FIsland& Island = Islands[IslandIndex];
+
+			FSeam Seam;
+			Seam.Patterns = FIntVector2(IslandIndex);  // Just patching an unwrap here, will only do pattern to same
+
+			// Gather edge information for the source mesh
+			TMap<FUintVector2, TArray<uint32>> SourceEdgeToTrianglesMap;
+			BuildEdgeMap(Island.SourceIndices, SourceEdgeToTrianglesMap);
+
+			auto GetTriangleEdgeMatchingSourceEdge = [&Island](uint32 Triangle, const FUintVector2& SourceEdge) -> FUintVector2
+			{
+				const uint32 TriangleBase = Triangle * 3;
+				const uint32 TriangleIndex0 = Island.SourceIndices[TriangleBase + 0];
+				const uint32 TriangleIndex1 = Island.SourceIndices[TriangleBase + 1];
+				const uint32 TriangleIndex2 = Island.SourceIndices[TriangleBase + 2];
+
+				if (SourceEdge[0] == TriangleIndex0)
+				{
+					if (SourceEdge[1] == TriangleIndex1)
+					{
+						return FUintVector2(Island.Indices[TriangleBase + 0], Island.Indices[TriangleBase + 1]);  // Edge 0 1
+					}
+					else if (SourceEdge[1] == TriangleIndex2)
+					{
+						return FUintVector2(Island.Indices[TriangleBase + 0], Island.Indices[TriangleBase + 2]);  // Edge 0 2
+					}
+				}
+				else if (SourceEdge[0] == TriangleIndex1)
+				{
+					if (SourceEdge[1] == TriangleIndex0)
+					{
+						return FUintVector2(Island.Indices[TriangleBase + 1], Island.Indices[TriangleBase + 0]);  // Edge 1 0
+					}
+					else if (SourceEdge[1] == TriangleIndex2)
+					{
+						return FUintVector2(Island.Indices[TriangleBase + 1], Island.Indices[TriangleBase + 2]);  // Edge 1 2
+					}
+				}
+				else if (SourceEdge[0] == TriangleIndex2)
+				{
+					if (SourceEdge[1] == TriangleIndex0)
+					{
+						return FUintVector2(Island.Indices[TriangleBase + 2], Island.Indices[TriangleBase + 0]);  // Edge 2 0
+					}
+					else if (SourceEdge[1] == TriangleIndex1)
+					{
+						return FUintVector2(Island.Indices[TriangleBase + 2], Island.Indices[TriangleBase + 1]);  // Edge 2 1
+					}
+				}
+				check(false);  // Sanity check
+				return FUintVector2(0);
+			};
+
+			// Look for disconnected triangles
+			for (const TPair<FUintVector2, TArray<uint32>>& SourceEdgeToTriangles : SourceEdgeToTrianglesMap)
+			{
+				const TArray<uint32>& Triangles = SourceEdgeToTriangles.Value;
+				const int32 NumTriangles = Triangles.Num();
+				if (NumTriangles > 1)
+				{
+					const FUintVector2& SourceEdge = SourceEdgeToTriangles.Key;
+
+					for (int32 Index0 = 0; Index0 < NumTriangles - 1; ++Index0)
+					{
+						const uint32 TriangleIndex0 = Triangles[Index0];
+
+						const FUintVector2 Edge0 = GetTriangleEdgeMatchingSourceEdge(TriangleIndex0, SourceEdge);
+						check(Edge0[0] <= (uint32)TNumericLimits<int32>::Max() && Edge0[1] <= (uint32)TNumericLimits<int32>::Max());
+
+						for (int32 Index1 = Index0 + 1; Index1 < NumTriangles; ++Index1)
+						{
+							const uint32 TriangleIndex1 = Triangles[Index1];
+
+							const FUintVector2 Edge1 = GetTriangleEdgeMatchingSourceEdge(TriangleIndex1, SourceEdge);
+							check(Edge1[0] <= (uint32)TNumericLimits<int32>::Max() && Edge1[1] <= (uint32)TNumericLimits<int32>::Max());
+
+							if (Edge0[0] != Edge1[0])
+							{
+								Seam.Stitches.Emplace(FIntVector2(MakeSortedUintVector2(Edge0[0], Edge1[0])));
+							}
+							if (Edge0[1] != Edge1[1])
+							{
+								Seam.Stitches.Emplace(FIntVector2(MakeSortedUintVector2(Edge0[1], Edge1[1])));
+							}
+						}
+					}
+				}
+			}
+
+			// Add this island's seams
+			if (Seam.Stitches.Num())
+			{
+				OutSeams.Emplace(Seam);
+			}
+		}
+	}
+}  // End namespace UE::ClothAsset::Private
+
+namespace UE::Chaos::ClothAsset
+{
+	int32 FCollectionClothLodConstFacade::GetNumMaterials() const
+	{
+		return ClothCollection->GetNumElements(ClothCollection->GetMaterialStart(), ClothCollection->GetMaterialEnd(), LodIndex);
+	}
+
+	int32 FCollectionClothLodConstFacade::GetNumTetherBatches() const
+	{
+		return ClothCollection->GetNumElements(ClothCollection->GetTetherBatchStart(), ClothCollection->GetTetherBatchEnd(), LodIndex);
+	}
+
+	int32 FCollectionClothLodConstFacade::GetNumSeams() const
+	{
+		return ClothCollection->GetNumElements(ClothCollection->GetSeamStart(), ClothCollection->GetSeamEnd(), LodIndex);
+	}
+
+	int32 FCollectionClothLodConstFacade::GetNumPatterns() const
+	{
+		return ClothCollection->GetNumElements(ClothCollection->GetPatternStart(), ClothCollection->GetPatternEnd(), LodIndex);
+	}
+
+	TConstArrayView<FString> FCollectionClothLodConstFacade::GetRenderMaterialPathName() const
+	{
+		return ClothCollection->GetElements(ClothCollection->GetRenderMaterialPathName(), ClothCollection->GetMaterialStart(), ClothCollection->GetMaterialEnd(), LodIndex);
+	}
+
+	TConstArrayView<FIntVector2> FCollectionClothLodConstFacade::GetSeamPatterns() const
+	{
+		return ClothCollection->GetElements(ClothCollection->GetSeamPatterns(), ClothCollection->GetSeamStart(), ClothCollection->GetSeamEnd(), LodIndex);
+	}
+
+	TConstArrayView<TArray<FIntVector2>> FCollectionClothLodConstFacade::GetSeamStitches() const
+	{
+		return ClothCollection->GetElements(ClothCollection->GetSeamStitches(), ClothCollection->GetSeamStart(), ClothCollection->GetSeamEnd(), LodIndex);
+	}
+
+	FCollectionClothPatternConstFacade FCollectionClothLodConstFacade::GetPattern(int32 PatternIndex) const
+	{
+		return FCollectionClothPatternConstFacade(ClothCollection, LodIndex, PatternIndex);
+	}
+
+	const FString& FCollectionClothLodConstFacade::GetPhysicsAssetPathName() const
+	{
+		static const FString EmptyString;
+		return ClothCollection->GetPhysicsAssetPathName() ? (*ClothCollection->GetPhysicsAssetPathName())[LodIndex] : EmptyString;
+	}
+
+	const FString& FCollectionClothLodConstFacade::GetSkeletonAssetPathName() const
+	{
+		static const FString EmptyString;
+		return ClothCollection->GetSkeletonAssetPathName() ? (*ClothCollection->GetSkeletonAssetPathName())[LodIndex] : EmptyString;
+	}
+
+	int32 FCollectionClothLodConstFacade::GetNumSimVertices() const
+	{
+		return ClothCollection->GetNumSubElements(
+			ClothCollection->GetPatternStart(),
+			ClothCollection->GetPatternEnd(),
+			ClothCollection->GetSimVerticesStart(),
+			ClothCollection->GetSimVerticesEnd(),
+			LodIndex);
+	}
+
+	TConstArrayView<FVector2f> FCollectionClothLodConstFacade::GetSimPosition() const
+	{
+		return ClothCollection->GetSubElements(
+			ClothCollection->GetSimPosition(),
+			ClothCollection->GetPatternStart(),
+			ClothCollection->GetPatternEnd(),
+			ClothCollection->GetSimVerticesStart(),
+			ClothCollection->GetSimVerticesEnd(),
+			LodIndex);
+	}
+
+	TConstArrayView<FVector3f> FCollectionClothLodConstFacade::GetSimRestPosition() const
+	{
+		return ClothCollection->GetSubElements(
+			ClothCollection->GetSimRestPosition(),
+			ClothCollection->GetPatternStart(),
+			ClothCollection->GetPatternEnd(),
+			ClothCollection->GetSimVerticesStart(),
+			ClothCollection->GetSimVerticesEnd(),
+			LodIndex);
+	}
+
+	TConstArrayView<FVector3f> FCollectionClothLodConstFacade::GetSimRestNormal() const
+	{
+		return ClothCollection->GetSubElements(
+			ClothCollection->GetSimRestNormal(),
+			ClothCollection->GetPatternStart(),
+			ClothCollection->GetPatternEnd(),
+			ClothCollection->GetSimVerticesStart(),
+			ClothCollection->GetSimVerticesEnd(),
+			LodIndex);
+	}
+
+	TConstArrayView<int32> FCollectionClothLodConstFacade::GetSimNumBoneInfluences() const
+	{
+		return ClothCollection->GetSubElements(
+			ClothCollection->GetSimNumBoneInfluences(),
+			ClothCollection->GetPatternStart(),
+			ClothCollection->GetPatternEnd(),
+			ClothCollection->GetSimVerticesStart(),
+			ClothCollection->GetSimVerticesEnd(),
+			LodIndex);
+	}
+
+	TConstArrayView<TArray<int32>> FCollectionClothLodConstFacade::GetSimBoneIndices() const
+	{
+		return ClothCollection->GetSubElements(
+			ClothCollection->GetSimBoneIndices(),
+			ClothCollection->GetPatternStart(),
+			ClothCollection->GetPatternEnd(),
+			ClothCollection->GetSimVerticesStart(),
+			ClothCollection->GetSimVerticesEnd(),
+			LodIndex);
+	}
+
+	TConstArrayView<TArray<float>> FCollectionClothLodConstFacade::GetSimBoneWeights() const
+	{
+		return ClothCollection->GetSubElements(
+			ClothCollection->GetSimBoneWeights(),
+			ClothCollection->GetPatternStart(),
+			ClothCollection->GetPatternEnd(),
+			ClothCollection->GetSimVerticesStart(),
+			ClothCollection->GetSimVerticesEnd(),
+			LodIndex);
+	}
+
+	int32 FCollectionClothLodConstFacade::GetNumSimFaces() const
+	{
+		return ClothCollection->GetNumSubElements(
+			ClothCollection->GetPatternStart(),
+			ClothCollection->GetPatternEnd(),
+			ClothCollection->GetSimFacesStart(),
+			ClothCollection->GetSimFacesEnd(),
+			LodIndex);
+	}
+
+	TConstArrayView<FIntVector3> FCollectionClothLodConstFacade::GetSimIndices() const
+	{
+		return ClothCollection->GetSubElements(
+			ClothCollection->GetSimIndices(),
+			ClothCollection->GetPatternStart(),
+			ClothCollection->GetPatternEnd(),
+			ClothCollection->GetSimFacesStart(),
+			ClothCollection->GetSimFacesEnd(),
+			LodIndex);
+	}
+
+	int32 FCollectionClothLodConstFacade::GetNumRenderVertices() const
+	{
+		return ClothCollection->GetNumSubElements(
+			ClothCollection->GetPatternStart(),
+			ClothCollection->GetPatternEnd(),
+			ClothCollection->GetRenderVerticesStart(),
+			ClothCollection->GetRenderVerticesEnd(),
+			LodIndex);
+	}
+
+	TConstArrayView<FVector3f> FCollectionClothLodConstFacade::GetRenderPosition() const
+	{
+		return ClothCollection->GetSubElements(
+			ClothCollection->GetRenderPosition(),
+			ClothCollection->GetPatternStart(),
+			ClothCollection->GetPatternEnd(),
+			ClothCollection->GetRenderVerticesStart(),
+			ClothCollection->GetRenderVerticesEnd(),
+			LodIndex);
+	}
+
+	TConstArrayView<FVector3f> FCollectionClothLodConstFacade::GetRenderNormal() const
+	{
+		return ClothCollection->GetSubElements(
+			ClothCollection->GetRenderNormal(),
+			ClothCollection->GetPatternStart(),
+			ClothCollection->GetPatternEnd(),
+			ClothCollection->GetRenderVerticesStart(),
+			ClothCollection->GetRenderVerticesEnd(),
+			LodIndex);
+	}
+
+	TConstArrayView<FVector3f> FCollectionClothLodConstFacade::GetRenderTangentU() const
+	{
+		return ClothCollection->GetSubElements(
+			ClothCollection->GetRenderTangentU(),
+			ClothCollection->GetPatternStart(),
+			ClothCollection->GetPatternEnd(),
+			ClothCollection->GetRenderVerticesStart(),
+			ClothCollection->GetRenderVerticesEnd(),
+			LodIndex);
+	}
+
+	TConstArrayView<FVector3f> FCollectionClothLodConstFacade::GetRenderTangentV() const
+	{
+		return ClothCollection->GetSubElements(
+			ClothCollection->GetRenderTangentV(),
+			ClothCollection->GetPatternStart(),
+			ClothCollection->GetPatternEnd(),
+			ClothCollection->GetRenderVerticesStart(),
+			ClothCollection->GetRenderVerticesEnd(),
+			LodIndex);
+	}
+
+	TConstArrayView<TArray<FVector2f>> FCollectionClothLodConstFacade::GetRenderUVs() const
+	{
+		return ClothCollection->GetSubElements(
+			ClothCollection->GetRenderUVs(),
+			ClothCollection->GetPatternStart(),
+			ClothCollection->GetPatternEnd(),
+			ClothCollection->GetRenderVerticesStart(),
+			ClothCollection->GetRenderVerticesEnd(),
+			LodIndex);
+	}
+
+	TConstArrayView<FLinearColor> FCollectionClothLodConstFacade::GetRenderColor() const
+	{
+		return ClothCollection->GetSubElements(
+			ClothCollection->GetRenderColor(),
+			ClothCollection->GetPatternStart(),
+			ClothCollection->GetPatternEnd(),
+			ClothCollection->GetRenderVerticesStart(),
+			ClothCollection->GetRenderVerticesEnd(),
+			LodIndex);
+	}
+
+	TConstArrayView<int32> FCollectionClothLodConstFacade::GetRenderNumBoneInfluences() const
+	{
+		return ClothCollection->GetSubElements(
+			ClothCollection->GetRenderNumBoneInfluences(),
+			ClothCollection->GetPatternStart(),
+			ClothCollection->GetPatternEnd(),
+			ClothCollection->GetRenderVerticesStart(),
+			ClothCollection->GetRenderVerticesEnd(),
+			LodIndex);
+	}
+
+	TConstArrayView<TArray<int32>> FCollectionClothLodConstFacade::GetRenderBoneIndices() const
+	{
+		return ClothCollection->GetSubElements(
+			ClothCollection->GetRenderBoneIndices(),
+			ClothCollection->GetPatternStart(),
+			ClothCollection->GetPatternEnd(),
+			ClothCollection->GetRenderVerticesStart(),
+			ClothCollection->GetRenderVerticesEnd(),
+			LodIndex);
+	}
+
+	TConstArrayView<TArray<float>> FCollectionClothLodConstFacade::GetRenderBoneWeights() const
+	{
+		return ClothCollection->GetSubElements(
+			ClothCollection->GetRenderBoneWeights(),
+			ClothCollection->GetPatternStart(),
+			ClothCollection->GetPatternEnd(),
+			ClothCollection->GetRenderVerticesStart(),
+			ClothCollection->GetRenderVerticesEnd(),
+			LodIndex);
+	}
+
+	int32 FCollectionClothLodConstFacade::GetNumRenderFaces() const
+	{
+		return ClothCollection->GetNumSubElements(
+			ClothCollection->GetPatternStart(),
+			ClothCollection->GetPatternEnd(),
+			ClothCollection->GetRenderFacesStart(),
+			ClothCollection->GetRenderFacesEnd(),
+			LodIndex);
+	}
+
+	TConstArrayView<FIntVector3> FCollectionClothLodConstFacade::GetRenderIndices() const
+	{
+		return ClothCollection->GetSubElements(
+			ClothCollection->GetRenderIndices(),
+			ClothCollection->GetPatternStart(),
+			ClothCollection->GetPatternEnd(),
+			ClothCollection->GetRenderFacesStart(),
+			ClothCollection->GetRenderFacesEnd(),
+			LodIndex);
+	}
+
+	TConstArrayView<int32> FCollectionClothLodConstFacade::GetRenderMaterialIndex() const
+	{
+		return ClothCollection->GetSubElements(
+			ClothCollection->GetRenderMaterialIndex(),
+			ClothCollection->GetPatternStart(),
+			ClothCollection->GetPatternEnd(),
+			ClothCollection->GetRenderFacesStart(),
+			ClothCollection->GetRenderFacesEnd(),
+			LodIndex);
+	}
+
+	TConstArrayView<float> FCollectionClothLodConstFacade::GetWeightMap(const FName& Name) const
+	{
+		const TManagedArray<float>* const WeightMap = ClothCollection->GetUserDefinedAttribute<float>(Name, FClothCollection::SimVerticesGroup);
+		return ClothCollection->GetSubElements(
+			WeightMap,
+			ClothCollection->GetPatternStart(),
+			ClothCollection->GetPatternEnd(),
+			ClothCollection->GetSimVerticesStart(),
+			ClothCollection->GetSimVerticesEnd(),
+			LodIndex);
+	}
+
+	void FCollectionClothLodConstFacade::BuildSimulationMesh(TArray<FVector3f>& Positions, TArray<FVector3f>& Normals, TArray<uint32>& Indices, TArray<int32>& WeldingMap) const
+	{
+		const int32 NumSimVertices = GetNumSimVertices();
+
+		// Initialize welding map with same index
+		// The welding map redirects to an existing vertex index if these two are part of the same welding group.
+		// The redirected index must be the smallest index in the group.
+
+		WeldingMap.SetNumUninitialized(NumSimVertices);
+		for (int32 SimVertexIndex = 0; SimVertexIndex < NumSimVertices; ++SimVertexIndex)
+		{
+			WeldingMap[SimVertexIndex] = SimVertexIndex;
+		}
+
+		// Define welding groups
+		// Welding groups contain all stitched pair of indices to be welded together that are required to build the welding map.
+		// Key is the smallest redirected index in the group, and will be the one index used in the welding map redirects.
+		TMap<int32, TSet<int32>> WeldingGroups;
+
+		auto UpdateWeldingMap =
+			[&WeldingMap, &WeldingGroups](int32 Index0, int32 Index1)
+		{
+			// Only process pairs that are not already redirected to the same index
+			if (WeldingMap[Index0] != WeldingMap[Index1])
+			{
+				// Make sure Index0 points to the the smallest redirected index, so that merges are done into the correct group
+				if (WeldingMap[Index0] > WeldingMap[Index1])
+				{
+					Swap(Index0, Index1);
+				}
+
+				// Find the group for Index0 if any
+				const int32 Key0 = WeldingMap[Index0];
+				TSet<int32>* WeldingGroup0 = WeldingGroups.Find(Key0);
+				if (!WeldingGroup0)
+				{
+					// No existing group, create a new one
+					check(Key0 == Index0);  // No group means this index can't already have been redirected  // TODO: Make this a checkSlow
+					WeldingGroup0 = &WeldingGroups.Add(Key0);
+					WeldingGroup0->Add(Index0);
+				}
+
+				// Find the group for Index1, if it exists merge the two groups
+				const int32 Key1 = WeldingMap[Index1];
+				if (TSet<int32>* const WeldingGroup1 = WeldingGroups.Find(Key1))
+				{
+					// Update group1 redirected indices with the new key
+					for (int32 Index : *WeldingGroup1)
+					{
+						WeldingMap[Index] = Key0;
+					}
+
+					// Merge group0 & group1
+					WeldingGroup0->Append(*WeldingGroup1);
+
+					// Remove group1
+					WeldingGroups.Remove(Key1);
+
+					// Sanity check
+					check(WeldingGroup0->Contains(Key0) && WeldingGroup0->Contains(Key1));  // TODO: Make this a checkSlow
+				}
+				else
+				{
+					// Otherwise add Index1 to Index0's group
+					check(Key1 == Index1);  // No group means this index can't already have been redirected  // TODO: Make this a checkSlow
+					WeldingMap[Index1] = Key0;
+					WeldingGroup0->Add(Index1);
+				}
+			}
+		};
+
+		// Apply all seams
+		const int32 NumSeams = GetNumSeams();
+		const TConstArrayView<TArray<FIntVector2>> SeamStitches = GetSeamStitches();
+		for (int32 SeamIndex = 0; SeamIndex < NumSeams; ++SeamIndex)
+		{
+			for (const FIntVector2& Stitch : SeamStitches[SeamIndex])
+			{
+				UpdateWeldingMap(Stitch[0], Stitch[1]);
+			}
+		}
+
+		// Calculate the number of welded vertices
+		int32 NumWeldedVertices = NumSimVertices;
+		for (const TPair<int32, TSet<int32>>& WeldingGroup : WeldingGroups)
+		{
+			NumWeldedVertices -= WeldingGroup.Value.Num() - 1;
+		}
+
+		// Fill up the vertex arrays
+		Positions.SetNumUninitialized(NumWeldedVertices);
+		Normals.SetNumUninitialized(NumWeldedVertices);
+
+		const TConstArrayView<FVector3f> SimRestPosition = GetSimRestPosition();
+		const TConstArrayView<FVector3f> SimRestNormal = GetSimRestNormal();
+
+		TArray<uint32> WeldedIndices;
+		WeldedIndices.SetNumUninitialized(NumSimVertices);
+
+		uint32 WeldedIndex = 0;
+		for (int32 VertexIndex = 0; VertexIndex < NumSimVertices; ++VertexIndex)
+		{
+			if (WeldingMap[VertexIndex] == VertexIndex)
+			{
+				Positions[WeldedIndex] = SimRestPosition[VertexIndex];
+				Normals[WeldedIndex] = SimRestNormal[VertexIndex];
+				WeldedIndices[VertexIndex] = WeldedIndex++;
+			}
+			else
+			{
+				WeldedIndices[VertexIndex] = WeldedIndices[WeldingMap[VertexIndex]];
+			}
+		}
+
+		// Fill up the face array
+		const int32 NumSimFaces = GetNumSimFaces();
+		Indices.SetNumUninitialized(NumSimFaces * 3);
+
+		const TConstArrayView<FIntVector3> SimIndices = GetSimIndices();
+
+		for (int32 FaceIndex = 0; FaceIndex < NumSimFaces; ++FaceIndex)
+		{
+			Indices[FaceIndex * 3 + 0] = WeldedIndices[SimIndices[FaceIndex][0]];
+			Indices[FaceIndex * 3 + 1] = WeldedIndices[SimIndices[FaceIndex][1]];
+			Indices[FaceIndex * 3 + 2] = WeldedIndices[SimIndices[FaceIndex][2]];
+		}
+	}
+
+	FCollectionClothLodConstFacade::FCollectionClothLodConstFacade(const TSharedPtr<const FClothCollection>& InClothCollection, int32 InLodIndex)
+		: ClothCollection(InClothCollection)
+		, LodIndex(InLodIndex)
+	{
+		check(ClothCollection.IsValid());
+		check(ClothCollection->IsValid());
+		check(LodIndex >= 0 && LodIndex < ClothCollection->GetNumElements(FClothCollection::LodsGroup));
+	}
+
+	void FCollectionClothLodFacade::Reset()
+	{
+		SetNumMaterials(0);
+		SetNumTetherBatches(0);
+		SetNumSeams(0);
+		SetNumPatterns(0);
+	}
+
+	void FCollectionClothLodFacade::Initialize(const TArray<FVector3f>& Positions, const TArray<uint32>& Indices)
+	{
+		using namespace UE::Chaos::ClothAsset::Private;
+		Reset();
+
+		TArray<FIsland> Islands;
+		UnwrapMesh(Positions, Indices, Islands);  // Unwrap to 2D and reconstruct indices on 3D mesh
+
+		for (FIsland& Island : Islands)
+		{
+			if (Island.Indices.Num() && Island.Positions.Num() && Island.RestPositions.Num())
+			{
+				FCollectionClothPatternFacade Pattern = AddGetPattern();
+				Pattern.Initialize(Island.Positions, Island.RestPositions, Island.Indices);
+			}
+		}
+
+		TArray<FSeam> Seams;
+		BuildSeams(Islands, Seams);  // Build the seam information as to be able to re-weld the mesh for simulation
+
+		SetNumSeams(Seams.Num());
+
+		const TArrayView<TArray<FIntVector2>> SeamStitches = GetSeamStitches();
+		const TArrayView<FIntVector2> SeamPatterns = GetSeamPatterns();
+		const TManagedArray<int32>* const SimVerticesStart = ClothCollection->GetSimVerticesStart();
+
+		for (int32 SeamIndex = 0; SeamIndex < Seams.Num(); ++SeamIndex)
+		{
+			const FIntVector2& Patterns = Seams[SeamIndex].Patterns;
+			SeamPatterns[SeamIndex] = Patterns;
+
+			const int32 NumStitches = Seams[SeamIndex].Stitches.Num();
+			SeamStitches[SeamIndex].Reset(NumStitches);
+			for (const FIntVector2& Stitch : Seams[SeamIndex].Stitches)
+			{
+				SeamStitches[SeamIndex].Emplace(Stitch[0] + (*SimVerticesStart)[Patterns[0]], Stitch[1] + (*SimVerticesStart)[Patterns[1]]);
+			}
+		}
+	}
+
+	void FCollectionClothLodFacade::Initialize(const FCollectionClothLodConstFacade& Other)
+	{
+		Reset();
+
+		// Patterns Group
+		const int32 NumPatterns = Other.GetNumPatterns();
+		SetNumPatterns(NumPatterns);
+		for (int32 PatternIndex = 0; PatternIndex < NumPatterns; ++PatternIndex)
+		{
+			GetPattern(PatternIndex).Initialize(Other.GetPattern(PatternIndex));
+		}
+
+		// Seam Group
+		const int32 NumSeams = Other.GetNumSeams();
+		SetNumSeams(NumSeams);
+		for (int32 SeamIndex = 0; SeamIndex < NumSeams; ++SeamIndex)
+		{
+			GetSeamPatterns()[SeamIndex] = Other.GetSeamPatterns()[SeamIndex];
+			GetSeamStitches()[SeamIndex] = Other.GetSeamStitches()[SeamIndex];
+		}
+
+		// Tether Batches Group
+		const int32 NumTetherBatches = Other.GetNumPatterns();
+		SetNumTetherBatches(NumTetherBatches);
+		for (int32 TetherBatchIndex = 0; TetherBatchIndex < TetherBatchIndex; ++TetherBatchIndex)
+		{
+			// GetTetherBatch(TetherBatchIndex).Initialize(Other.GetTetherBatch(TetherBatchIndex));  // TODO: Tether Batches facade
+		}
+
+		// Materials Group
+		const int32 NumMaterials = Other.GetNumMaterials();
+		SetNumMaterials(NumMaterials);
+		for (int32 MaterialIndex = 0; MaterialIndex < NumMaterials; ++MaterialIndex)
+		{
+			GetRenderMaterialPathName()[MaterialIndex] = Other.GetRenderMaterialPathName()[MaterialIndex];
+		}
+
+		// LODs Group
+		SetPhysicsAssetPathName(Other.GetPhysicsAssetPathName());
+		SetSkeletonAssetPathName(Other.GetSkeletonAssetPathName());
+	}
+
+	void FCollectionClothLodFacade::SetNumMaterials(int32 InNumMaterials)
+	{
+		GetClothCollection()->SetNumElements(
+			InNumMaterials,
+			FClothCollection::MaterialsGroup,
+			GetClothCollection()->GetMaterialStart(),
+			GetClothCollection()->GetMaterialEnd(),
+			LodIndex);
+	}
+
+	void FCollectionClothLodFacade::SetNumTetherBatches(int32 InNumTetherBatches)
+	{
+		const int32 NumTetherBatches = GetNumTetherBatches();
+
+		for (int32 TetherBatchIndex = InNumTetherBatches; TetherBatchIndex < NumTetherBatches; ++TetherBatchIndex)
+		{
+			// GetTetherBatch(TetherBatchIndex).Reset();  // TODO: Tether reset
+		}
+
+		GetClothCollection()->SetNumElements(
+			InNumTetherBatches,
+			FClothCollection::TetherBatchesGroup,
+			GetClothCollection()->GetTetherBatchStart(),
+			GetClothCollection()->GetTetherBatchEnd(),
+			LodIndex);
+
+		for (int32 TetherBatchIndex = NumTetherBatches; TetherBatchIndex < InNumTetherBatches; ++TetherBatchIndex)
+		{
+			// GetTetherBatch(TetherBatchIndex).SetDefaults();  // TODO: Tether set default
+		}
+	}
+
+	void FCollectionClothLodFacade::SetNumSeams(int32 InNumSeams)
+	{
+		GetClothCollection()->SetNumElements(
+			InNumSeams,
+			FClothCollection::SeamsGroup,
+			GetClothCollection()->GetSeamStart(),
+			GetClothCollection()->GetSeamEnd(),
+			LodIndex);
+	}
+
+
+	void FCollectionClothLodFacade::SetNumPatterns(int32 InNumPatterns)
+	{
+		const int32 NumPatterns = GetNumPatterns();
+
+		for (int32 PatternIndex = InNumPatterns; PatternIndex < NumPatterns; ++PatternIndex)
+		{
+			GetPattern(PatternIndex).Reset();
+		}
+
+		GetClothCollection()->SetNumElements(
+			InNumPatterns,
+			FClothCollection::PatternsGroup,
+			GetClothCollection()->GetPatternStart(),
+			GetClothCollection()->GetPatternEnd(),
+			LodIndex);
+
+		for (int32 PatternIndex = NumPatterns; PatternIndex < InNumPatterns; ++PatternIndex)
+		{
+			GetPattern(PatternIndex).SetDefaults();
+		}
+	}
+
+	TArrayView<FString> FCollectionClothLodFacade::GetRenderMaterialPathName()
+	{
+		return GetClothCollection()->GetElements(GetClothCollection()->GetRenderMaterialPathName(), ClothCollection->GetMaterialStart(), ClothCollection->GetMaterialEnd(), LodIndex);
+	}
+
+	TArrayView<FIntVector2> FCollectionClothLodFacade::GetSeamPatterns()
+	{
+		return GetClothCollection()->GetElements(GetClothCollection()->GetSeamPatterns(), ClothCollection->GetSeamStart(), ClothCollection->GetSeamEnd(), LodIndex);
+	}
+
+	TArrayView<TArray<FIntVector2>> FCollectionClothLodFacade::GetSeamStitches()
+	{
+		return GetClothCollection()->GetElements(GetClothCollection()->GetSeamStitches(), ClothCollection->GetSeamStart(), ClothCollection->GetSeamEnd(), LodIndex);
+	}
+
+	int32 FCollectionClothLodFacade::AddPattern()
+	{
+		const int32 PatternIndex = GetNumPatterns();
+		SetNumPatterns(PatternIndex + 1);
+		return PatternIndex;
+	}
+
+	FCollectionClothPatternFacade FCollectionClothLodFacade::GetPattern(int32 PatternIndex)
+	{
+		return FCollectionClothPatternFacade(GetClothCollection(), LodIndex, PatternIndex);
+	}
+
+	void FCollectionClothLodFacade::SetPhysicsAssetPathName(const FString& PhysicsAssetPathName)
+	{
+		(*GetClothCollection()->GetPhysicsAssetPathName())[LodIndex] = PhysicsAssetPathName;
+	}
+
+	void FCollectionClothLodFacade::SetSkeletonAssetPathName(const FString& SkeletonAssetPathName)
+	{
+		(*GetClothCollection()->GetSkeletonAssetPathName())[LodIndex] = SkeletonAssetPathName;
+	}
+
+	TArrayView<FVector2f> FCollectionClothLodFacade::GetSimPosition()
+	{
+		return GetClothCollection()->GetSubElements(
+			GetClothCollection()->GetSimPosition(),
+			GetClothCollection()->GetPatternStart(),
+			GetClothCollection()->GetPatternEnd(),
+			GetClothCollection()->GetSimVerticesStart(),
+			GetClothCollection()->GetSimVerticesEnd(),
+			LodIndex);
+	}
+
+	TArrayView<FVector3f> FCollectionClothLodFacade::GetSimRestPosition()
+	{
+		return GetClothCollection()->GetSubElements(
+			GetClothCollection()->GetSimRestPosition(),
+			GetClothCollection()->GetPatternStart(),
+			GetClothCollection()->GetPatternEnd(),
+			GetClothCollection()->GetSimVerticesStart(),
+			GetClothCollection()->GetSimVerticesEnd(),
+			LodIndex);
+	}
+
+	TArrayView<FVector3f> FCollectionClothLodFacade::GetSimRestNormal()
+	{
+		return GetClothCollection()->GetSubElements(
+			GetClothCollection()->GetSimRestNormal(),
+			GetClothCollection()->GetPatternStart(),
+			GetClothCollection()->GetPatternEnd(),
+			GetClothCollection()->GetSimVerticesStart(),
+			GetClothCollection()->GetSimVerticesEnd(),
+			LodIndex);
+	}
+
+	TArrayView<int32> FCollectionClothLodFacade::GetSimNumBoneInfluences()
+	{
+		return GetClothCollection()->GetSubElements(
+			GetClothCollection()->GetSimNumBoneInfluences(),
+			GetClothCollection()->GetPatternStart(),
+			GetClothCollection()->GetPatternEnd(),
+			GetClothCollection()->GetSimVerticesStart(),
+			GetClothCollection()->GetSimVerticesEnd(),
+			LodIndex);
+	}
+
+	TArrayView<TArray<int32>> FCollectionClothLodFacade::GetSimBoneIndices()
+	{
+		return GetClothCollection()->GetSubElements(
+			GetClothCollection()->GetSimBoneIndices(),
+			GetClothCollection()->GetPatternStart(),
+			GetClothCollection()->GetPatternEnd(),
+			GetClothCollection()->GetSimVerticesStart(),
+			GetClothCollection()->GetSimVerticesEnd(),
+			LodIndex);
+	}
+
+	TArrayView<TArray<float>> FCollectionClothLodFacade::GetSimBoneWeights()
+	{
+		return GetClothCollection()->GetSubElements(
+			GetClothCollection()->GetSimBoneWeights(),
+			GetClothCollection()->GetPatternStart(),
+			GetClothCollection()->GetPatternEnd(),
+			GetClothCollection()->GetSimVerticesStart(),
+			GetClothCollection()->GetSimVerticesEnd(),
+			LodIndex);
+	}
+
+	TArrayView<FIntVector3> FCollectionClothLodFacade::GetSimIndices()
+	{
+		return GetClothCollection()->GetSubElements(
+			GetClothCollection()->GetSimIndices(),
+			GetClothCollection()->GetPatternStart(),
+			GetClothCollection()->GetPatternEnd(),
+			GetClothCollection()->GetSimFacesStart(),
+			GetClothCollection()->GetSimFacesEnd(),
+			LodIndex);
+	}
+
+	TArrayView<FVector3f> FCollectionClothLodFacade::GetRenderPosition()
+	{
+		return GetClothCollection()->GetSubElements(
+			GetClothCollection()->GetRenderPosition(),
+			GetClothCollection()->GetPatternStart(),
+			GetClothCollection()->GetPatternEnd(),
+			GetClothCollection()->GetRenderVerticesStart(),
+			GetClothCollection()->GetRenderVerticesEnd(),
+			LodIndex);
+	}
+
+	TArrayView<FVector3f> FCollectionClothLodFacade::GetRenderNormal()
+	{
+		return GetClothCollection()->GetSubElements(
+			GetClothCollection()->GetRenderNormal(),
+			GetClothCollection()->GetPatternStart(),
+			GetClothCollection()->GetPatternEnd(),
+			GetClothCollection()->GetRenderVerticesStart(),
+			GetClothCollection()->GetRenderVerticesEnd(),
+			LodIndex);
+	}
+
+	TArrayView<FVector3f> FCollectionClothLodFacade::GetRenderTangentU()
+	{
+		return GetClothCollection()->GetSubElements(
+			GetClothCollection()->GetRenderTangentU(),
+			GetClothCollection()->GetPatternStart(),
+			GetClothCollection()->GetPatternEnd(),
+			GetClothCollection()->GetRenderVerticesStart(),
+			GetClothCollection()->GetRenderVerticesEnd(),
+			LodIndex);
+	}
+
+	TArrayView<FVector3f> FCollectionClothLodFacade::GetRenderTangentV()
+	{
+		return GetClothCollection()->GetSubElements(
+			GetClothCollection()->GetRenderTangentV(),
+			GetClothCollection()->GetPatternStart(),
+			GetClothCollection()->GetPatternEnd(),
+			GetClothCollection()->GetRenderVerticesStart(),
+			GetClothCollection()->GetRenderVerticesEnd(),
+			LodIndex);
+	}
+
+	TArrayView<TArray<FVector2f>> FCollectionClothLodFacade::GetRenderUVs()
+	{
+		return GetClothCollection()->GetSubElements(
+			GetClothCollection()->GetRenderUVs(),
+			GetClothCollection()->GetPatternStart(),
+			GetClothCollection()->GetPatternEnd(),
+			GetClothCollection()->GetRenderVerticesStart(),
+			GetClothCollection()->GetRenderVerticesEnd(),
+			LodIndex);
+	}
+
+	TArrayView<FLinearColor> FCollectionClothLodFacade::GetRenderColor()
+	{
+		return GetClothCollection()->GetSubElements(
+			GetClothCollection()->GetRenderColor(),
+			GetClothCollection()->GetPatternStart(),
+			GetClothCollection()->GetPatternEnd(),
+			GetClothCollection()->GetRenderVerticesStart(),
+			GetClothCollection()->GetRenderVerticesEnd(),
+			LodIndex);
+	}
+
+	TArrayView<int32> FCollectionClothLodFacade::GetRenderNumBoneInfluences()
+	{
+		return GetClothCollection()->GetSubElements(
+			GetClothCollection()->GetRenderNumBoneInfluences(),
+			GetClothCollection()->GetPatternStart(),
+			GetClothCollection()->GetPatternEnd(),
+			GetClothCollection()->GetRenderVerticesStart(),
+			GetClothCollection()->GetRenderVerticesEnd(),
+			LodIndex);
+	}
+
+	TArrayView<TArray<int32>> FCollectionClothLodFacade::GetRenderBoneIndices()
+	{
+		return GetClothCollection()->GetSubElements(
+			GetClothCollection()->GetRenderBoneIndices(),
+			GetClothCollection()->GetPatternStart(),
+			GetClothCollection()->GetPatternEnd(),
+			GetClothCollection()->GetRenderVerticesStart(),
+			GetClothCollection()->GetRenderVerticesEnd(),
+			LodIndex);
+	}
+
+	TArrayView<TArray<float>> FCollectionClothLodFacade::GetRenderBoneWeights()
+	{
+		return GetClothCollection()->GetSubElements(
+			GetClothCollection()->GetRenderBoneWeights(),
+			GetClothCollection()->GetPatternStart(),
+			GetClothCollection()->GetPatternEnd(),
+			GetClothCollection()->GetRenderVerticesStart(),
+			GetClothCollection()->GetRenderVerticesEnd(),
+			LodIndex);
+	}
+
+	TArrayView<FIntVector3> FCollectionClothLodFacade::GetRenderIndices()
+	{
+		return GetClothCollection()->GetSubElements(
+			GetClothCollection()->GetRenderIndices(),
+			GetClothCollection()->GetPatternStart(),
+			GetClothCollection()->GetPatternEnd(),
+			GetClothCollection()->GetRenderFacesStart(),
+			GetClothCollection()->GetRenderFacesEnd(),
+			LodIndex);
+	}
+
+	TArrayView<int32> FCollectionClothLodFacade::GetRenderMaterialIndex()
+	{
+		return GetClothCollection()->GetSubElements(
+			GetClothCollection()->GetRenderMaterialIndex(),
+			GetClothCollection()->GetPatternStart(),
+			GetClothCollection()->GetPatternEnd(),
+			GetClothCollection()->GetRenderFacesStart(),
+			GetClothCollection()->GetRenderFacesEnd(),
+			LodIndex);
+	}
+
+	TArrayView<float> FCollectionClothLodFacade::GetWeightMap(const FName& Name)
+	{
+		TManagedArray<float>* const WeightMap = GetClothCollection()->GetUserDefinedAttribute<float>(Name, FClothCollection::SimVerticesGroup);
+		return GetClothCollection()->GetSubElements(
+			WeightMap,
+			ClothCollection->GetPatternStart(),
+			ClothCollection->GetPatternEnd(),
+			ClothCollection->GetSimVerticesStart(),
+			ClothCollection->GetSimVerticesEnd(),
+			LodIndex);
+	}
+
+	FCollectionClothLodFacade::FCollectionClothLodFacade(const TSharedPtr<FClothCollection>& InClothCollection, int32 InLodIndex)
+		: FCollectionClothLodConstFacade(InClothCollection, InLodIndex)
+	{
+	}
+
+	void FCollectionClothLodFacade::SetDefaults()
+	{
+		(*GetClothCollection()->GetPatternStart())[LodIndex] = INDEX_NONE;
+		(*GetClothCollection()->GetPatternEnd())[LodIndex] = INDEX_NONE;
+		(*GetClothCollection()->GetSeamStart())[LodIndex] = INDEX_NONE;
+		(*GetClothCollection()->GetSeamEnd())[LodIndex] = INDEX_NONE;
+		(*GetClothCollection()->GetTetherBatchStart())[LodIndex] = INDEX_NONE;
+		(*GetClothCollection()->GetTetherBatchEnd())[LodIndex] = INDEX_NONE;
+		(*GetClothCollection()->GetMaterialStart())[LodIndex] = INDEX_NONE;
+		(*GetClothCollection()->GetMaterialEnd())[LodIndex] = INDEX_NONE;
+		(*GetClothCollection()->GetPhysicsAssetPathName())[LodIndex].Empty();
+		(*GetClothCollection()->GetSkeletonAssetPathName())[LodIndex].Empty();
+	}
+}  // End namespace UE::Chaos::ClothAsset
diff --git a/Engine/Plugins/Experimental/ChaosClothAsset/Source/ChaosClothAsset/Private/ChaosClothAsset/CollectionClothPatternFacade.cpp b/Engine/Plugins/Experimental/ChaosClothAsset/Source/ChaosClothAsset/Private/ChaosClothAsset/CollectionClothPatternFacade.cpp
new file mode 100644
--- /dev/null
+++ b/Engine/Plugins/Experimental/ChaosClothAsset/Source/ChaosClothAsset/Private/ChaosClothAsset/CollectionClothPatternFacade.cpp
@@ -0,0 +1,689 @@
+// Copyright Epic Games, Inc. All Rights Reserved.
+
+#include "ChaosClothAsset/CollectionClothPatternFacade.h"
+#include "ChaosClothAsset/ClothCollection.h"
+
+namespace UE::Chaos::ClothAsset
+{
+	int32 FCollectionClothPatternConstFacade::GetStatusFlags() const
+	{
+		return (*ClothCollection->GetStatusFlags())[GetElementIndex()];
+	}
+
+	int32 FCollectionClothPatternConstFacade::GetNumSimVertices() const
+	{
+		return ClothCollection->GetNumElements(
+			ClothCollection->GetSimVerticesStart(),
+			ClothCollection->GetSimVerticesEnd(),
+			GetElementIndex());
+	}
+
+	int32 FCollectionClothPatternConstFacade::GetNumSimFaces() const
+	{
+		return ClothCollection->GetNumElements(
+			ClothCollection->GetSimFacesStart(),
+			ClothCollection->GetSimFacesEnd(),
+			GetElementIndex());
+	}
+
+	int32 FCollectionClothPatternConstFacade::GetNumRenderVertices() const
+	{
+		return ClothCollection->GetNumElements(
+			ClothCollection->GetRenderVerticesStart(),
+			ClothCollection->GetRenderVerticesEnd(),
+			GetElementIndex());
+	}
+
+	int32 FCollectionClothPatternConstFacade::GetNumRenderFaces() const
+	{
+		return ClothCollection->GetNumElements(
+			ClothCollection->GetRenderFacesStart(),
+			ClothCollection->GetRenderFacesEnd(),
+			GetElementIndex());
+	}
+
+	int32 FCollectionClothPatternConstFacade::GetSimVerticesOffset() const
+	{
+		return ClothCollection->GetElementsOffset(
+			ClothCollection->GetSimVerticesStart(),
+			GetBaseElementIndex(),
+			GetElementIndex());
+	}
+
+	int32 FCollectionClothPatternConstFacade::GetSimFacesOffset() const
+	{
+		return ClothCollection->GetElementsOffset(
+			ClothCollection->GetSimFacesStart(),
+			GetBaseElementIndex(),
+			GetElementIndex());
+	}
+
+	int32 FCollectionClothPatternConstFacade::GetRenderVerticesOffset() const
+	{
+		return ClothCollection->GetElementsOffset(
+			ClothCollection->GetRenderVerticesStart(),
+			GetBaseElementIndex(),
+			GetElementIndex());
+	}
+
+	int32 FCollectionClothPatternConstFacade::GetRenderFacesOffset() const
+	{
+		return ClothCollection->GetElementsOffset(
+			ClothCollection->GetRenderFacesStart(),
+			GetBaseElementIndex(),
+			GetElementIndex());
+	}
+
+	TConstArrayView<FVector2f> FCollectionClothPatternConstFacade::GetSimPosition() const
+	{
+		return ClothCollection->GetElements(
+			ClothCollection->GetSimPosition(),
+			ClothCollection->GetSimVerticesStart(),
+			ClothCollection->GetSimVerticesEnd(),
+			GetElementIndex());
+	}
+
+	TConstArrayView<FVector3f> FCollectionClothPatternConstFacade::GetSimRestPosition() const
+	{
+		return ClothCollection->GetElements(
+			ClothCollection->GetSimRestPosition(),
+			ClothCollection->GetSimVerticesStart(),
+			ClothCollection->GetSimVerticesEnd(),
+			GetElementIndex());
+	}
+
+	TConstArrayView<FVector3f> FCollectionClothPatternConstFacade::GetSimRestNormal() const
+	{
+		return ClothCollection->GetElements(
+			ClothCollection->GetSimRestNormal(),
+			ClothCollection->GetSimVerticesStart(),
+			ClothCollection->GetSimVerticesEnd(),
+			GetElementIndex());
+	}
+
+	TConstArrayView<int32> FCollectionClothPatternConstFacade::GetSimNumBoneInfluences() const
+	{
+		return ClothCollection->GetElements(
+			ClothCollection->GetSimNumBoneInfluences(),
+			ClothCollection->GetSimVerticesStart(),
+			ClothCollection->GetSimVerticesEnd(),
+			GetElementIndex());
+	}
+
+	TConstArrayView<TArray<int32>> FCollectionClothPatternConstFacade::GetSimBoneIndices() const
+	{
+		return ClothCollection->GetElements(
+			ClothCollection->GetSimBoneIndices(),
+			ClothCollection->GetSimVerticesStart(),
+			ClothCollection->GetSimVerticesEnd(),
+			GetElementIndex());
+	}
+
+	TConstArrayView<TArray<float>> FCollectionClothPatternConstFacade::GetSimBoneWeights() const
+	{
+		return ClothCollection->GetElements(
+			ClothCollection->GetSimBoneWeights(),
+			ClothCollection->GetSimVerticesStart(),
+			ClothCollection->GetSimVerticesEnd(),
+			GetElementIndex());
+	}
+
+	TConstArrayView<FIntVector3> FCollectionClothPatternConstFacade::GetSimIndices() const
+	{
+		return ClothCollection->GetElements(
+			ClothCollection->GetSimIndices(),
+			ClothCollection->GetSimFacesStart(),
+			ClothCollection->GetSimFacesEnd(),
+			GetElementIndex());
+	}
+
+	TConstArrayView<FVector3f> FCollectionClothPatternConstFacade::GetRenderPosition() const
+	{
+		return ClothCollection->GetElements(
+			ClothCollection->GetRenderPosition(),
+			ClothCollection->GetRenderVerticesStart(),
+			ClothCollection->GetRenderVerticesEnd(),
+			GetElementIndex());
+	}
+
+	TConstArrayView<FVector3f> FCollectionClothPatternConstFacade::GetRenderNormal() const
+	{
+		return ClothCollection->GetElements(
+			ClothCollection->GetRenderNormal(),
+			ClothCollection->GetRenderVerticesStart(),
+			ClothCollection->GetRenderVerticesEnd(),
+			GetElementIndex());
+	}
+
+	TConstArrayView<FVector3f> FCollectionClothPatternConstFacade::GetRenderTangentU() const
+	{
+		return ClothCollection->GetElements(
+			ClothCollection->GetRenderTangentU(),
+			ClothCollection->GetRenderVerticesStart(),
+			ClothCollection->GetRenderVerticesEnd(),
+			GetElementIndex());
+	}
+
+	TConstArrayView<FVector3f> FCollectionClothPatternConstFacade::GetRenderTangentV() const
+	{
+		return ClothCollection->GetElements(
+			ClothCollection->GetRenderTangentV(),
+			ClothCollection->GetRenderVerticesStart(),
+			ClothCollection->GetRenderVerticesEnd(),
+			GetElementIndex());
+	}
+
+	TConstArrayView<TArray<FVector2f>> FCollectionClothPatternConstFacade::GetRenderUVs() const
+	{
+		return ClothCollection->GetElements(
+			ClothCollection->GetRenderUVs(),
+			ClothCollection->GetRenderVerticesStart(),
+			ClothCollection->GetRenderVerticesEnd(),
+			GetElementIndex());
+	}
+
+	TConstArrayView<FLinearColor> FCollectionClothPatternConstFacade::GetRenderColor() const
+	{
+		return ClothCollection->GetElements(
+			ClothCollection->GetRenderColor(),
+			ClothCollection->GetRenderVerticesStart(),
+			ClothCollection->GetRenderVerticesEnd(),
+			GetElementIndex());
+	}
+
+	TConstArrayView<int32> FCollectionClothPatternConstFacade::GetRenderNumBoneInfluences() const
+	{
+		return ClothCollection->GetElements(
+			ClothCollection->GetRenderNumBoneInfluences(),
+			ClothCollection->GetRenderVerticesStart(),
+			ClothCollection->GetRenderVerticesEnd(),
+			GetElementIndex());
+	}
+
+	TConstArrayView<TArray<int32>> FCollectionClothPatternConstFacade::GetRenderBoneIndices() const
+	{
+		return ClothCollection->GetElements(
+			ClothCollection->GetRenderBoneIndices(),
+			ClothCollection->GetRenderVerticesStart(),
+			ClothCollection->GetRenderVerticesEnd(),
+			GetElementIndex());
+	}
+
+	TConstArrayView<TArray<float>> FCollectionClothPatternConstFacade::GetRenderBoneWeights() const
+	{
+		return ClothCollection->GetElements(
+			ClothCollection->GetRenderBoneWeights(),
+			ClothCollection->GetRenderVerticesStart(),
+			ClothCollection->GetRenderVerticesEnd(),
+			GetElementIndex());
+	}
+
+	TConstArrayView<FIntVector3> FCollectionClothPatternConstFacade::GetRenderIndices() const
+	{
+		return ClothCollection->GetElements(
+			ClothCollection->GetRenderIndices(),
+			ClothCollection->GetRenderFacesStart(),
+			ClothCollection->GetRenderFacesEnd(),
+			GetElementIndex());
+	}
+
+	TConstArrayView<int32> FCollectionClothPatternConstFacade::GetRenderMaterialIndex() const
+	{
+		return ClothCollection->GetElements(
+			ClothCollection->GetRenderMaterialIndex(),
+			ClothCollection->GetRenderFacesStart(),
+			ClothCollection->GetRenderFacesEnd(),
+			GetElementIndex());
+	}
+
+	TConstArrayView<float> FCollectionClothPatternConstFacade::GetWeightMap(const FName& Name) const
+	{
+		const TManagedArray<float>* const WeightMap = ClothCollection->GetUserDefinedAttribute<float>(Name, FClothCollection::SimVerticesGroup);
+		return ClothCollection->GetElements(
+			WeightMap,
+			ClothCollection->GetSimVerticesStart(),
+			ClothCollection->GetSimVerticesEnd(),
+			GetElementIndex());
+	}
+
+	FCollectionClothPatternConstFacade::FCollectionClothPatternConstFacade(const TSharedPtr<const FClothCollection>& ClothCollection, int32 InLodIndex, int32 InPatternIndex)
+		: ClothCollection(ClothCollection)
+		, LodIndex(InLodIndex)
+		, PatternIndex(InPatternIndex)
+	{
+		check(ClothCollection.IsValid());
+		check(ClothCollection->IsValid());
+		check(LodIndex >= 0 && LodIndex < ClothCollection->GetNumElements(FClothCollection::LodsGroup));
+		check(PatternIndex >= 0 && PatternIndex < ClothCollection->GetNumElements(ClothCollection->GetPatternStart(), ClothCollection->GetPatternEnd(), LodIndex));
+	}
+
+	int32 FCollectionClothPatternConstFacade::GetBaseElementIndex() const
+	{
+		return (*ClothCollection->GetPatternStart())[LodIndex];
+	}
+
+	void FCollectionClothPatternFacade::Reset()
+	{
+		SetNumSimVertices(0);
+		SetNumSimFaces(0);
+		SetNumRenderVertices(0);
+		SetNumRenderFaces(0);
+		SetDefaults();
+	}
+
+	void FCollectionClothPatternFacade::Initialize(const TArray<FVector2f>& Positions, const TArray<FVector3f>& RestPositions, const TArray<uint32>& Indices)
+	{
+		Reset();
+
+		const int32 NumSimVertices = Positions.Num();
+		check(NumSimVertices == RestPositions.Num());
+
+		SetNumSimVertices(NumSimVertices);
+
+		const TArrayView<FVector2f> SimPosition = GetSimPosition();
+		const TArrayView<FVector3f> SimRestPosition = GetSimRestPosition();
+		const TArrayView<FVector3f> SimRestNormal = GetSimRestNormal();
+
+		for (int32 SimVertexIndex = 0; SimVertexIndex < NumSimVertices; ++SimVertexIndex)
+		{
+			SimPosition[SimVertexIndex] = Positions[SimVertexIndex];
+			SimRestPosition[SimVertexIndex] = RestPositions[SimVertexIndex];
+			SimRestNormal[SimVertexIndex] = FVector3f::ZeroVector;
+		}
+
+		const int32 NumSimFaces = (uint32)Indices.Num() / 3;
+		check(NumSimFaces * 3 == Indices.Num());
+
+		SetNumSimFaces(NumSimFaces);
+
+		const TArrayView<FIntVector3> SimIndices = GetSimIndices();
+
+		// Face indices always index from the first vertex of the LOD, but these indices are indexed from the start of the pattern and need to be offset
+		const int32 LodSimVerticesOffset = ClothCollection->GetElementsOffset(
+			ClothCollection->GetSimVerticesStart(),
+			GetBaseElementIndex(),
+			GetElementIndex());
+
+		for (int32 SimFaceIndex = 0; SimFaceIndex < NumSimFaces; ++SimFaceIndex)
+		{
+			// Indices from the start of the pattern
+			const int32 PattternVertexIndex0 = Indices[SimFaceIndex * 3];
+			const int32 PattternVertexIndex1 = Indices[SimFaceIndex * 3 + 1];
+			const int32 PattternVertexIndex2 = Indices[SimFaceIndex * 3 + 2];
+
+			// Indices from the start of the LOD
+			const int32 LodVertexIndex0 = PattternVertexIndex0 + LodSimVerticesOffset;
+			const int32 LodVertexIndex1 = PattternVertexIndex1 + LodSimVerticesOffset;
+			const int32 LodVertexIndex2 = PattternVertexIndex2 + LodSimVerticesOffset;
+
+			// Set indices in LOD index space
+			SimIndices[SimFaceIndex] = FIntVector3(LodVertexIndex0, LodVertexIndex1, LodVertexIndex2);
+
+			// Calculate face normal contribution
+			const FVector3f& Pos0 = SimRestPosition[PattternVertexIndex0];
+			const FVector3f& Pos1 = SimRestPosition[PattternVertexIndex1];
+			const FVector3f& Pos2 = SimRestPosition[PattternVertexIndex2];
+			const FVector3f Normal = (Pos1 - Pos0).Cross(Pos2 - Pos0).GetSafeNormal();
+			SimRestNormal[PattternVertexIndex0] += Normal;
+			SimRestNormal[PattternVertexIndex1] += Normal;
+			SimRestNormal[PattternVertexIndex2] += Normal;
+		}
+
+		// Normalize normals
+		for (int32 SimVertexIndex = 0; SimVertexIndex < NumSimVertices; ++SimVertexIndex)
+		{
+			SimRestNormal[SimVertexIndex] = SimRestNormal[SimVertexIndex].GetSafeNormal(UE_SMALL_NUMBER, FVector3f::XAxisVector);
+		}
+	}
+
+	void FCollectionClothPatternFacade::Initialize(const FCollectionClothPatternConstFacade& Other)
+	{
+		Reset();
+
+		// Patterns Group
+		SetStatusFlags(Other.GetStatusFlags());
+
+		// Sim Vertices Group
+		const int32 NumSimVertices = Other.GetNumSimVertices();
+		SetNumSimVertices(NumSimVertices);
+		
+		const TConstArrayView<FVector2f> OtherSimPosition = Other.GetSimPosition();
+		const TConstArrayView<FVector3f> OtherSimRestPosition = Other.GetSimRestPosition();
+		const TConstArrayView<FVector3f> OtherSimRestNormal = Other.GetSimRestNormal();
+		const TConstArrayView<int32> OtherSimNumBoneInfluences = Other.GetSimNumBoneInfluences();
+		const TConstArrayView<TArray<int32>> OtherSimBoneIndices = Other.GetSimBoneIndices();
+		const TConstArrayView<TArray<float>> OtherSimBoneWeights = Other.GetSimBoneWeights();
+		const TArrayView<FVector2f> SimPosition = GetSimPosition();
+		const TArrayView<FVector3f> SimRestPosition = GetSimRestPosition();
+		const TArrayView<FVector3f> SimRestNormal = GetSimRestNormal();
+		const TArrayView<int32> SimNumBoneInfluences = GetSimNumBoneInfluences();
+		const TArrayView<TArray<int32>> SimBoneIndices = GetSimBoneIndices();
+		const TArrayView<TArray<float>> SimBoneWeights = GetSimBoneWeights();
+
+		for (int32 SimVertexIndex = 0; SimVertexIndex < NumSimVertices; ++SimVertexIndex)
+		{
+			SimPosition[SimVertexIndex] = Other.GetSimPosition()[SimVertexIndex];
+			SimRestPosition[SimVertexIndex] = Other.GetSimRestPosition()[SimVertexIndex];
+			SimRestNormal[SimVertexIndex] = Other.GetSimRestNormal()[SimVertexIndex];
+			SimNumBoneInfluences[SimVertexIndex] = Other.GetSimNumBoneInfluences()[SimVertexIndex];
+			SimBoneIndices[SimVertexIndex] = Other.GetSimBoneIndices()[SimVertexIndex];
+			SimBoneWeights[SimVertexIndex] = Other.GetSimBoneWeights()[SimVertexIndex];
+		}
+
+		// Sim Faces Group
+		const int32 NumSimFaces = Other.GetNumSimFaces();
+		SetNumSimFaces(NumSimFaces);
+
+		const int32 LodSimVerticesOffset = ClothCollection->GetElementsOffset(
+			ClothCollection->GetSimVerticesStart(),
+			GetBaseElementIndex(),
+			GetElementIndex());
+		const int32 OtherLodSimVerticesOffset = Other.ClothCollection->GetElementsOffset(
+			ClothCollection->GetSimVerticesStart(),
+			GetBaseElementIndex(),
+			GetElementIndex());
+		const FIntVector3 OtherToLodSimVerticesOffset(LodSimVerticesOffset - OtherLodSimVerticesOffset);
+
+		const TConstArrayView<FIntVector3> OtherSimIndices = Other.GetSimIndices();
+		const TArrayView<FIntVector3> SimIndices = GetSimIndices();
+		
+		for (int32 SimFaceIndex = 0; SimFaceIndex < NumSimFaces; ++SimFaceIndex)
+		{
+			SimIndices[SimFaceIndex] = OtherSimIndices[SimFaceIndex] + OtherToLodSimVerticesOffset;
+		}
+
+		// Render Vertices Group
+		const int32 NumRenderVertices = Other.GetNumRenderVertices();
+		SetNumRenderVertices(NumRenderVertices);
+
+		const TConstArrayView<FVector3f> OtherRenderPosition = Other.GetRenderPosition();
+		const TConstArrayView<FVector3f> OtherRenderNormal = Other.GetRenderNormal();
+		const TConstArrayView<FVector3f> OtherRenderTangentU = Other.GetRenderTangentU();
+		const TConstArrayView<FVector3f> OtherRenderTangentV = Other.GetRenderTangentV();
+		const TConstArrayView<TArray<FVector2f>> OtherRenderUVs = Other.GetRenderUVs();
+		const TConstArrayView<FLinearColor> OtherRenderColor = Other.GetRenderColor();
+		const TConstArrayView<int32> OtherRenderNumBoneInfluences = Other.GetRenderNumBoneInfluences();
+		const TConstArrayView<TArray<int32>> OtherRenderBoneIndices = Other.GetRenderBoneIndices();
+		const TConstArrayView<TArray<float>> OtherRenderBoneWeights = Other.GetRenderBoneWeights();
+		const TArrayView<FVector3f> RenderPosition = GetRenderPosition();
+		const TArrayView<FVector3f> RenderNormal = GetRenderNormal();
+		const TArrayView<FVector3f> RenderTangentU = GetRenderTangentU();
+		const TArrayView<FVector3f> RenderTangentV = GetRenderTangentV();
+		const TArrayView<TArray<FVector2f>> RenderUVs = GetRenderUVs();
+		const TArrayView<FLinearColor> RenderColor = GetRenderColor();
+		const TArrayView<int32> RenderNumBoneInfluences = GetRenderNumBoneInfluences();
+		const TArrayView<TArray<int32>> RenderBoneIndices = GetRenderBoneIndices();
+		const TArrayView<TArray<float>> RenderBoneWeights = GetRenderBoneWeights();
+
+		for (int32 RenderVertexIndex = 0; RenderVertexIndex < NumRenderVertices; ++RenderVertexIndex)
+		{
+			RenderPosition[RenderVertexIndex] = OtherRenderPosition[RenderVertexIndex];
+			RenderNormal[RenderVertexIndex] = OtherRenderNormal[RenderVertexIndex];
+			RenderTangentU[RenderVertexIndex] = OtherRenderTangentU[RenderVertexIndex];
+			RenderTangentV[RenderVertexIndex] = OtherRenderTangentV[RenderVertexIndex];
+			RenderUVs[RenderVertexIndex] = OtherRenderUVs[RenderVertexIndex];
+			RenderColor[RenderVertexIndex] = OtherRenderColor[RenderVertexIndex];
+			RenderNumBoneInfluences[RenderVertexIndex] = OtherRenderNumBoneInfluences[RenderVertexIndex];
+			RenderBoneIndices[RenderVertexIndex] = OtherRenderBoneIndices[RenderVertexIndex];
+			RenderBoneWeights[RenderVertexIndex] = OtherRenderBoneWeights[RenderVertexIndex];
+		}
+
+		// Render Faces Group
+		const int32 NumRenderFaces = Other.GetNumRenderFaces();
+		SetNumRenderFaces(NumRenderFaces);
+
+		const int32 LodRenderVerticesOffset = ClothCollection->GetElementsOffset(
+			ClothCollection->GetRenderVerticesStart(),
+			GetBaseElementIndex(),
+			GetElementIndex());
+		const int32 OtherLodRenderVerticesOffset = Other.ClothCollection->GetElementsOffset(
+			ClothCollection->GetRenderVerticesStart(),
+			GetBaseElementIndex(),
+			GetElementIndex());
+
+		const FIntVector3 OtherToLodRenderVerticesOffset(LodRenderVerticesOffset - OtherLodRenderVerticesOffset);
+
+		for (int32 RenderFaceIndex = 0; RenderFaceIndex < NumRenderFaces; ++RenderFaceIndex)
+		{
+			GetRenderIndices()[RenderFaceIndex] = Other.GetRenderIndices()[RenderFaceIndex] + OtherToLodRenderVerticesOffset;
+			GetRenderMaterialIndex()[RenderFaceIndex] = Other.GetRenderMaterialIndex()[RenderFaceIndex];
+		}
+	}
+
+	void FCollectionClothPatternFacade::SetNumSimVertices(int32 NumSimVertices)
+	{
+		GetClothCollection()->SetNumElements(
+			NumSimVertices,
+			FClothCollection::SimVerticesGroup,
+			GetClothCollection()->GetSimVerticesStart(),
+			GetClothCollection()->GetSimVerticesEnd(),
+			GetElementIndex());
+	}
+
+	void FCollectionClothPatternFacade::SetNumSimFaces(int32 NumSimFaces)
+	{
+		GetClothCollection()->SetNumElements(
+			NumSimFaces,
+			FClothCollection::SimFacesGroup,
+			GetClothCollection()->GetSimFacesStart(),
+			GetClothCollection()->GetSimFacesEnd(),
+			GetElementIndex());
+	}
+
+	void FCollectionClothPatternFacade::SetNumRenderVertices(int32 NumRenderVertices)
+	{
+		GetClothCollection()->SetNumElements(
+			NumRenderVertices,
+			FClothCollection::RenderVerticesGroup,
+			GetClothCollection()->GetRenderVerticesStart(),
+			GetClothCollection()->GetRenderVerticesEnd(),
+			GetElementIndex());
+	}
+
+	void FCollectionClothPatternFacade::SetNumRenderFaces(int32 NumRenderFaces)
+	{
+		GetClothCollection()->SetNumElements(
+			NumRenderFaces,
+			FClothCollection::RenderFacesGroup,
+			GetClothCollection()->GetRenderFacesStart(),
+			GetClothCollection()->GetRenderFacesEnd(),
+			GetElementIndex());
+	}
+
+	void FCollectionClothPatternFacade::SetStatusFlags(int32 StatusFlags)
+	{
+		(*GetClothCollection()->GetStatusFlags())[GetElementIndex()] = StatusFlags;
+	}
+
+	TArrayView<FVector2f> FCollectionClothPatternFacade::GetSimPosition()
+	{
+		return GetClothCollection()->GetElements(
+			GetClothCollection()->GetSimPosition(),
+			GetClothCollection()->GetSimVerticesStart(),
+			GetClothCollection()->GetSimVerticesEnd(),
+			GetElementIndex());
+	}
+
+	TArrayView<FVector3f> FCollectionClothPatternFacade::GetSimRestPosition()
+	{
+		return GetClothCollection()->GetElements(
+			GetClothCollection()->GetSimRestPosition(),
+			GetClothCollection()->GetSimVerticesStart(),
+			GetClothCollection()->GetSimVerticesEnd(),
+			GetElementIndex());
+	}
+
+	TArrayView<FVector3f> FCollectionClothPatternFacade::GetSimRestNormal()
+	{
+		return GetClothCollection()->GetElements(
+			GetClothCollection()->GetSimRestNormal(),
+			GetClothCollection()->GetSimVerticesStart(),
+			GetClothCollection()->GetSimVerticesEnd(),
+			GetElementIndex());
+	}
+
+	TArrayView<int32> FCollectionClothPatternFacade::GetSimNumBoneInfluences()
+	{
+		return GetClothCollection()->GetElements(
+			GetClothCollection()->GetSimNumBoneInfluences(),
+			GetClothCollection()->GetSimVerticesStart(),
+			GetClothCollection()->GetSimVerticesEnd(),
+			GetElementIndex());
+	}
+
+	TArrayView<TArray<int32>> FCollectionClothPatternFacade::GetSimBoneIndices()
+	{
+		return GetClothCollection()->GetElements(
+			GetClothCollection()->GetSimBoneIndices(),
+			GetClothCollection()->GetSimVerticesStart(),
+			GetClothCollection()->GetSimVerticesEnd(),
+			GetElementIndex());
+	}
+
+	TArrayView<TArray<float>> FCollectionClothPatternFacade::GetSimBoneWeights()
+	{
+		return GetClothCollection()->GetElements(
+			GetClothCollection()->GetSimBoneWeights(),
+			GetClothCollection()->GetSimVerticesStart(),
+			GetClothCollection()->GetSimVerticesEnd(),
+			GetElementIndex());
+	}
+
+	TArrayView<FIntVector3> FCollectionClothPatternFacade::GetSimIndices()
+	{
+		return GetClothCollection()->GetElements(
+			GetClothCollection()->GetSimIndices(),
+			GetClothCollection()->GetSimFacesStart(),
+			GetClothCollection()->GetSimFacesEnd(),
+			GetElementIndex());
+	}
+
+	TArrayView<FVector3f> FCollectionClothPatternFacade::GetRenderPosition()
+	{
+		return GetClothCollection()->GetElements(
+			GetClothCollection()->GetRenderPosition(),
+			GetClothCollection()->GetRenderVerticesStart(),
+			GetClothCollection()->GetRenderVerticesEnd(),
+			GetElementIndex());
+	}
+
+	TArrayView<FVector3f> FCollectionClothPatternFacade::GetRenderNormal()
+	{
+		return GetClothCollection()->GetElements(
+			GetClothCollection()->GetRenderNormal(),
+			GetClothCollection()->GetRenderVerticesStart(),
+			GetClothCollection()->GetRenderVerticesEnd(),
+			GetElementIndex());
+	}
+
+	TArrayView<FVector3f> FCollectionClothPatternFacade::GetRenderTangentU()
+	{
+		return GetClothCollection()->GetElements(
+			GetClothCollection()->GetRenderTangentU(),
+			GetClothCollection()->GetRenderVerticesStart(),
+			GetClothCollection()->GetRenderVerticesEnd(),
+			GetElementIndex());
+	}
+
+	TArrayView<FVector3f> FCollectionClothPatternFacade::GetRenderTangentV()
+	{
+		return GetClothCollection()->GetElements(
+			GetClothCollection()->GetRenderTangentV(),
+			GetClothCollection()->GetRenderVerticesStart(),
+			GetClothCollection()->GetRenderVerticesEnd(),
+			GetElementIndex());
+	}
+
+	TArrayView<TArray<FVector2f>> FCollectionClothPatternFacade::GetRenderUVs()
+	{
+		return GetClothCollection()->GetElements(
+			GetClothCollection()->GetRenderUVs(),
+			GetClothCollection()->GetRenderVerticesStart(),
+			GetClothCollection()->GetRenderVerticesEnd(),
+			GetElementIndex());
+	}
+
+	TArrayView<FLinearColor> FCollectionClothPatternFacade::GetRenderColor()
+	{
+		return GetClothCollection()->GetElements(
+			GetClothCollection()->GetRenderColor(),
+			GetClothCollection()->GetRenderVerticesStart(),
+			GetClothCollection()->GetRenderVerticesEnd(),
+			GetElementIndex());
+	}
+
+	TArrayView<int32> FCollectionClothPatternFacade::GetRenderNumBoneInfluences()
+	{
+		return GetClothCollection()->GetElements(
+			GetClothCollection()->GetRenderNumBoneInfluences(),
+			GetClothCollection()->GetRenderVerticesStart(),
+			GetClothCollection()->GetRenderVerticesEnd(),
+			GetElementIndex());
+	}
+
+	TArrayView<TArray<int32>> FCollectionClothPatternFacade::GetRenderBoneIndices()
+	{
+		return GetClothCollection()->GetElements(
+			GetClothCollection()->GetRenderBoneIndices(),
+			GetClothCollection()->GetRenderVerticesStart(),
+			GetClothCollection()->GetRenderVerticesEnd(),
+			GetElementIndex());
+	}
+
+	TArrayView<TArray<float>> FCollectionClothPatternFacade::GetRenderBoneWeights()
+	{
+		return GetClothCollection()->GetElements(
+			GetClothCollection()->GetRenderBoneWeights(),
+			GetClothCollection()->GetRenderVerticesStart(),
+			GetClothCollection()->GetRenderVerticesEnd(),
+			GetElementIndex());
+	}
+
+	TArrayView<FIntVector3> FCollectionClothPatternFacade::GetRenderIndices()
+	{
+		return GetClothCollection()->GetElements(
+			GetClothCollection()->GetRenderIndices(),
+			GetClothCollection()->GetRenderFacesStart(),
+			GetClothCollection()->GetRenderFacesEnd(),
+			GetElementIndex());
+	}
+
+	TArrayView<int32> FCollectionClothPatternFacade::GetRenderMaterialIndex()
+	{
+		return GetClothCollection()->GetElements(
+			GetClothCollection()->GetRenderMaterialIndex(),
+			GetClothCollection()->GetRenderFacesStart(),
+			GetClothCollection()->GetRenderFacesEnd(),
+			GetElementIndex());
+	}
+
+	TArrayView<float> FCollectionClothPatternFacade::GetWeightMap(const FName& Name)
+	{
+		TManagedArray<float>* const WeightMap = GetClothCollection()->GetUserDefinedAttribute<float>(Name, FClothCollection::SimVerticesGroup);
+		return GetClothCollection()->GetElements(
+			WeightMap,
+			GetClothCollection()->GetSimVerticesStart(),
+			GetClothCollection()->GetSimVerticesEnd(),
+			GetElementIndex());
+	}
+
+	FCollectionClothPatternFacade::FCollectionClothPatternFacade(const TSharedPtr<FClothCollection>& ClothCollection, int32 InLodIndex, int32 InPatternIndex)
+		: FCollectionClothPatternConstFacade(ClothCollection, InLodIndex, InPatternIndex)
+	{
+	}
+
+	void FCollectionClothPatternFacade::SetDefaults()
+	{
+		const int32 ElementIndex = GetElementIndex();
+
+		(*GetClothCollection()->GetSimVerticesStart())[ElementIndex] = INDEX_NONE;
+		(*GetClothCollection()->GetSimVerticesEnd())[ElementIndex] = INDEX_NONE;
+		(*GetClothCollection()->GetSimFacesStart())[ElementIndex] = INDEX_NONE;
+		(*GetClothCollection()->GetSimFacesEnd())[ElementIndex] = INDEX_NONE;
+		(*GetClothCollection()->GetRenderVerticesStart())[ElementIndex] = INDEX_NONE;
+		(*GetClothCollection()->GetRenderVerticesEnd())[ElementIndex] = INDEX_NONE;
+		(*GetClothCollection()->GetRenderFacesStart())[ElementIndex] = INDEX_NONE;
+		(*GetClothCollection()->GetRenderFacesEnd())[ElementIndex] = INDEX_NONE;
+		(*GetClothCollection()->GetStatusFlags())[ElementIndex] = 0;
+	}
+}  // End namespace UE::Chaos::ClothAsset
diff --git a/Engine/Plugins/Experimental/ChaosClothAsset/Source/ChaosClothAsset/Public/ChaosClothAsset/ClothAdapter.h b/Engine/Plugins/Experimental/ChaosClothAsset/Source/ChaosClothAsset/Public/ChaosClothAsset/ClothAdapter.h
deleted file mode 100644
diff --git a/Engine/Plugins/Experimental/ChaosClothAsset/Source/ChaosClothAsset/Public/ChaosClothAsset/ClothCollection.h b/Engine/Plugins/Experimental/ChaosClothAsset/Source/ChaosClothAsset/Public/ChaosClothAsset/ClothCollection.h
--- a/Engine/Plugins/Experimental/ChaosClothAsset/Source/ChaosClothAsset/Public/ChaosClothAsset/ClothCollection.h
+++ b/Engine/Plugins/Experimental/ChaosClothAsset/Source/ChaosClothAsset/Public/ChaosClothAsset/ClothCollection.h
@@ -2,185 +2,429 @@
 
 #pragma once
 
-#include "GeometryCollection/ManagedArrayCollection.h"
+#include "GeometryCollection/ManagedArray.h"
 
-// TODO:
-// - Transition Sim LOD data?
+struct FManagedArrayCollection;
+template<class InElementType> class TManagedArray;
 
 namespace UE::Chaos::ClothAsset
 {
+	/** User defined attribute types. */
+	template<typename T> struct TIsUserAttributeType { static constexpr bool Value = false; };
+	template<> struct TIsUserAttributeType<bool> { static constexpr bool Value = true; };
+	template<> struct TIsUserAttributeType<int32> { static constexpr bool Value = true; };
+	template<> struct TIsUserAttributeType<float> { static constexpr bool Value = true; };
+	template<> struct TIsUserAttributeType<FVector3f> { static constexpr bool Value = true; };
+
 	/**
-	 * Tailored Cloth Asset Collection containing draping and pattern information.
+	 * Cloth collection facade data.
 	 */
-	class CHAOSCLOTHASSET_API FClothCollection : public FManagedArrayCollection
+	class FClothCollection final
 	{
 	public:
-		typedef FManagedArrayCollection Super;
+		static const FName LodsGroup;  // LOD information
+		static const FName MaterialsGroup;  // Materials information
+		static const FName TetherBatchesGroup;  // Tethers parallel batch processing information
+		static const FName TethersGroup;  // Tethers information
+		static const FName SeamsGroup;  // Contains pairs of stitched sim vertex indices and matching pattern indices
+		static const FName PatternsGroup;  // Contains pattern relationships to other groups
+		static const FName SimFacesGroup;  // Contains indices to sim vertex
+		static const FName SimVerticesGroup;  // Contains patterns' 2D positions, 3D draped position (rest)
+		static const FName RenderFacesGroup;  // Contains indices to render vertex
+		static const FName RenderVerticesGroup;  // Contains pattern's 3D render model
+
+		explicit FClothCollection(const TSharedPtr<FManagedArrayCollection>& InManagedArrayCollection);
+		~FClothCollection() = default;
 
-		FClothCollection();
-		FClothCollection(FClothCollection&) = delete;
+		FClothCollection(const FClothCollection&) = delete;
 		FClothCollection& operator=(const FClothCollection&) = delete;
-		FClothCollection(FClothCollection&&) = default;
-		FClothCollection& operator=(FClothCollection&&) = default;
+		FClothCollection(FClothCollection&&) = delete;
+		FClothCollection& operator=(FClothCollection&&) = delete;
+
+		/** Return whether the underlying collection is a valid cloth collection. */
+		bool IsValid() const;
+
+		/** Make the underlying collection a cloth collection. */
+		void DefineSchema();
 
-		//~ Begin FManagedArrayCollection interface
-		virtual void Reset() override { Super::Reset(); Construct(); }
-		//~ Begin FManagedArrayCollection interface
+		/** Get the number of elements of a group. */
+		int32 GetNumElements(const FName& GroupName) const;
 
-		using Super::Serialize;
-		void Serialize(FArchive& Ar);
+		/** Get the number of elements of one of the sub groups that have start/end indices. */
+		int32 GetNumElements(const TManagedArray<int32>* StartArray, const TManagedArray<int32>* EndArray, int32 Index) const;
 
-		/** Set the number of elements to one of the groups that have start/end indices while maintaining the correct order of the data. */
-		int32 SetNumElements(int32 InNumElements, const FName& GroupName, TManagedArray<int32>& StartArray, TManagedArray<int32>& EndArray, int32 StartEndIndex);
+		/** Set the number of elements of a group. */
+		void SetNumElements(int32 InNumElements, const FName& GroupName);
 
-		/** Get the number of elements to one of the groups that have start/end indices. */
-		int32 GetNumElements(const TManagedArray<int32>& StartArray, const TManagedArray<int32>& EndArray, int32 StartEndIndex) const;
+		/** Set the number of elements to one of the sub groups while maintaining the correct order of the data, and return the first index of the range. */
+		int32 SetNumElements(int32 InNumElements, const FName& GroupName, TManagedArray<int32>* StartArray, TManagedArray<int32>* EndArray, int32 Index);
 
 		template<typename T>
-		inline TConstArrayView<T> GetElements(const TManagedArray<T>& ElementArray, const TManagedArray<int32>& StartArray, const TManagedArray<int32>& EndArray, int32 StartEndIndex) const;
+		inline TConstArrayView<T> GetElements(const TManagedArray<T>* ElementArray, const TManagedArray<int32>* StartArray, const TManagedArray<int32>* EndArray, int32 ArrayIndex) const;
+
+		template<typename T>
+		inline TArrayView<T> GetElements(TManagedArray<T>* ElementArray, const TManagedArray<int32>* StartArray, const TManagedArray<int32>* EndArray, int32 ArrayIndex);
+
+		/**
+		 * Return the difference between the start index of an element to the start index of the first sub-element in the group (base).
+		 * Usefull for getting back and forth between LOD/pattern indexation modes.
+		 */
+		int32 GetElementsOffset(const TManagedArray<int32>* StartArray, int32 BaseElementIndex, int32 ElementIndex) const { return (*StartArray)[ElementIndex] - (*StartArray)[BaseElementIndex]; }
+
+		int32 GetNumSubElements(
+			const TManagedArray<int32>* StartArray,
+			const TManagedArray<int32>* EndArray,
+			const TManagedArray<int32>* StartSubArray,
+			const TManagedArray<int32>* EndSubArray,
+			int32 LodIndex) const;
+
+		UE_DEPRECATED(5.3, "Use GetNumSubElements instead.")
+		int32 GetPatternsNumElements(const TManagedArray<int32>* StartSubArray, const TManagedArray<int32>* EndSubArray, int32 LodIndex) const
+		{
+			return GetNumSubElements( PatternStart, PatternEnd, StartSubArray, EndSubArray, LodIndex);
+		}
+
 		template<typename T>
-		inline TArrayView<T> GetElements(TManagedArray<T>& ElementArray, const TManagedArray<int32>& StartArray, const TManagedArray<int32>& EndArray, int32 StartEndIndex);
+		inline TConstArrayView<T> GetSubElements(
+			const TManagedArray<T>* SubElementArray,
+			const TManagedArray<int32>* StartArray,
+			const TManagedArray<int32>* EndArray,
+			const TManagedArray<int32>* StartSubArray,
+			const TManagedArray<int32>* EndSubArray,
+			int32 ArrayIndex) const;
 
-		int32 GetPatternsNumElements(const TManagedArray<int32>& StartArray, const TManagedArray<int32>& EndArray, int32 LodIndex) const;
+		template<typename T>
+		UE_DEPRECATED(5.3, "Use GetSubElements instead.")
+		TConstArrayView<T> GetPatternElements(const TManagedArray<T>* SubElementArray, const TManagedArray<int32>* StartSubArray, const TManagedArray<int32>* EndSubArray, int32 LodIndex) const
+		{
+			return GetSubElements(SubElementArray, PatternStart, PatternEnd, StartSubArray, EndSubArray,LodIndex);
+		}
 
 		template<typename T>
-		inline TConstArrayView<T> GetPatternsElements(const TManagedArray<T>& ElementArray, const TManagedArray<int32>& StartArray, const TManagedArray<int32>& EndArray, int32 LodIndex) const;
+		inline TArrayView<T> GetSubElements(
+			TManagedArray<T>* SubElementArray,
+			const TManagedArray<int32>* StartArray,
+			const TManagedArray<int32>* EndArray,
+			const TManagedArray<int32>* StartSubArray,
+			const TManagedArray<int32>* EndSubArray,
+			int32 ArrayIndex);
+
 		template<typename T>
-		inline TArrayView<T> GetPatternsElements(TManagedArray<T>& ElementArray, const TManagedArray<int32>& StartArray, const TManagedArray<int32>& EndArray, int32 LodIndex);
+		UE_DEPRECATED(5.3, "Use GetSubElements instead.")
+		TArrayView<T> GetPatternElements(TManagedArray<T>* SubElementArray, const TManagedArray<int32>* StartSubArray, const TManagedArray<int32>* EndSubArray, int32 LodIndex)
+		{
+			return GetSubElements(SubElementArray, PatternStart, PatternEnd, StartSubArray, EndSubArray, LodIndex);
+		}
 
 		template<bool bStart = true, bool bEnd = true>
-		TTuple<int32, int32> GetPatternsElementsStartEnd(const TManagedArray<int32>& StartArray, const TManagedArray<int32>& EndArray, int32 StartEndIndex) const;
+		TTuple<int32, int32> GetSubElementsStartEnd(
+			const TManagedArray<int32>* StartArray,
+			const TManagedArray<int32>* EndArray,
+			const TManagedArray<int32>* StartSubArray,
+			const TManagedArray<int32>* EndSubArray,
+			int32 ArrayIndex) const;
 
-		// Attribute groups, predefined data member of the collection
-		static const FName SimVerticesGroup;  // Contains patterns' 2D positions, 3D draped position (rest)
-		static const FName SimFacesGroup;  // Contains indices to sim vertex
-		static const FName RenderVerticesGroup;  // Contains pattern's 3D render model
-		static const FName RenderFacesGroup;  // Contains indices to render vertex
-		static const FName WrapDeformersGroup;  // Contains the wrap deformers cloth capture information
-		static const FName PatternsGroup;  // Contains pattern relationships to other groups
-		static const FName SeamsGroup;  // Contains pairs of stitched sim vertex indices and matching pattern indices
-		static const FName TethersGroup;  // Tethers information
-		static const FName TetherBatchesGroup;  // Tethers parallel batch processing information
-		static const FName LodsGroup;  // Lod split information
-		static const FName MaterialsGroup;  // Material information
-		static const FName CollisionsGroup;  // Collision information
-		static const FName SkeletonsGroup;  // Skeleton information
-
-		// Sim Vertices Group
-		TManagedArray<FVector2f> SimPosition;
-		TManagedArray<FVector3f> SimRestPosition;
-		TManagedArray<FVector3f> SimRestNormal;  // Used for capture, maxdistance, backstop authoring ...etc
-		TManagedArray<int32> SimNumBoneInfluences;
-		TManagedArray<TArray<int32>> SimBoneIndices;
-		TManagedArray<TArray<float>> SimBoneWeights;
-
-		// Sim Faces Group
-		TManagedArray<FIntVector3> SimIndices;  // The indices point to the elements in the Sim Vertices arrays but don't include the LOD start offset
-
-		// Render Vertices Group
-		TManagedArray<FVector3f> RenderPosition;
-		TManagedArray<FVector3f> RenderNormal;
-		TManagedArray<FVector3f> RenderTangentU;
-		TManagedArray<FVector3f> RenderTangentV;
-		TManagedArray<TArray<FVector2f>> RenderUVs;
-		TManagedArray<FLinearColor> RenderColor;
-		TManagedArray<int32> RenderNumBoneInfluences;
-		TManagedArray<TArray<int32>> RenderBoneIndices;
-		TManagedArray<TArray<float>> RenderBoneWeights;
-
-		// Render Faces Group
-		TManagedArray<FIntVector3> RenderIndices;  // The indices point to the elements in the Render Vertices arrays but don't include the LOD start offset
-		TManagedArray<int32> RenderMaterialIndex;  // Render material per triangle
-
-		// Patterns Group
-		TManagedArray<int32> SimVerticesStart;
-		TManagedArray<int32> SimVerticesEnd;
-		TManagedArray<int32> SimFacesStart;
-		TManagedArray<int32> SimFacesEnd;
-		TManagedArray<int32> RenderVerticesStart;
-		TManagedArray<int32> RenderVerticesEnd;
-		TManagedArray<int32> RenderFacesStart;
-		TManagedArray<int32> RenderFacesEnd;
-		TManagedArray<int32> WrapDeformerStart;
-		TManagedArray<int32> WrapDeformerEnd;
-		TManagedArray<int32> NumWeights;  // Number of weights stored between WrapDeformerStart and WrapDeformerEnd
-		TManagedArray<int32> StatusFlags;  // Whether this pattern should be rendered, simulated, ...etc
-		TManagedArray<int32> SimMaterialIndex;  // Cloth material per pattern
-
-		// Seam Group
-		TManagedArray<FIntVector2> SeamPatterns;  // Stitched pattern indices (LOD based indices) for SeamIndicesFirst
-		TManagedArray<TArray<FIntVector2>> SeamStitches;  // Stitched vertex indices pair (LOD based indices)
-
-		// Tethers Group
-		TManagedArray<int32> TetherKinematicIndex;
-		TManagedArray<int32> TetherDynamicIndex;
-		TManagedArray<float> TetherReferenceLength;
-
-		// Tether Batches Group
-		TManagedArray<int32> TetherStart;
-		TManagedArray<int32> TetherEnd;
-
-		// LODs Group
-		TManagedArray<int32> PatternStart;
-		TManagedArray<int32> PatternEnd;
-		TManagedArray<int32> SeamStart;
-		TManagedArray<int32> SeamEnd;
-		TManagedArray<int32> TetherBatchStart;
-		TManagedArray<int32> TetherBatchEnd;
-		TManagedArray<int32> LodBiasDepth;  // The number of LODs covered by each Sim LOD (for the wrap deformer)
-
-		// Materials Group
-		TManagedArray<FString> MaterialPathName;
-
-		// Collisions Group
-		TManagedArray<FString> PhysicsAssetPathName;
-
-		// Skeletons Group
-		TManagedArray<FString> SkeletonAssetPathName;
-
-	protected:
-		void Construct();
+		template<bool bStart = true, bool bEnd = true>
+		UE_DEPRECATED(5.3, "Use GetSubElementsStartEnd instead.")
+		TTuple<int32, int32> GetPatternElementsStartEnd(const TManagedArray<int32>* StartSubArray, const TManagedArray<int32>* EndSubArray, int32 LodIndex) const
+		{
+			return GetSubElementsStartEnd(PatternStart, PatternEnd, StartSubArray, EndSubArray, LodIndex);
+		}
+
+		//~ Weight maps
+		template<typename T, TEMPLATE_REQUIRES(TIsUserAttributeType<T>::Value)>
+		TArray<FName> GetUserDefinedAttributeNames(const FName& GroupName) const;
+
+		template<typename T, TEMPLATE_REQUIRES(TIsUserAttributeType<T>::Value)>
+		void AddUserDefinedAttribute(const FName& Name, const FName& GroupName);
+
+		void RemoveUserDefinedAttribute(const FName& Name, const FName& GroupName);
+
+		template<typename T, TEMPLATE_REQUIRES(TIsUserAttributeType<T>::Value)>
+		bool HasUserDefinedAttribute(const FName& Name, const FName& GroupName) const;
+
+		template<typename T, TEMPLATE_REQUIRES(TIsUserAttributeType<T>::Value)>
+		const TManagedArray<T>* GetUserDefinedAttribute(const FName& Name, const FName& GroupName) const;
+
+		template<typename T, TEMPLATE_REQUIRES(TIsUserAttributeType<T>::Value)>
+		TManagedArray<T>* GetUserDefinedAttribute(const FName& Name, const FName& GroupName);
+
+
+		//~ LODs Group
+		const TManagedArray<int32>* GetMaterialStart() const { return MaterialStart; }
+		const TManagedArray<int32>* GetMaterialEnd() const { return MaterialEnd; }
+		const TManagedArray<int32>* GetTetherBatchStart() const { return TetherBatchStart; }
+		const TManagedArray<int32>* GetTetherBatchEnd() const { return TetherBatchEnd; }
+		const TManagedArray<int32>* GetSeamStart() const { return SeamStart; }
+		const TManagedArray<int32>* GetSeamEnd() const { return SeamEnd; }
+		const TManagedArray<int32>* GetPatternStart() const { return PatternStart; }
+		const TManagedArray<int32>* GetPatternEnd() const { return PatternEnd; }
+		const TManagedArray<FString>* GetPhysicsAssetPathName() const { return PhysicsAssetPathName; }
+		const TManagedArray<FString>* GetSkeletonAssetPathName() const { return SkeletonAssetPathName; }
+
+		//~ Materials Group
+		const TManagedArray<FString>* GetRenderMaterialPathName() const { return RenderMaterialPathName; }
+
+		//~ Tether Batches Group
+		const TManagedArray<int32>* GetTetherStart() const { return TetherStart; }
+		const TManagedArray<int32>* GetTetherEnd() const { return TetherEnd; }
+
+		//~ Tethers Group
+		const TManagedArray<int32>* GetTetherKinematicIndex() const { return TetherKinematicIndex; }
+		const TManagedArray<int32>* GetTetherDynamicIndex() const { return TetherDynamicIndex; }
+		const TManagedArray<float>* GetTetherReferenceLength() const { return TetherReferenceLength; }
+
+		//~ Seam Group
+		const TManagedArray<FIntVector2>* GetSeamPatterns() const { return SeamPatterns; }
+		const TManagedArray<TArray<FIntVector2>>* GetSeamStitches() const { return SeamStitches; }
+
+		//~ Patterns Group
+		const TManagedArray<int32>* GetSimVerticesStart() const { return SimVerticesStart; }
+		const TManagedArray<int32>* GetSimVerticesEnd() const { return SimVerticesEnd; }
+		const TManagedArray<int32>* GetSimFacesStart() const { return SimFacesStart; }
+		const TManagedArray<int32>* GetSimFacesEnd() const { return SimFacesEnd; }
+		const TManagedArray<int32>* GetRenderVerticesStart() const { return RenderVerticesStart; }
+		const TManagedArray<int32>* GetRenderVerticesEnd() const { return RenderVerticesEnd; }
+		const TManagedArray<int32>* GetRenderFacesStart() const { return RenderFacesStart; }
+		const TManagedArray<int32>* GetRenderFacesEnd() const { return RenderFacesEnd; }
+		const TManagedArray<int32>* GetStatusFlags() const { return StatusFlags; }
+
+		//~ Sim Faces Group
+		const TManagedArray<FIntVector3>* GetSimIndices() const { return SimIndices; }
+
+		//~ Sim Vertices Group
+		const TManagedArray<FVector2f>* GetSimPosition() const { return SimPosition; }
+		const TManagedArray<FVector3f>* GetSimRestPosition() const { return SimRestPosition; }
+		const TManagedArray<FVector3f>* GetSimRestNormal() const { return SimRestNormal; }
+		const TManagedArray<int32>* GetSimNumBoneInfluences() const { return SimNumBoneInfluences; }
+		const TManagedArray<TArray<int32>>* GetSimBoneIndices() const { return SimBoneIndices; }
+		const TManagedArray<TArray<float>>* GetSimBoneWeights() const { return SimBoneWeights; }
+
+		//~ Render Faces Group
+		const TManagedArray<FIntVector3>* GetRenderIndices() const { return RenderIndices; }
+		const TManagedArray<int32>* GetRenderMaterialIndex() const { return RenderMaterialIndex; }
+
+		//~ Render Vertices Group
+		const TManagedArray<FVector3f>* GetRenderPosition() const { return RenderPosition; }
+		const TManagedArray<FVector3f>* GetRenderNormal() const { return RenderNormal; }
+		const TManagedArray<FVector3f>* GetRenderTangentU() const { return RenderTangentU; }
+		const TManagedArray<FVector3f>* GetRenderTangentV() const { return RenderTangentV; }
+		const TManagedArray<TArray<FVector2f>>* GetRenderUVs() const { return RenderUVs; }
+		const TManagedArray<FLinearColor>* GetRenderColor() const { return RenderColor; }
+		const TManagedArray<int32>* GetRenderNumBoneInfluences() const { return RenderNumBoneInfluences; }
+		const TManagedArray<TArray<int32>>* GetRenderBoneIndices() const { return RenderBoneIndices; }
+		const TManagedArray<TArray<float>>* GetRenderBoneWeights() const { return RenderBoneWeights; }
+
+		//~ LODs Group
+		TManagedArray<int32>* GetMaterialStart() { return MaterialStart; }
+		TManagedArray<int32>* GetMaterialEnd() { return MaterialEnd; }
+		TManagedArray<int32>* GetTetherBatchStart() { return TetherBatchStart; }
+		TManagedArray<int32>* GetTetherBatchEnd() { return TetherBatchEnd; }
+		TManagedArray<int32>* GetSeamStart() { return SeamStart; }
+		TManagedArray<int32>* GetSeamEnd() { return SeamEnd; }
+		TManagedArray<int32>* GetPatternStart() { return PatternStart; }
+		TManagedArray<int32>* GetPatternEnd() { return PatternEnd; }
+		TManagedArray<FString>* GetPhysicsAssetPathName() { return PhysicsAssetPathName; }
+		TManagedArray<FString>* GetSkeletonAssetPathName() { return SkeletonAssetPathName; }
+
+		//~ Materials Group
+		TManagedArray<FString>* GetRenderMaterialPathName() { return RenderMaterialPathName; }
+
+		//~ Tether Batches Group
+		TManagedArray<int32>* GetTetherStart() { return TetherStart; }
+		TManagedArray<int32>* GetTetherEnd() { return TetherEnd; }
+
+		//~ Tethers Group
+		TManagedArray<int32>* GetTetherKinematicIndex() { return TetherKinematicIndex; }
+		TManagedArray<int32>* GetTetherDynamicIndex() { return TetherDynamicIndex; }
+		TManagedArray<float>* GetTetherReferenceLength() { return TetherReferenceLength; }
+
+		//~ Seam Group
+		TManagedArray<FIntVector2>* GetSeamPatterns() { return SeamPatterns; }
+		TManagedArray<TArray<FIntVector2>>* GetSeamStitches() { return SeamStitches; }
+
+		//~ Patterns Group
+		TManagedArray<int32>* GetSimVerticesStart() { return SimVerticesStart; }
+		TManagedArray<int32>* GetSimVerticesEnd() { return SimVerticesEnd; }
+		TManagedArray<int32>* GetSimFacesStart() { return SimFacesStart; }
+		TManagedArray<int32>* GetSimFacesEnd() { return SimFacesEnd; }
+		TManagedArray<int32>* GetRenderVerticesStart() { return RenderVerticesStart; }
+		TManagedArray<int32>* GetRenderVerticesEnd() { return RenderVerticesEnd; }
+		TManagedArray<int32>* GetRenderFacesStart() { return RenderFacesStart; }
+		TManagedArray<int32>* GetRenderFacesEnd() { return RenderFacesEnd; }
+		TManagedArray<int32>* GetStatusFlags() { return StatusFlags; }
+
+		//~ Sim Faces Group
+		TManagedArray<FIntVector3>* GetSimIndices() { return SimIndices; }
+
+		//~ Sim Vertices Group
+		TManagedArray<FVector2f>* GetSimPosition() { return SimPosition; }
+		TManagedArray<FVector3f>* GetSimRestPosition() { return SimRestPosition; }
+		TManagedArray<FVector3f>* GetSimRestNormal() { return SimRestNormal; }
+		TManagedArray<int32>* GetSimNumBoneInfluences() { return SimNumBoneInfluences; }
+		TManagedArray<TArray<int32>>* GetSimBoneIndices() { return SimBoneIndices; }
+		TManagedArray<TArray<float>>* GetSimBoneWeights() { return SimBoneWeights; }
+
+		//~ Render Faces Group
+		TManagedArray<FIntVector3>* GetRenderIndices() { return RenderIndices; }
+		TManagedArray<int32>* GetRenderMaterialIndex() { return RenderMaterialIndex; }
+
+		//~ Render Vertices Group
+		TManagedArray<FVector3f>* GetRenderPosition() { return RenderPosition; }
+		TManagedArray<FVector3f>* GetRenderNormal() { return RenderNormal; }
+		TManagedArray<FVector3f>* GetRenderTangentU() { return RenderTangentU; }
+		TManagedArray<FVector3f>* GetRenderTangentV() { return RenderTangentV; }
+		TManagedArray<TArray<FVector2f>>* GetRenderUVs() { return RenderUVs; }
+		TManagedArray<FLinearColor>* GetRenderColor() { return RenderColor; }
+		TManagedArray<int32>* GetRenderNumBoneInfluences() { return RenderNumBoneInfluences; }
+		TManagedArray<TArray<int32>>* GetRenderBoneIndices() { return RenderBoneIndices; }
+		TManagedArray<TArray<float>>* GetRenderBoneWeights() { return RenderBoneWeights; }
+
+	private:
+		//~ Cloth collection
+		TSharedPtr<FManagedArrayCollection> ManagedArrayCollection;
+
+		//~ LODs Group
+		TManagedArray<int32>* MaterialStart;
+		TManagedArray<int32>* MaterialEnd;
+		TManagedArray<int32>* TetherBatchStart;
+		TManagedArray<int32>* TetherBatchEnd;
+		TManagedArray<int32>* SeamStart;
+		TManagedArray<int32>* SeamEnd;
+		TManagedArray<int32>* PatternStart;
+		TManagedArray<int32>* PatternEnd;
+		TManagedArray<FString>* PhysicsAssetPathName;
+		TManagedArray<FString>* SkeletonAssetPathName;
+
+		//~ Materials Group
+		TManagedArray<FString>* RenderMaterialPathName;
+
+		//~ Tether Batches Group
+		TManagedArray<int32>* TetherStart;
+		TManagedArray<int32>* TetherEnd;
+
+		//~ Tethers Group
+		TManagedArray<int32>* TetherKinematicIndex;
+		TManagedArray<int32>* TetherDynamicIndex;
+		TManagedArray<float>* TetherReferenceLength;
+
+		//~ Seam Group
+		TManagedArray<FIntVector2>* SeamPatterns;  // Stitched pattern indices (LOD based indices) for SeamIndicesFirst
+		TManagedArray<TArray<FIntVector2>>* SeamStitches;  // Stitched vertex indices pair (LOD based indices)
+
+		//~ Patterns Group
+		TManagedArray<int32>* SimVerticesStart;
+		TManagedArray<int32>* SimVerticesEnd;
+		TManagedArray<int32>* SimFacesStart;
+		TManagedArray<int32>* SimFacesEnd;
+		TManagedArray<int32>* RenderVerticesStart;
+		TManagedArray<int32>* RenderVerticesEnd;
+		TManagedArray<int32>* RenderFacesStart;
+		TManagedArray<int32>* RenderFacesEnd;
+		TManagedArray<int32>* StatusFlags;  // Whether this pattern should be rendered, simulated, ...etc
+
+		//~ Sim Faces Group
+		TManagedArray<FIntVector3>* SimIndices;  // The indices point to the elements in the Sim Vertices arrays but don't include the LOD start offset
+
+		//~ Sim Vertices Group
+		TManagedArray<FVector2f>* SimPosition;
+		TManagedArray<FVector3f>* SimRestPosition;
+		TManagedArray<FVector3f>* SimRestNormal;  // Used for capture, maxdistance, backstop authoring ...etc
+		TManagedArray<int32>* SimNumBoneInfluences;
+		TManagedArray<TArray<int32>>* SimBoneIndices;
+		TManagedArray<TArray<float>>* SimBoneWeights;
+
+		//~ Render Faces Group
+		TManagedArray<FIntVector3>* RenderIndices;  // The indices point to the elements in the Render Vertices arrays but don't include the LOD start offset
+		TManagedArray<int32>* RenderMaterialIndex;  // Render material per triangle
+
+		//~ Render Vertices Group
+		TManagedArray<FVector3f>* RenderPosition;
+		TManagedArray<FVector3f>* RenderNormal;
+		TManagedArray<FVector3f>* RenderTangentU;
+		TManagedArray<FVector3f>* RenderTangentV;
+		TManagedArray<TArray<FVector2f>>* RenderUVs;
+		TManagedArray<FLinearColor>* RenderColor;
+		TManagedArray<int32>* RenderNumBoneInfluences;
+		TManagedArray<TArray<int32>>* RenderBoneIndices;
+		TManagedArray<TArray<float>>* RenderBoneWeights;
 	};
 
 	template<typename T>
-	inline TConstArrayView<T> FClothCollection::GetElements(const TManagedArray<T>& ElementArray, const TManagedArray<int32>& StartArray, const TManagedArray<int32>& EndArray, int32 StartEndIndex) const
+	inline TConstArrayView<T> FClothCollection::GetElements(const TManagedArray<T>* ElementArray, const TManagedArray<int32>* StartArray, const TManagedArray<int32>* EndArray, int32 ArrayIndex) const
 	{
-		const int32 Start = StartArray[StartEndIndex];
-		const int32 End = EndArray[StartEndIndex];
-		check(Start != INDEX_NONE || End == INDEX_NONE);  // Best to avoid situations where only one boundary of the range is set to INDEX_NONE
-
-		return Start == INDEX_NONE ? TConstArrayView<T>() : TConstArrayView<T>(ElementArray.GetData() + Start, End - Start + 1);
+		if (StartArray && EndArray && ElementArray)
+		{
+			const int32 Start = (*StartArray)[ArrayIndex];
+			const int32 End = (*EndArray)[ArrayIndex];
+			if (Start != INDEX_NONE && End != INDEX_NONE)
+			{
+				return TConstArrayView<T>(ElementArray->GetData() + Start, End - Start + 1);
+			}
+			checkf(Start == End, TEXT("Only one boundary of the range is set to INDEX_NONE, when both should."));
+		}
+		return TConstArrayView<T>();
 	}
 
 	template<typename T>
-	inline TArrayView<T> FClothCollection::GetElements(TManagedArray<T>& ElementArray, const TManagedArray<int32>& StartArray, const TManagedArray<int32>& EndArray, int32 StartEndIndex)
+	inline TArrayView<T> FClothCollection::GetElements(TManagedArray<T>* ElementArray, const TManagedArray<int32>* StartArray, const TManagedArray<int32>* EndArray, int32 ArrayIndex)
 	{
-		const int32 Start = StartArray[StartEndIndex];
-		const int32 End = EndArray[StartEndIndex];
-		check(Start != INDEX_NONE || End == INDEX_NONE);  // Best to avoid situations where only one boundary of the range is set to INDEX_NONE
-
-		return Start == INDEX_NONE ? TArrayView<T>() : TArrayView<T>(ElementArray.GetData() + Start, End - Start + 1);
+		if (StartArray && EndArray && ElementArray)
+		{
+			const int32 Start = (*StartArray)[ArrayIndex];
+			const int32 End = (*EndArray)[ArrayIndex];
+			if (Start != INDEX_NONE && End != INDEX_NONE)
+			{
+				return TArrayView<T>(ElementArray->GetData() + Start, End - Start + 1);
+			}
+			checkf(Start == End, TEXT("Only one boundary of the range is set to INDEX_NONE, when both should."));
+		}
+		return TArrayView<T>();
 	}
 
 	template<typename T>
-	inline TConstArrayView<T> FClothCollection::GetPatternsElements(const TManagedArray<T>& ElementArray, const TManagedArray<int32>& StartArray, const TManagedArray<int32>& EndArray, int32 LodIndex) const
+	inline TConstArrayView<T> FClothCollection::GetSubElements(
+		const TManagedArray<T>* SubElementArray,
+		const TManagedArray<int32>* StartArray,
+		const TManagedArray<int32>* EndArray,
+		const TManagedArray<int32>* StartSubArray,
+		const TManagedArray<int32>* EndSubArray,
+		int32 ArrayIndex) const
 	{
-		const TTuple<int32, int32> StartEnd = GetPatternsElementsStartEnd(StartArray, EndArray, LodIndex);
-		const int32 Start = StartEnd.Get<0>();
-		const int32 End = StartEnd.Get<1>();
-
-		check(Start != INDEX_NONE || End == INDEX_NONE);  // Best to avoid situations where only one boundary of the range is set to INDEX_NONE
-		return Start == INDEX_NONE ? TConstArrayView<T>() : TConstArrayView<T>(ElementArray.GetData() + Start, End - Start + 1);
+		if (StartArray && EndArray && SubElementArray)
+		{
+			const TTuple<int32, int32> StartEnd = GetSubElementsStartEnd(StartArray, EndArray, StartSubArray, EndSubArray, ArrayIndex);
+			const int32 Start = StartEnd.Get<0>();
+			const int32 End = StartEnd.Get<1>();
+			if (Start != INDEX_NONE && End != INDEX_NONE)
+			{
+				return TConstArrayView<T>(SubElementArray->GetData() + Start, End - Start + 1);
+			}
+			checkf(Start == End, TEXT("Only one boundary of the range is set to INDEX_NONE, when both should."));
+		}
+		return TConstArrayView<T>();
 	}
 
 	template<typename T>
-	inline TArrayView<T> FClothCollection::GetPatternsElements(TManagedArray<T>& ElementArray, const TManagedArray<int32>& StartArray, const TManagedArray<int32>& EndArray, int32 LodIndex)
+	inline TArrayView<T> FClothCollection::GetSubElements(
+		TManagedArray<T>* SubElementArray,
+		const TManagedArray<int32>* StartArray,
+		const TManagedArray<int32>* EndArray,
+		const TManagedArray<int32>* StartSubArray,
+		const TManagedArray<int32>* EndSubArray,
+		int32 ArrayIndex)
 	{
-		const TTuple<int32, int32> StartEnd = GetPatternsElementsStartEnd(StartArray, EndArray, LodIndex);
-		const int32 Start = StartEnd.Get<0>();
-		const int32 End = StartEnd.Get<1>();
-
-		check(Start != INDEX_NONE || End == INDEX_NONE);  // Best to avoid situations where only one boundary of the range is set to INDEX_NONE
-		return Start == INDEX_NONE ? TArrayView<T>() : TArrayView<T>(ElementArray.GetData() + Start, End - Start + 1);
+		if (StartArray && EndArray && SubElementArray)
+		{
+			const TTuple<int32, int32> StartEnd = GetSubElementsStartEnd(StartArray, EndArray, StartSubArray, EndSubArray, ArrayIndex);
+			const int32 Start = StartEnd.Get<0>();
+			const int32 End = StartEnd.Get<1>();
+			if (Start != INDEX_NONE && End != INDEX_NONE)
+			{
+				return TArrayView<T>(SubElementArray->GetData() + Start, End - Start + 1);
+			}
+			checkf(Start == End, TEXT("Only one boundary of the range is set to INDEX_NONE, when both should."));
+		}
+		return TArrayView<T>();
 	}
 }  // End namespace UE::Chaos::ClothAsset
diff --git a/Engine/Plugins/Experimental/ChaosClothAsset/Source/ChaosClothAsset/Public/ChaosClothAsset/ClothGeometryTools.h b/Engine/Plugins/Experimental/ChaosClothAsset/Source/ChaosClothAsset/Public/ChaosClothAsset/ClothGeometryTools.h
--- a/Engine/Plugins/Experimental/ChaosClothAsset/Source/ChaosClothAsset/Public/ChaosClothAsset/ClothGeometryTools.h
+++ b/Engine/Plugins/Experimental/ChaosClothAsset/Source/ChaosClothAsset/Public/ChaosClothAsset/ClothGeometryTools.h
@@ -2,7 +2,9 @@
 
 #pragma once
 
-#include "ChaosClothAsset/ClothCollection.h"
+#include "Templates/SharedPointer.h"
+
+struct FManagedArrayCollection;
 
 namespace UE::Chaos::ClothAsset
 {
@@ -11,27 +13,40 @@ namespace UE::Chaos::ClothAsset
 	 */
 	struct CHAOSCLOTHASSET_API FClothGeometryTools
 	{
+		/** Return whether at least one pattern of this collection has any faces to simulate. */
+		static bool HasSimMesh(const TSharedPtr<const FManagedArrayCollection>& ClothCollection);
+
+		/** Return whether at least one pattern of this collection has any faces to render. */
+		static bool HasRenderMesh(const TSharedPtr<const FManagedArrayCollection>& ClothCollection);
+
+		/** Set the skeleton path name on each LOD of the specified collection. */
+		static void SetSkeletonAssetPathName(const TSharedPtr<FManagedArrayCollection>& ClothCollection, const FString& SkeletonAssetPathName);
+
+		/** Set the physics path name on each LOD of the specified collection. */
+		static void SetPhysicsAssetPathName(const TSharedPtr<FManagedArrayCollection>& ClothCollection, const FString& PhysicsAssetPathName);
+
 		/** Delete the render mesh data. */
-		static void DeleteRenderMesh(const TSharedPtr<FClothCollection>& ClothCollection);
+		static void DeleteRenderMesh(const TSharedPtr<FManagedArrayCollection>& ClothCollection);
 
 		/** Turn the sim mesh portion of this ClothCollection into a render mesh. */
-		static void CopySimMeshToRenderMesh(const TSharedPtr<FClothCollection>& ClothCollection, int32 MaterialIndex);
+		static void CopySimMeshToRenderMesh(const TSharedPtr<FManagedArrayCollection>& ClothCollection, const FString& RenderMaterialPathName);
 
 		/** Reverse the mesh normals. Will reverse all normals if pattern selection is empty. */
 		static void ReverseNormals(
-			const TSharedPtr<FClothCollection>& ClothCollection,
+			const TSharedPtr<FManagedArrayCollection>& ClothCollection,
 			bool bReverseSimMeshNormals,
 			bool bReverseRenderMeshNormals,
 			const TArray<int32>& PatternSelection);
 
-		/** 
-		 * Set the skinning weights for all of the sim/render vertices in ClothCollection to be bound to the root node. 
-		 * 
-		 * @param LODs if empty will apply the change to all LODs. Otherwise only LODs specified in the array (if exist) are affected. 
+		/**
+		 * Set the skinning weights for all of the sim/render vertices in ClothCollection to be bound to the root node.
+		 *
+		 * @param Lods if empty will apply the change to all LODs. Otherwise only LODs specified in the array (if exist) are affected.
 		 */
-		static void BindMeshToRootBone(const TSharedPtr<FClothCollection>& ClothCollection,
-									   bool bBindSimMesh,
-									   bool bBindRenderMesh,
-									   const TArray<int32> Lods = {});
+		static void BindMeshToRootBone(
+			const TSharedPtr<FManagedArrayCollection>& ClothCollection,
+			bool bBindSimMesh,
+			bool bBindRenderMesh,
+			const TArray<int32> Lods = {});
 	};
 }  // End namespace UE::Chaos::ClothAsset
\ No newline at end of file
diff --git a/Engine/Plugins/Experimental/ChaosClothAsset/Source/ChaosClothAsset/Public/ChaosClothAsset/ClothLodAdapter.h b/Engine/Plugins/Experimental/ChaosClothAsset/Source/ChaosClothAsset/Public/ChaosClothAsset/ClothLodAdapter.h
deleted file mode 100644
diff --git a/Engine/Plugins/Experimental/ChaosClothAsset/Source/ChaosClothAsset/Public/ChaosClothAsset/ClothPatternAdapter.h b/Engine/Plugins/Experimental/ChaosClothAsset/Source/ChaosClothAsset/Public/ChaosClothAsset/ClothPatternAdapter.h
deleted file mode 100644
diff --git a/Engine/Plugins/Experimental/ChaosClothAsset/Source/ChaosClothAsset/Public/ChaosClothAsset/CollectionClothFacade.h b/Engine/Plugins/Experimental/ChaosClothAsset/Source/ChaosClothAsset/Public/ChaosClothAsset/CollectionClothFacade.h
new file mode 100644
--- /dev/null
+++ b/Engine/Plugins/Experimental/ChaosClothAsset/Source/ChaosClothAsset/Public/ChaosClothAsset/CollectionClothFacade.h
@@ -0,0 +1,94 @@
+// Copyright Epic Games, Inc. All Rights Reserved.
+
+#pragma once
+
+#include "ChaosClothAsset/CollectionClothLodFacade.h"
+
+class FClothCollection;
+
+namespace UE::Chaos::ClothAsset
+{
+	/**
+	 * Cloth Asset collection facade class focused on draping and pattern information.
+	 * Const access (read only) version.
+	 */
+	class CHAOSCLOTHASSET_API FCollectionClothConstFacade
+	{
+	public:
+		explicit FCollectionClothConstFacade(const TSharedPtr<const FManagedArrayCollection>& ManagedArrayCollection);
+
+		FCollectionClothConstFacade() = delete;
+
+		FCollectionClothConstFacade(const FCollectionClothConstFacade&) = delete;
+		FCollectionClothConstFacade& operator=(const FCollectionClothConstFacade&) = delete;
+
+		FCollectionClothConstFacade(FCollectionClothConstFacade&&) = default;
+		FCollectionClothConstFacade& operator=(FCollectionClothConstFacade&&) = default;
+
+		virtual ~FCollectionClothConstFacade() = default;
+
+		/** Return whether the facade is defined on the collection. */
+		bool IsValid() const;
+
+		/** Return the specified LOD. */
+		FCollectionClothLodConstFacade GetLod(int32 LodIndex) const;
+
+		/** Return the number of LODs contained in this Cloth. */
+		int32 GetNumLods() const;
+
+		/** Return whether this cloth collection has the specified weight map. */
+		bool HasWeightMap(const FName& Name) const;
+
+		/** Return the name of all user weight maps on this cloth collection. */
+		TArray<FName> GetWeightMapNames() const;
+
+	protected:
+		TSharedPtr<const FClothCollection> ClothCollection;
+	};
+
+	/**
+	 * Cloth Asset collection facade class focused on draping and pattern information.
+	 * Non-const access (read/write) version.
+	 */
+	class CHAOSCLOTHASSET_API FCollectionClothFacade final : public FCollectionClothConstFacade
+	{
+	public:
+		explicit FCollectionClothFacade(const TSharedPtr<FManagedArrayCollection>& ManagedArrayCollection);
+
+		FCollectionClothFacade() = delete;
+
+		FCollectionClothFacade(const FCollectionClothFacade&) = delete;
+		FCollectionClothFacade& operator=(const FCollectionClothFacade&) = delete;
+
+		FCollectionClothFacade(FCollectionClothFacade&&) = default;
+		FCollectionClothFacade& operator=(FCollectionClothFacade&&) = default;
+		virtual ~FCollectionClothFacade() override = default;
+
+		/** Create this facade's groups and attributes. */
+		void DefineSchema();
+
+		/** Remove all LODs from this cloth. */
+		void Reset();
+
+		/** Add a new LOD to this cloth. */
+		int32 AddLod();
+
+		/** Return the specified LOD. */
+		FCollectionClothLodFacade GetLod(int32 LodIndex);
+
+		/** Add a new LOD to this cloth, and return the cloth LOD facade set to its index. */
+		FCollectionClothLodFacade AddGetLod() { return GetLod(AddLod()); }
+
+		/** Set a new number of LODs for this cloth. */
+		void SetNumLod(int32 NumLods);
+
+		/** Add a new weight map to this cloth. Access is then done per pattern. */
+		void AddWeightMap(const FName& Name);
+
+		/** Remove a weight map from this cloth. */
+		void RemoveWeightMap(const FName& Name);
+
+	private:
+		TSharedPtr<FClothCollection> GetClothCollection() { return ConstCastSharedPtr<FClothCollection>(ClothCollection); }
+	};
+}  // End namespace UE::Chaos::ClothAsset
diff --git a/Engine/Plugins/Experimental/ChaosClothAsset/Source/ChaosClothAsset/Public/ChaosClothAsset/CollectionClothLodFacade.h b/Engine/Plugins/Experimental/ChaosClothAsset/Source/ChaosClothAsset/Public/ChaosClothAsset/CollectionClothLodFacade.h
new file mode 100644
--- /dev/null
+++ b/Engine/Plugins/Experimental/ChaosClothAsset/Source/ChaosClothAsset/Public/ChaosClothAsset/CollectionClothLodFacade.h
@@ -0,0 +1,207 @@
+// Copyright Epic Games, Inc. All Rights Reserved.
+
+#pragma once
+
+#include "ChaosClothAsset/CollectionClothPatternFacade.h"
+
+namespace UE::Chaos::ClothAsset
+{
+	class FClothCollection;
+
+	/**
+	 * Cloth Asset collection LOD facade class to access cloth LOD data.
+	 * Constructed from FCollectionClothConstFacade.
+	 * Const access (read only) version.
+	 */
+	class CHAOSCLOTHASSET_API FCollectionClothLodConstFacade
+	{
+	public:
+		FCollectionClothLodConstFacade() = delete;
+
+		FCollectionClothLodConstFacade(const FCollectionClothLodConstFacade&) = delete;
+		FCollectionClothLodConstFacade& operator=(const FCollectionClothLodConstFacade&) = delete;
+
+		FCollectionClothLodConstFacade(FCollectionClothLodConstFacade&&) = default;
+		FCollectionClothLodConstFacade& operator=(FCollectionClothLodConstFacade&&) = default;
+
+		virtual ~FCollectionClothLodConstFacade() = default;
+
+		/** Return the number of materials in this LOD. */
+		int32 GetNumMaterials() const;
+		/** Return the number of tether batches in this LOD. */
+		int32 GetNumTetherBatches() const;
+		/** Return the number of seams in this LOD. */
+		int32 GetNumSeams() const;
+		/** Return the number of patterns in this LOD. */
+		int32 GetNumPatterns() const;
+
+		/** Return a view of all the render materials used on this LOD. */
+		TConstArrayView<FString> GetRenderMaterialPathName() const;
+
+		/** Return a view of all the seam patterns used on this LOD. */
+		TConstArrayView<FIntVector2> GetSeamPatterns() const;
+		/** Return a view of all the seam stitches used on this LOD. */
+		TConstArrayView<TArray<FIntVector2>> GetSeamStitches() const;
+
+		/** Return a pattern facade for the specified pattern index. */
+		FCollectionClothPatternConstFacade GetPattern(int32 PatternIndex) const;
+
+		/** Return the physics asset path names used for this LOD. */
+		const FString& GetPhysicsAssetPathName() const;
+		/** Return the skeleton asset path names used for this LOD. */
+		const FString& GetSkeletonAssetPathName() const;
+
+		//~ Pattern Sim Vertices Group
+		/** Return the total number of simulation vertices for this LOD across all patterns. */
+		int32 GetNumSimVertices() const;
+		TConstArrayView<FVector2f> GetSimPosition() const;
+		TConstArrayView<FVector3f> GetSimRestPosition() const;
+		TConstArrayView<FVector3f> GetSimRestNormal() const;
+		TConstArrayView<int32> GetSimNumBoneInfluences() const;
+		TConstArrayView<TArray<int32>> GetSimBoneIndices() const;
+		TConstArrayView<TArray<float>> GetSimBoneWeights() const;
+
+		//~ Pattern Sim Faces Group
+		/** Return the total number of simulation faces for this LOD across all patterns. */
+		int32 GetNumSimFaces() const;
+		TConstArrayView<FIntVector3> GetSimIndices() const;
+
+		//~ Pattern Render Vertices Group
+		/** Return the total number of render vertices for this LOD across all patterns. */
+		int32 GetNumRenderVertices() const;
+		TConstArrayView<FVector3f> GetRenderPosition() const;
+		TConstArrayView<FVector3f> GetRenderNormal() const;
+		TConstArrayView<FVector3f> GetRenderTangentU() const;
+		TConstArrayView<FVector3f> GetRenderTangentV() const;
+		TConstArrayView<TArray<FVector2f>> GetRenderUVs() const;
+		TConstArrayView<FLinearColor> GetRenderColor() const;
+		TConstArrayView<int32> GetRenderNumBoneInfluences() const;
+		TConstArrayView<TArray<int32>> GetRenderBoneIndices() const;
+		TConstArrayView<TArray<float>> GetRenderBoneWeights() const;
+
+		//~ Pattern Sim Faces Group
+		/** Return the total number of render faces for this LOD across all patterns. */
+		int32 GetNumRenderFaces() const;
+		TConstArrayView<FIntVector3> GetRenderIndices() const;
+		TConstArrayView<int32> GetRenderMaterialIndex() const;
+
+		/**
+		 * Return a view on the specified vertex weightmap.
+		 * The attribute must exist otherwise the function will assert.
+		 * @see FCollectionClothFacade::AddWeightMap, GetNumSimVertices.
+		 */
+		TConstArrayView<float> GetWeightMap(const FName& Name) const;
+
+		/** Return the LOD index this facade has been created with. */
+		int32 GetLodIndex() const { return LodIndex; }
+
+		/** Return the welded simulation mesh for this LOD. */
+		void BuildSimulationMesh(TArray<FVector3f>& Positions, TArray<FVector3f>& Normals, TArray<uint32>& Indices, TArray<int32>& WeldingMap) const;
+
+	protected:
+		friend class FCollectionClothConstFacade;
+		FCollectionClothLodConstFacade(const TSharedPtr<const FClothCollection>& InClothCollection, int32 InLodIndex);
+
+		TSharedPtr<const FClothCollection> ClothCollection;
+		int32 LodIndex;
+	};
+
+	/**
+	 * Cloth Asset collection LOD facade class to access cloth LOD data.
+	 * Constructed from FCollectionClothFacade.
+	 * Non-const access (read/write) version.
+	 */
+	class CHAOSCLOTHASSET_API FCollectionClothLodFacade final : public FCollectionClothLodConstFacade
+	{
+	public:
+		FCollectionClothLodFacade() = delete;
+
+		FCollectionClothLodFacade(const FCollectionClothLodFacade&) = delete;
+		FCollectionClothLodFacade& operator=(const FCollectionClothLodFacade&) = delete;
+
+		FCollectionClothLodFacade(FCollectionClothLodFacade&&) = default;
+		FCollectionClothLodFacade& operator=(FCollectionClothLodFacade&&) = default;
+
+		virtual ~FCollectionClothLodFacade() override = default;
+
+		/** Remove all cloth patterns from this cloth LOD. */
+		void Reset();
+
+		/** Initialize the cloth LOD using the specified 3D triangle mesh and unwrapping it into a 2D geometry for generating the patterns. */
+		void Initialize(const TArray<FVector3f>& Positions, const TArray<uint32>& Indices);
+
+		/** Initialize the cloth LOD using another cloth collection. */
+		void Initialize(const FCollectionClothLodConstFacade& Other);
+
+		/** Set the new number of materials to this cloth LOD. */
+		void SetNumMaterials(int32 NumMaterials);
+		/** Set the new number of tether batches to this cloth LOD. */
+		void SetNumTetherBatches(int32 NumSeams);
+		/** Set the new number of seams to this cloth LOD. */
+		void SetNumSeams(int32 NumSeams);
+		/** Set the new number of patterns to this cloth LOD. */
+		void SetNumPatterns(int32 NumPatterns);
+
+		/** Return a view to the render material path names for this cloth LOD. */
+		TArrayView<FString> GetRenderMaterialPathName();
+		/** Return a view to the seam patterns for this cloth LOD. */
+		TArrayView<FIntVector2> GetSeamPatterns();
+		/** Return a view to the seam stitches for this cloth LOD. */
+		TArrayView<TArray<FIntVector2>> GetSeamStitches();
+
+		/** Add a new pattern to this cloth LOD and return its index in the LOD pattern list. */
+		int32 AddPattern();
+
+		/** Return a pattern facade for the specified pattern index. */
+		FCollectionClothPatternFacade GetPattern(int32 PatternIndex);
+
+		/** Add a new pattern to this cloth LOD, and return the cloth pattern facade set to its index. */
+		FCollectionClothPatternFacade AddGetPattern() { return GetPattern(AddPattern()); }
+
+		/** Set the physics asset path names used for this LOD. */
+		void SetPhysicsAssetPathName(const FString& PathName);
+		/** Set the skeleton asset path names used for this LOD. */
+		void SetSkeletonAssetPathName(const FString& PathName);
+
+		//~ Pattern Sim Vertices Group
+		TArrayView<FVector2f> GetSimPosition();
+		TArrayView<FVector3f> GetSimRestPosition();
+		TArrayView<FVector3f> GetSimRestNormal();
+		TArrayView<int32> GetSimNumBoneInfluences();
+		TArrayView<TArray<int32>> GetSimBoneIndices();
+		TArrayView<TArray<float>> GetSimBoneWeights();
+
+		//~ Pattern Sim Faces Group
+		TArrayView<FIntVector3> GetSimIndices();
+
+		//~ Pattern Render Vertices Group
+		TArrayView<FVector3f> GetRenderPosition();
+		TArrayView<FVector3f> GetRenderNormal();
+		TArrayView<FVector3f> GetRenderTangentU();
+		TArrayView<FVector3f> GetRenderTangentV();
+		TArrayView<TArray<FVector2f>> GetRenderUVs();
+		TArrayView<FLinearColor> GetRenderColor();
+		TArrayView<int32> GetRenderNumBoneInfluences();
+		TArrayView<TArray<int32>> GetRenderBoneIndices();
+		TArrayView<TArray<float>> GetRenderBoneWeights();
+
+		//~ Pattern Sim Faces Group
+		TArrayView<FIntVector3> GetRenderIndices();
+		TArrayView<int32> GetRenderMaterialIndex();
+
+		/**
+		 * Return a view on the specified vertex weightmap.
+		 * The attribute must exist otherwise the function will assert.
+		 * @see FCollectionClothFacade::AddWeightMap, GetNumSimVertices.
+		 */
+		TArrayView<float> GetWeightMap(const FName& Name);
+
+	private:
+		friend class FCollectionClothFacade;
+		FCollectionClothLodFacade(const TSharedPtr<FClothCollection>& InClothCollection, int32 InLodIndex);
+
+		void SetDefaults();
+
+		TSharedPtr<FClothCollection> GetClothCollection() { return ConstCastSharedPtr<FClothCollection>(ClothCollection); }
+	};
+}  // End namespace UE::Chaos::ClothAsset
diff --git a/Engine/Plugins/Experimental/ChaosClothAsset/Source/ChaosClothAsset/Public/ChaosClothAsset/CollectionClothPatternFacade.h b/Engine/Plugins/Experimental/ChaosClothAsset/Source/ChaosClothAsset/Public/ChaosClothAsset/CollectionClothPatternFacade.h
new file mode 100644
--- /dev/null
+++ b/Engine/Plugins/Experimental/ChaosClothAsset/Source/ChaosClothAsset/Public/ChaosClothAsset/CollectionClothPatternFacade.h
@@ -0,0 +1,193 @@
+// Copyright Epic Games, Inc. All Rights Reserved.
+
+#pragma once
+
+#include "Templates/SharedPointer.h"
+
+struct FManagedArrayCollection;
+
+namespace UE::Chaos::ClothAsset
+{
+	class FClothCollection;
+
+	/**
+	 * Cloth Asset collection pattern facade class to access cloth pattern data.
+	 * Constructed from FCollectionClothLodConstFacade.
+	 * Const access (read only) version.
+	 */
+	class CHAOSCLOTHASSET_API FCollectionClothPatternConstFacade
+	{
+	public:
+		FCollectionClothPatternConstFacade() = delete;
+
+		FCollectionClothPatternConstFacade(const FCollectionClothPatternConstFacade&) = delete;
+		FCollectionClothPatternConstFacade& operator=(const FCollectionClothPatternConstFacade&) = delete;
+
+		FCollectionClothPatternConstFacade(FCollectionClothPatternConstFacade&&) = default;
+		FCollectionClothPatternConstFacade& operator=(FCollectionClothPatternConstFacade&&) = default;
+
+		virtual ~FCollectionClothPatternConstFacade() = default;
+
+		/** Return the pattern status flag. */
+		int32 GetStatusFlags() const;
+
+		/** Return the total number of simulation vertices for this pattern. */
+		int32 GetNumSimVertices() const;
+		/** Return the total number of simulation faces for this pattern. */
+		int32 GetNumSimFaces() const;
+		/** Return the total number of render vertices for this pattern. */
+		int32 GetNumRenderVertices() const;
+		/** Return the total number of render faces for this pattern. */
+		int32 GetNumRenderFaces() const;
+
+		/** Return the simulation vertices offset for this pattern in the simulation vertices for the LOD. */
+		int32 GetSimVerticesOffset() const;
+		/** Return the simulation faces offset for this pattern in the simulation faces for the LOD. */
+		int32 GetSimFacesOffset() const;
+		/** Return the render vertices offset for this pattern in the render vertices for the LOD. */
+		int32 GetRenderVerticesOffset() const;
+		/** Return the render faces offset for this pattern in the render faces for the LOD. */
+		int32 GetRenderFacesOffset() const;
+
+		//~ Sim Vertices Group
+		// Note: Use the FCollectionClothLodConstFacade accessors instead of these for the array indices to match the SimIndices values
+		TConstArrayView<FVector2f> GetSimPosition() const;
+		TConstArrayView<FVector3f> GetSimRestPosition() const;
+		TConstArrayView<FVector3f> GetSimRestNormal() const;
+		TConstArrayView<int32> GetSimNumBoneInfluences() const;
+		TConstArrayView<TArray<int32>> GetSimBoneIndices() const;
+		TConstArrayView<TArray<float>> GetSimBoneWeights() const;
+
+		//~ Sim Faces Group, note: SimIndices points to the LOD arrays, not the pattern arrays
+		TConstArrayView<FIntVector3> GetSimIndices() const;
+
+		//~ Render Vertices Group
+		// Note: Use the FCollectionClothLodConstFacade accessors instead of these for the array indices to match the RenderIndices values
+		TConstArrayView<FVector3f> GetRenderPosition() const;
+		TConstArrayView<FVector3f> GetRenderNormal() const;
+		TConstArrayView<FVector3f> GetRenderTangentU() const;
+		TConstArrayView<FVector3f> GetRenderTangentV() const;
+		TConstArrayView<TArray<FVector2f>> GetRenderUVs() const;
+		TConstArrayView<FLinearColor> GetRenderColor() const;
+		TConstArrayView<int32> GetRenderNumBoneInfluences() const;
+		TConstArrayView<TArray<int32>> GetRenderBoneIndices() const;
+		TConstArrayView<TArray<float>> GetRenderBoneWeights() const;
+
+		//~ Render Faces Group, note: RenderIndices points to the LOD arrays, not the pattern arrays
+		TConstArrayView<FIntVector3> GetRenderIndices() const;
+		TConstArrayView<int32> GetRenderMaterialIndex() const;
+
+		/** Return the LOD index this facade has been created with. */
+		int32 GetLodIndex() const { return LodIndex; }
+
+		/** Return the Pattern index this facade has been created with. */
+		int32 GetPatternIndex() const { return PatternIndex; }
+
+		/**
+		 * Return a const view on the specified vertex weightmap.
+		 * The attribute must exist otherwise the function will assert.
+		 * @see FCollectionClothFacade::AddWeightMap, GetNumSimVertices.
+		 */
+		TConstArrayView<float> GetWeightMap(const FName& Name) const;
+
+	protected:
+		friend class FCollectionClothPatternFacade;  // For other instances access
+		friend class FCollectionClothLodConstFacade;
+		FCollectionClothPatternConstFacade(const TSharedPtr<const FClothCollection>& InClothCollection, int32 InLodIndex, int32 InPatternIndex);
+
+		int32 GetBaseElementIndex() const;
+		int32 GetElementIndex() const { return GetBaseElementIndex() + PatternIndex; }
+
+		TSharedPtr<const FClothCollection> ClothCollection;
+		int32 LodIndex;
+		int32 PatternIndex;
+	};
+
+	/**
+	 * Cloth Asset collection pattern facade class to access cloth pattern data.
+	 * Constructed from FCollectionClothLodFacade.
+	 * Non-const access (read/write) version.
+	 */
+	class CHAOSCLOTHASSET_API FCollectionClothPatternFacade final : public FCollectionClothPatternConstFacade
+	{
+	public:
+		FCollectionClothPatternFacade() = delete;
+
+		FCollectionClothPatternFacade(const FCollectionClothPatternFacade&) = delete;
+		FCollectionClothPatternFacade& operator=(const FCollectionClothPatternFacade&) = delete;
+
+		FCollectionClothPatternFacade(FCollectionClothPatternFacade&&) = default;
+		FCollectionClothPatternFacade& operator=(FCollectionClothPatternFacade&&) = default;
+
+		virtual ~FCollectionClothPatternFacade() override = default;
+
+		/** Remove all geometry from this cloth pattern. */
+		void Reset();
+
+		/** Initialize the cloth pattern using the specified 3D and 2D positions, and topology. */
+		void Initialize(const TArray<FVector2f>& Positions, const TArray<FVector3f>& RestPositions, const TArray<uint32>& Indices);
+
+		/** Initialize the cloth pattern using another pattern. */
+		void Initialize(const FCollectionClothPatternConstFacade& Other);
+
+		/** Grow or shrink the space reserved for simulation vertices for this pattern within the cloth collection and return its start index. */
+		void SetNumSimVertices(int32 NumSimVertices);
+
+		/** Grow or shrink the space reserved for simulation faces for this pattern within the cloth collection and return its start index. */
+		void SetNumSimFaces(int32 NumSimFaces);
+
+		/** Grow or shrink the space reserved for render vertices for this pattern within the cloth collection and return its start index. */
+		void SetNumRenderVertices(int32 NumRenderVertices);
+
+		/** Grow or shrink the space reserved for render faces for this pattern within the cloth collection and return its start index. */
+		void SetNumRenderFaces(int32 NumRenderFaces);
+
+		/** Set the pattern status flag. */
+		void SetStatusFlags(int32 StatusFlags);
+
+		//~ Sim Vertices Group
+		// Note: Use the FCollectionClothLodFacade accessors instead of these for the array indices to match the SimIndices values
+		TArrayView<FVector2f> GetSimPosition();
+		TArrayView<FVector3f> GetSimRestPosition();
+		TArrayView<FVector3f> GetSimRestNormal();
+		TArrayView<int32> GetSimNumBoneInfluences();
+		TArrayView<TArray<int32>> GetSimBoneIndices();
+		TArrayView<TArray<float>> GetSimBoneWeights();
+
+		//~ Sim Faces Group
+		// Note: SimIndices points to the LOD arrays, not the pattern arrays
+		TArrayView<FIntVector3> GetSimIndices();
+
+		//~ Render Vertices Group
+		// Note: Use the FCollectionClothLodFacade accessors instead of these for the array indices to match the RenderIndices values
+		TArrayView<FVector3f> GetRenderPosition();
+		TArrayView<FVector3f> GetRenderNormal();
+		TArrayView<FVector3f> GetRenderTangentU();
+		TArrayView<FVector3f> GetRenderTangentV();
+		TArrayView<TArray<FVector2f>> GetRenderUVs();
+		TArrayView<FLinearColor> GetRenderColor();
+		TArrayView<int32> GetRenderNumBoneInfluences();
+		TArrayView<TArray<int32>> GetRenderBoneIndices();
+		TArrayView<TArray<float>> GetRenderBoneWeights();
+
+		//~ Render Faces Group
+		// Note: RenderIndices points to the LOD arrays, not the pattern arrays
+		TArrayView<FIntVector3> GetRenderIndices();
+		TArrayView<int32> GetRenderMaterialIndex();
+
+		/**
+		 * Return a view on the specified vertex weightmap.
+		 * The attribute must exist otherwise the function will assert.
+		 * @see FCollectionClothFacade::AddWeightMap, GetNumSimVertices.
+		 */
+		TArrayView<float> GetWeightMap(const FName& Name);
+
+	protected:
+		friend class FCollectionClothLodFacade;
+		FCollectionClothPatternFacade(const TSharedPtr<FClothCollection>& InClothCollection, int32 InLodIndex, int32 InPatternIndex);
+
+		void SetDefaults();
+
+		TSharedPtr<FClothCollection> GetClothCollection() { return ConstCastSharedPtr<FClothCollection>(ClothCollection); }
+	};
+}  // End namespace UE::Chaos::ClothAsset
diff --git a/Engine/Plugins/Experimental/ChaosClothAsset/Source/ChaosClothAssetEngine/Private/ChaosClothAsset/ClothAsset.cpp b/Engine/Plugins/Experimental/ChaosClothAsset/Source/ChaosClothAssetEngine/Private/ChaosClothAsset/ClothAsset.cpp
--- a/Engine/Plugins/Experimental/ChaosClothAsset/Source/ChaosClothAssetEngine/Private/ChaosClothAsset/ClothAsset.cpp
+++ b/Engine/Plugins/Experimental/ChaosClothAsset/Source/ChaosClothAssetEngine/Private/ChaosClothAsset/ClothAsset.cpp
@@ -2,7 +2,7 @@
 
 #include "ChaosClothAsset/ClothAsset.h"
 #include "ChaosClothAsset/ClothAssetBuilder.h"
-#include "ChaosClothAsset/ClothAdapter.h"
+#include "ChaosClothAsset/CollectionClothFacade.h"
 #include "ChaosClothAsset/ClothComponent.h"
 #include "ChaosClothAsset/ClothGeometryTools.h"
 #include "ChaosClothAsset/ClothAssetPrivate.h"
@@ -11,10 +11,12 @@
 #include "Engine/RendererSettings.h"
 #include "Engine/SkinnedAssetAsyncCompileUtils.h"
 #include "Features/IModularFeatures.h"
+#include "GeometryCollection/ManagedArrayCollection.h"
+#include "PhysicsEngine/PhysicsAsset.h"
 #include "Rendering/SkeletalMeshModel.h"
-#include "EngineUtils.h"
 #include "Rendering/SkeletalMeshRenderData.h"
 #include "UObject/Package.h"
+#include "EngineUtils.h"
 #if WITH_EDITOR
 #include "IMeshBuilderModule.h"
 #include "DerivedDataCacheInterface.h"
@@ -29,11 +31,19 @@
 UChaosClothAsset::UChaosClothAsset(const FObjectInitializer& ObjectInitializer)
 	: Super(ObjectInitializer)
 	, DisableBelowMinLodStripping(FPerPlatformBool(false))
-	, ClothCollection(MakeShared<UE::Chaos::ClothAsset::FClothCollection>())
+	, ClothCollection(MakeShared<FManagedArrayCollection>())
 #if WITH_EDITORONLY_DATA
 	, MeshModel(MakeShareable(new FSkeletalMeshModel()))
 #endif
 {
+	// Setup the collection to be a cloth collection and have at least one LOD
+	UE::Chaos::ClothAsset::FCollectionClothFacade ClothFacade(ClothCollection);
+	ClothFacade.DefineSchema();
+	ClothFacade.AddLod();
+
+	// Set default skeleton (must be done after having added the LOD)
+	constexpr bool bRebuildModels = false;
+	SetSkeleton(nullptr, bRebuildModels);
 }
 
 UChaosClothAsset::UChaosClothAsset(FVTableHelper& Helper)
@@ -77,7 +87,8 @@ void UChaosClothAsset::Serialize(FArchive& Ar)
 	bool bCooked = Ar.IsCooking();
 	Ar << bCooked;
 
-	ClothCollection->Serialize(Ar);
+	Chaos::FChaosArchive ChaosArchive(Ar);
+	ClothCollection->Serialize(ChaosArchive);
 
 	Ar << GetRefSkeleton();
 	if (Ar.IsLoading())
@@ -117,7 +128,7 @@ void UChaosClothAsset::BeginPostLoadInternal(FSkinnedAssetPostLoadContext& Conte
 #if WITH_EDITOR
 	TRACE_CPUPROFILER_EVENT_SCOPE(UChaosClothAsset::BeginPostLoadInternal);
 
-	checkf(IsInGameThread(), TEXT("Cannot execute function UChaosClothAsset::BeginPostLoadInternal asynchronously. Asset: %s"), *this->GetFullName());
+	checkf(IsInGameThread(), TEXT("Cannot execute function UChaosClothAsset::BeginPostLoadInternal asynchronously. Asset: %s"), *GetFullName());
 	SetInternalFlags(EInternalObjectFlags::Async);
 
 	// Lock all properties that should not be modified/accessed during async post-load
@@ -126,6 +137,23 @@ void UChaosClothAsset::BeginPostLoadInternal(FSkinnedAssetPostLoadContext& Conte
 	// This scope allows us to use any locked properties without causing stalls
 	FSkinnedAssetAsyncBuildScope AsyncBuildScope(this);
 
+	// Make sure that the collection is still compatible and valid
+	UE::Chaos::ClothAsset::FCollectionClothFacade ClothFacade(ClothCollection);
+	if (!ClothFacade.IsValid())
+	{
+		UE_LOG(LogChaosClothAsset, Warning, TEXT("Invalid Cloth Collection found while loading Cloth Asset %s."), *GetFullName());
+		ClothCollection = MakeShared<FManagedArrayCollection>();
+		ClothFacade = UE::Chaos::ClothAsset::FCollectionClothFacade(ClothCollection);
+		ClothFacade.DefineSchema();
+	}
+	if (!ClothFacade.GetNumLods())
+	{
+		ClothFacade.AddLod();
+		constexpr bool bRebuildModels = false;
+		SetSkeleton(Skeleton, bRebuildModels);  // Re-update the collection with the skeleton information if any
+		SetPhysicsAsset(PhysicsAsset);  // Re-update the collection with the physics asset information if any
+	}
+
 	BuildClothSimulationModel();  // TODO: Cache ClothSimulationModel?
 
 	BuildMeshModel();
@@ -285,26 +313,16 @@ void UChaosClothAsset::CalculateBounds()
 
 	FBox BoundingBox(ForceInit);
 
-	const FClothConstAdapter Cloth(ClothCollection);
+	const FCollectionClothConstFacade Cloth(ClothCollection);
 
 	for (int32 LodIndex = 0; LodIndex < Cloth.GetNumLods(); ++LodIndex)
 	{
-		const FClothLodConstAdapter ClothLod = Cloth.GetLod(LodIndex);
+		const FCollectionClothLodConstFacade ClothLod = Cloth.GetLod(LodIndex);
+		const TConstArrayView<FVector3f> RenderPositionArray = ClothLod.GetRenderPosition();
 
-		for (int32 ClothPatternIndex = 0; ClothPatternIndex < ClothLod.GetNumPatterns(); ++ClothPatternIndex)
+		for (const FVector3f& RenderPosition : RenderPositionArray)
 		{
-			const FClothPatternConstAdapter ClothPattern = ClothLod.GetPattern(ClothPatternIndex);
-			const int32 PatternElementIndex = ClothPattern.GetElementIndex();
-
-			const int32 RenderVerticesStart = ClothCollection->RenderVerticesStart[PatternElementIndex];
-			const int32 RenderVerticesEnd = ClothCollection->RenderVerticesEnd[PatternElementIndex];
-			if (RenderVerticesStart != INDEX_NONE && RenderVerticesEnd != INDEX_NONE)
-			{
-				for (int32 RenderVertexIndex = RenderVerticesStart; RenderVertexIndex <= RenderVerticesEnd; ++RenderVertexIndex)
-				{
-					BoundingBox += (FVector)ClothCollection->RenderPosition[RenderVertexIndex];
-				}
-			}
+			BoundingBox += (FVector)RenderPosition;
 		}
 	}
 
@@ -338,15 +356,14 @@ void UChaosClothAsset::Build()
 	// Set a new Guid to invalidate the DDC
 	AssetGuid = FGuid::NewGuid();
 
-	// Rebuild Skeleton
-	constexpr bool bRebuildClothSimulationModel = false;
-	UpdateSkeleton(bRebuildClothSimulationModel);
+	// Rebuild matrices
+	CalculateInvRefMatrices();
 
 	// Update bounds
 	CalculateBounds();
 
 	// Add LODs to the render data
-	const FClothConstAdapter Cloth(ClothCollection);
+	const FCollectionClothConstFacade Cloth(ClothCollection);
 	const int32 NumLods = FMath::Max(Cloth.GetNumLods(), 1);  // The render data will always look for at least one default LOD 0
 
 	// Rebuild LOD Infos
@@ -428,7 +445,7 @@ void UChaosClothAsset::BuildMeshModel()
 		{
 			using namespace UE::Chaos::ClothAsset;
 
-			const FClothConstAdapter Cloth(GetClothCollection());
+			const FCollectionClothConstFacade Cloth(GetClothCollection());
 			const int32 NumLods = Cloth.GetNumLods();
 
 			check(MeshModel);  // MeshModel should always be created in the Cloth Asset constructor WITH_EDITORONLY_DATA
@@ -503,6 +520,36 @@ void UChaosClothAsset::CacheDerivedData(FSkinnedAssetCompilationContext* Context
 
 }
 
+void UChaosClothAsset::SetPhysicsAsset(UPhysicsAsset* InPhysicsAsset)
+{
+	using namespace UE::Chaos::ClothAsset;
+
+	PhysicsAsset = InPhysicsAsset;
+	FClothGeometryTools::SetPhysicsAssetPathName(ClothCollection, PhysicsAsset ? PhysicsAsset->GetPathName() : FString());
+}
+
+void UChaosClothAsset::SetSkeleton(USkeleton* InSkeleton, bool bRebuildModels)
+{
+	using namespace UE::Chaos::ClothAsset;
+	check(ClothCollection.IsValid());
+
+	const TCHAR* const DefaultSkeletonPathName = TEXT("/Engine/EditorMeshes/SkeletalMesh/DefaultSkeletalMesh_Skeleton.DefaultSkeletalMesh_Skeleton");
+	Skeleton = InSkeleton ?
+		InSkeleton :
+		LoadObject<USkeleton>(nullptr, DefaultSkeletonPathName, nullptr, LOAD_None, nullptr);
+
+	RefSkeleton = Skeleton->GetReferenceSkeleton();
+
+	constexpr bool bBindSimMesh = true;
+	constexpr bool bBindRenderMesh = true;
+	FClothGeometryTools::BindMeshToRootBone(ClothCollection, bBindSimMesh, bBindRenderMesh);
+	FClothGeometryTools::SetSkeletonAssetPathName(ClothCollection, Skeleton->GetPathName());
+	if (bRebuildModels)
+	{
+		Build();
+	}
+}
+
 FString UChaosClothAsset::BuildDerivedDataKey(const ITargetPlatform* TargetPlatform)
 {
 	FString KeySuffix(TEXT(""));
@@ -578,12 +625,27 @@ bool UChaosClothAsset::IsInitialBuildDone() const
 }
 #endif // WITH_EDITOR
 
-void UChaosClothAsset::CopySimMeshToRenderMesh(int32 MaterialIndex)
+void UChaosClothAsset::CopySimMeshToRenderMesh(UMaterialInterface* Material)
 {
 	using namespace UE::Chaos::ClothAsset;
-
 	check(ClothCollection.IsValid());
-	FClothGeometryTools::CopySimMeshToRenderMesh(ClothCollection, MaterialIndex);
+
+	// Add a default material if none is specified
+	const FString RenderMaterialPathName = Material ?
+		Material->GetPathName() :
+		FString(TEXT("/Engine/EditorMaterials/Cloth/CameraLitDoubleSided.CameraLitDoubleSided"));
+
+	FClothGeometryTools::CopySimMeshToRenderMesh(ClothCollection, RenderMaterialPathName);
+
+	// Set new material
+	Materials.Reset(1);
+	if (FClothGeometryTools::HasRenderMesh(ClothCollection))
+	{
+		if (UMaterialInterface* const LoadedMaterial = LoadObject<UMaterialInterface>(nullptr, *RenderMaterialPathName, nullptr, LOAD_None, nullptr))
+		{
+			Materials.Emplace(LoadedMaterial, true, false, LoadedMaterial->GetFName());
+		}
+	}
 }
 
 void UChaosClothAsset::ReregisterComponents()
diff --git a/Engine/Plugins/Experimental/ChaosClothAsset/Source/ChaosClothAssetEngine/Private/ChaosClothAsset/ClothComponent.cpp b/Engine/Plugins/Experimental/ChaosClothAsset/Source/ChaosClothAssetEngine/Private/ChaosClothAsset/ClothComponent.cpp
--- a/Engine/Plugins/Experimental/ChaosClothAsset/Source/ChaosClothAssetEngine/Private/ChaosClothAsset/ClothComponent.cpp
+++ b/Engine/Plugins/Experimental/ChaosClothAsset/Source/ChaosClothAssetEngine/Private/ChaosClothAsset/ClothComponent.cpp
@@ -3,10 +3,10 @@
 #include "ChaosClothAsset/ClothComponent.h"
 #include "ChaosClothAsset/ClothAsset.h"
 #include "ChaosClothAsset/ClothAssetPrivate.h"
-#include "ChaosClothAsset/ClothCollection.h"
 #include "ChaosClothAsset/ClothSimulationModel.h"
 #include "ChaosClothAsset/ClothSimulationProxy.h"
 #include "Chaos/CollectionPropertyFacade.h"
+#include "GeometryCollection/ManagedArrayCollection.h"
 #include "HAL/IConsoleManager.h"
 #include "Rendering/SkeletalMeshRenderData.h"
 #include "SkeletalRenderPublic.h"
diff --git a/Engine/Plugins/Experimental/ChaosClothAsset/Source/ChaosClothAssetEngine/Private/ChaosClothAsset/ClothSimulationModel.cpp b/Engine/Plugins/Experimental/ChaosClothAsset/Source/ChaosClothAssetEngine/Private/ChaosClothAsset/ClothSimulationModel.cpp
--- a/Engine/Plugins/Experimental/ChaosClothAsset/Source/ChaosClothAssetEngine/Private/ChaosClothAsset/ClothSimulationModel.cpp
+++ b/Engine/Plugins/Experimental/ChaosClothAsset/Source/ChaosClothAssetEngine/Private/ChaosClothAsset/ClothSimulationModel.cpp
@@ -1,8 +1,7 @@
 // Copyright Epic Games, Inc. All Rights Reserved.
 
 #include "ChaosClothAsset/ClothSimulationModel.h"
-#include "ChaosClothAsset/ClothCollection.h"
-#include "ChaosClothAsset/ClothAdapter.h"
+#include "ChaosClothAsset/CollectionClothFacade.h"
 #include "ReferenceSkeleton.h"
 
 namespace UE::Chaos::ClothAsset
@@ -106,11 +105,11 @@ namespace UE::Chaos::ClothAsset
 	}
 }  // End namespace UE::Chaos::ClothAsset
 
-FChaosClothSimulationModel::FChaosClothSimulationModel(const TSharedPtr<const UE::Chaos::ClothAsset::FClothCollection>& ClothCollection, const FReferenceSkeleton& ReferenceSkeleton)
+FChaosClothSimulationModel::FChaosClothSimulationModel(const TSharedPtr<const FManagedArrayCollection>& ClothCollection, const FReferenceSkeleton& ReferenceSkeleton)
 {
 	using namespace UE::Chaos::ClothAsset;
 
-	const FClothConstAdapter Cloth(ClothCollection);
+	const FCollectionClothConstFacade Cloth(ClothCollection);
 
 	// Initialize LOD models
 	const int32 NumLods = Cloth.GetNumLods();
@@ -119,27 +118,29 @@ FChaosClothSimulationModel::FChaosClothSimulationModel(const TSharedPtr<const UE
 	for (int32 LodIndex = 0; LodIndex < NumLods; ++LodIndex)
 	{
 		FChaosClothSimulationLodModel& LodModel = ClothSimulationLodModels[LodIndex];
-		const FClothLodConstAdapter ClothLod = Cloth.GetLod(LodIndex);
+		const FCollectionClothLodConstFacade ClothLod = Cloth.GetLod(LodIndex);
 		TArray<int32> WeldingMap;
 		ClothLod.BuildSimulationMesh(LodModel.Positions, LodModel.Normals, LodModel.Indices, WeldingMap);
 
-		const TManagedArray<float>* const MaxDistanceValues = ClothCollection->FindAttributeTyped<float>("MaxDistance", FClothCollection::SimVerticesGroup);
-		if (MaxDistanceValues)
+		static const FName MaxDistanceName = TEXT("MaxDistance");
+		const TConstArrayView<float> MaxDistanceValues = ClothLod.GetWeightMap(MaxDistanceName);
+
+		if (MaxDistanceValues.Num())
 		{
-			LodModel.MaxDistance = MaxDistanceValues->GetConstArray();
+			LodModel.MaxDistance = MaxDistanceValues;
 		}
 		else
 		{
 			LodModel.MaxDistance.Init(1.0, LodModel.Positions.Num());
 		}
 
-		TConstArrayView<int32> NumBoneInfluences = ClothLod.GetPatternsSimNumBoneInfluences();
-		TConstArrayView<TArray<int32>> SimBoneIndices = ClothLod.GetPatternsSimBoneIndices();
-		TConstArrayView<TArray<float>> SimBoneWeights = ClothLod.GetPatternsSimBoneWeights();
+		TConstArrayView<int32> NumBoneInfluences = ClothLod.GetSimNumBoneInfluences();
+		TConstArrayView<TArray<int32>> SimBoneIndices = ClothLod.GetSimBoneIndices();
+		TConstArrayView<TArray<float>> SimBoneWeights = ClothLod.GetSimBoneWeights();
 		LodModel.BoneData.SetNum(LodModel.Positions.Num());
 
 		uint32 WeldedIndex = 0;
-		for (int32 VertexIndex = 0; VertexIndex <  ClothLod.GetPatternsNumSimVertices(); ++VertexIndex)
+		for (int32 VertexIndex = 0; VertexIndex <  ClothLod.GetNumSimVertices(); ++VertexIndex)
 		{
 			if (WeldingMap[VertexIndex] == VertexIndex)
 			{
diff --git a/Engine/Plugins/Experimental/ChaosClothAsset/Source/ChaosClothAssetEngine/Private/ChaosClothAsset/ClothSimulationModel.h b/Engine/Plugins/Experimental/ChaosClothAsset/Source/ChaosClothAssetEngine/Private/ChaosClothAsset/ClothSimulationModel.h
--- a/Engine/Plugins/Experimental/ChaosClothAsset/Source/ChaosClothAssetEngine/Private/ChaosClothAsset/ClothSimulationModel.h
+++ b/Engine/Plugins/Experimental/ChaosClothAsset/Source/ChaosClothAssetEngine/Private/ChaosClothAsset/ClothSimulationModel.h
@@ -10,10 +10,7 @@
 #include "ClothSimulationModel.generated.h"
 
 struct FReferenceSkeleton;
-namespace UE::Chaos::ClothAsset
-{
-	class FClothCollection;
-}
+struct FManagedArrayCollection;
 
 /**
  * Cloth simulation LOD model.
@@ -70,7 +67,7 @@ struct FChaosClothSimulationModel
 	int32 ReferenceBoneIndex = INDEX_NONE;
 
 	FChaosClothSimulationModel() = default;
-	FChaosClothSimulationModel(const TSharedPtr<const UE::Chaos::ClothAsset::FClothCollection>& ClothCollection, const FReferenceSkeleton& ReferenceSkeleton);
+	FChaosClothSimulationModel(const TSharedPtr<const FManagedArrayCollection>& ClothCollection, const FReferenceSkeleton& ReferenceSkeleton);
 
 	int32 GetNumLods() const { return ClothSimulationLodModels.Num(); }
 
diff --git a/Engine/Plugins/Experimental/ChaosClothAsset/Source/ChaosClothAssetEngine/Public/ChaosClothAsset/ClothAsset.h b/Engine/Plugins/Experimental/ChaosClothAsset/Source/ChaosClothAssetEngine/Public/ChaosClothAsset/ClothAsset.h
--- a/Engine/Plugins/Experimental/ChaosClothAsset/Source/ChaosClothAssetEngine/Public/ChaosClothAsset/ClothAsset.h
+++ b/Engine/Plugins/Experimental/ChaosClothAsset/Source/ChaosClothAssetEngine/Public/ChaosClothAsset/ClothAsset.h
@@ -7,13 +7,13 @@
 #include "RenderCommandFence.h"
 #include "ClothAsset.generated.h"
 
-namespace UE::Chaos::ClothAsset { class FClothCollection; }
 class FSkeletalMeshRenderData;
 class FSkeletalMeshModel;
 struct FChaosClothSimulationModel;
 struct FSkeletalMeshLODInfo;
-class UDataflow;
+struct FManagedArrayCollection;
 class FSkinnedAssetCompilationContext;
+class UDataflow;
 
 UENUM()
 enum class EClothAssetAsyncProperties : uint64
@@ -74,9 +74,9 @@ public:
 	virtual TArray<class USkeletalMeshSocket*> GetActiveSocketList() const override { static TArray<class USkeletalMeshSocket*> Dummy; return Dummy; }
 	virtual USkeletalMeshSocket* FindSocket(FName InSocketName) const override	{ return nullptr; }
 	virtual USkeletalMeshSocket* FindSocketInfo(FName InSocketName, FTransform& OutTransform, int32& OutBoneIndex, int32& OutIndex) const override { return nullptr; }
-	virtual USkeleton* GetSkeleton() override									{ return nullptr; }
-	virtual const USkeleton* GetSkeleton() const override						{ return nullptr; }
-	virtual void SetSkeleton(USkeleton* InSkeleton) override					{}
+	virtual USkeleton* GetSkeleton() override									{ return Skeleton; }
+	virtual const USkeleton* GetSkeleton() const override						{ return Skeleton; }
+	virtual void SetSkeleton(USkeleton* InSkeleton) override					{ constexpr bool bRebuildModels = true; SetSkeleton(InSkeleton, bRebuildModels); }
 	virtual UMeshDeformer* GetDefaultMeshDeformer() const override				{ return nullptr; }
 	virtual bool IsValidLODIndex(int32 Index) const override					{ return LODInfo.IsValidIndex(Index); }
 	virtual int32 GetMinLodIdx(bool bForceLowestLODIdx = false) const override	{ return 0; }
@@ -98,10 +98,10 @@ public:
 	//~ End USkinnedAsset interface
 
 	/** Return the enclosed Cloth Collection object. */
-	TSharedPtr<UE::Chaos::ClothAsset::FClothCollection> GetClothCollection() { return ClothCollection; }
+	TSharedPtr<FManagedArrayCollection>& GetClothCollection() { return ClothCollection; }
 
 	/** Return the enclosed Cloth Collection object, const version. */
-	TSharedPtr<const UE::Chaos::ClothAsset::FClothCollection> GetClothCollection() const { return ClothCollection; }
+	TSharedPtr<const FManagedArrayCollection> GetClothCollection() const { return ClothCollection; }
 
 	/** Return the cloth simulation ready LOD model data. */
 	TSharedPtr<const FChaosClothSimulationModel> GetClothSimulationModel() const { return ClothSimulationModel; }
@@ -112,17 +112,28 @@ public:
 	/**
 	 * Copy the draped simulation mesh patterns into the render mesh data.
 	 * This is useful to visualize the simulation mesh, or when the simulation mesh can be used for both simulation and rendering.
-	 * @param MaterialIndex The index of the Materials array.
+	 * @param Material The material used to render. If none is specified the default debug cloth material will be used.
 	 */
-	void CopySimMeshToRenderMesh(int32 MaterialIndex);
+	void CopySimMeshToRenderMesh(UMaterialInterface* Material = nullptr);
 
 	/** Set the physics asset for this cloth. */
-	void SetPhysicsAsset(UPhysicsAsset* InPhysicsAsset) { PhysicsAsset = InPhysicsAsset; }
+	void SetPhysicsAsset(UPhysicsAsset* InPhysicsAsset);
+
+	/**
+	 * Set the skeleton asset and bone hierachy to use for this cloth.
+	 * @param InSkeleton The skeleton asset to use for this cloth, or nullptr to setup a default skeleton asset.
+	 * @param bRebuildClothModel Whether to rebuild the simulation and rendering models (editor builds only).
+	 */
+	void SetSkeleton(USkeleton* InSkeleton, bool bRebuildModels);
 
 	/** Set the bone hierachy to use for this cloth. */
+	UE_DEPRECATED(5.3, "Use SetSkeleton instead")
+	PRAGMA_DISABLE_DEPRECATION_WARNINGS
 	void SetReferenceSkeleton(const FReferenceSkeleton& InReferenceSkeleton, bool bRebuildClothSimulationModel = true) { RefSkeleton = InReferenceSkeleton; UpdateSkeleton(bRebuildClothSimulationModel); }
+	PRAGMA_ENABLE_DEPRECATION_WARNINGS
 
 	/** Set the skinning weights for all of the sim vertices to be bound to the root node of the reference skeleton. */
+	UE_DEPRECATED(5.3, "Use SetSkeleton instead")
 	void BindSimMeshToRootBone();
 
 	//
@@ -159,10 +170,8 @@ private:
 	/** Re-calculate the bounds for this asset. */
 	void CalculateBounds();
 
-	/** Re-calculate the reference bone index, from which the cloth gets transformed when the simulation is suspended. */
-	void CalculateReferenceBoneIndex();
-
 	/** Update the bone informations after a change of skeleton. */
+	UE_DEPRECATED(5.3, "Use Build() instead")
 	void UpdateSkeleton(bool bRebuildClothSimulationModel = true);
 
 #if WITH_EDITORONLY_DATA
@@ -204,6 +213,10 @@ private:
 	UPROPERTY(EditAnywhere, Category = Materials)
 	TArray<FSkeletalMaterial> Materials;
 
+	/** Skeleton asset used at import time. */
+	UPROPERTY(EditAnywhere, NoClear, Setter = SetSkeleton, Category = Skeleton)
+	TObjectPtr<USkeleton> Skeleton;
+
 	/** Physics asset used for collision. */
 	UPROPERTY(EditAnywhere, Category = Collision)
 	TObjectPtr<UPhysicsAsset> PhysicsAsset;
@@ -256,7 +269,7 @@ private:
 	TArray<FMatrix> CachedComposedRefPoseMatrices;
 
 	/** Cloth Collection containing this asset data. */
-	TSharedPtr<UE::Chaos::ClothAsset::FClothCollection> ClothCollection;
+	TSharedPtr<FManagedArrayCollection> ClothCollection;
 
 	/** Reference skeleton created from the provided skeleton asset. */
 	FReferenceSkeleton RefSkeleton;
diff --git a/Engine/Plugins/Experimental/ChaosClothAssetEditor/Source/ChaosClothAssetEditor/Private/ChaosClothAsset/ClothEditorMode.cpp b/Engine/Plugins/Experimental/ChaosClothAssetEditor/Source/ChaosClothAssetEditor/Private/ChaosClothAsset/ClothEditorMode.cpp
--- a/Engine/Plugins/Experimental/ChaosClothAssetEditor/Source/ChaosClothAssetEditor/Private/ChaosClothAsset/ClothEditorMode.cpp
+++ b/Engine/Plugins/Experimental/ChaosClothAssetEditor/Source/ChaosClothAssetEditor/Private/ChaosClothAsset/ClothEditorMode.cpp
@@ -5,8 +5,7 @@
 #include "ChaosClothAsset/ClothEditorModeToolkit.h"
 #include "ChaosClothAsset/ClothAsset.h"
 #include "ChaosClothAsset/ClothComponent.h"
-#include "ChaosClothAsset/ClothCollection.h"
-#include "ChaosClothAsset/ClothAdapter.h"
+#include "ChaosClothAsset/CollectionClothFacade.h"
 #include "ChaosClothAsset/ClothEditorRestSpaceViewportClient.h"
 #include "ChaosClothAsset/ClothPatternToDynamicMesh.h"
 #include "ChaosClothAsset/ClothEditorPreviewScene.h"
@@ -63,31 +62,13 @@ const FEditorModeID UChaosClothAssetEditorMode::EM_ChaosClothAssetEditorModeId =
 
 namespace ChaosClothAssetEditorModeHelpers
 {
-	TArray<FName> GetClothAssetWeightMapNames(const UE::Chaos::ClothAsset::FClothAdapter& ClothAdapter)
+	void RemoveClothWeightMaps(UE::Chaos::ClothAsset::FCollectionClothFacade& ClothFacade, const TArray<FName>& WeightMapNames)
 	{
-		const TSharedPtr<const UE::Chaos::ClothAsset::FClothCollection> ClothCollection = ClothAdapter.FClothConstAdapter::GetClothCollection();
-
-		TArray<FName> OutWeightMapNames;
-		const TArray<FName> SimVerticesAttributeNames = ClothCollection->AttributeNames(UE::Chaos::ClothAsset::FClothCollection::SimVerticesGroup);
-		for (const FName& AttributeName : SimVerticesAttributeNames)
-		{
-			if (ClothCollection->FindAttributeTyped<float>(AttributeName, UE::Chaos::ClothAsset::FClothCollection::SimVerticesGroup))
-			{
-				OutWeightMapNames.Add(AttributeName);
-			}
-		}
-
-		return OutWeightMapNames;
-	}
-
-	void RemoveClothWeightMaps(UE::Chaos::ClothAsset::FClothAdapter& ClothAdapter, const TArray<FName>& WeightMapNames)
-	{
-		const TSharedPtr<const UE::Chaos::ClothAsset::FClothCollection> ClothCollection = ClothAdapter.GetClothCollection();
-		for (const FName& RemovedMap : WeightMapNames)
+		for (const FName& WeightMapName : WeightMapNames)
 		{
-			if (ClothCollection->FindAttributeTyped<float>(RemovedMap, UE::Chaos::ClothAsset::FClothCollection::SimVerticesGroup))
+			if (ClothFacade.HasWeightMap(WeightMapName))
 			{
-				ClothAdapter.RemoveWeightMap(RemovedMap);
+				ClothFacade.RemoveWeightMap(WeightMapName);
 			}
 		}
 	}
@@ -363,13 +344,13 @@ void UChaosClothAssetEditorMode::UpdateSimulationMeshes()
 	UChaosClothAsset* ChaosClothAsset = PreviewScene->ClothComponent->GetClothAsset();
 	ChaosClothAsset->Modify();
 
-	UE::Chaos::ClothAsset::FClothAdapter ClothAdapter(ChaosClothAsset->GetClothCollection());
+	UE::Chaos::ClothAsset::FCollectionClothFacade ClothFacade(ChaosClothAsset->GetClothCollection());
 
 	check(DynamicMeshSourceInfos.Num() == DynamicMeshComponents.Num());
 
 	// Search for weight map names that are in the original cloth asset, but are not in *all* dynamic mesh components. These are weight maps
 	// that have been removed by one of the mesh tools, and so should be removed from the asset.
-	const TSet<FName> ClothAssetWeightMapNames = TSet<FName>(ChaosClothAssetEditorModeHelpers::GetClothAssetWeightMapNames(ClothAdapter));
+	const TSet<FName> ClothAssetWeightMapNames = TSet<FName>(ClothFacade.GetWeightMapNames());
 	TSet<FName> CommonDynamicMeshWeightMapNames;
 	bool bFirstIteration = true;
 
@@ -378,8 +359,8 @@ void UChaosClothAssetEditorMode::UpdateSimulationMeshes()
 		const int32 LodIndex = DynamicMeshSourceInfos[DynamicMeshIndex].LodIndex;
 		const int32 PatternIndex = DynamicMeshSourceInfos[DynamicMeshIndex].PatternIndex;
 
-		UE::Chaos::ClothAsset::FClothLodAdapter ClothLodAdapter = ClothAdapter.GetLod(LodIndex);
-		UE::Chaos::ClothAsset::FClothPatternAdapter ClothPatternAdapter = ClothLodAdapter.GetPattern(PatternIndex);
+		UE::Chaos::ClothAsset::FCollectionClothLodFacade ClothLodFacade = ClothFacade.GetLod(LodIndex);
+		UE::Chaos::ClothAsset::FCollectionClothPatternFacade ClothPatternFacade = ClothLodFacade.GetPattern(PatternIndex);
 
 		const UDynamicMeshComponent& DynamicMeshComponent = *DynamicMeshComponents[DynamicMeshIndex];
 		const UE::Geometry::FDynamicMesh3* DynamicMesh = DynamicMeshComponent.GetMesh();
@@ -416,9 +397,9 @@ void UChaosClothAssetEditorMode::UpdateSimulationMeshes()
 			const UE::Geometry::FDynamicMeshWeightAttribute* WeightMapAttribute = DynamicMesh->Attributes()->GetWeightLayer(DynamicMeshWeightMapIndex);
 			const FName WeightMapName = WeightMapAttribute->GetName();
 
-			ClothAdapter.AddWeightMap(WeightMapName);		// Does nothing if weight map already exists
+			ClothFacade.AddWeightMap(WeightMapName);		// Does nothing if weight map already exists
 
-			TArrayView<float> PatternWeights = ClothPatternAdapter.GetWeightMap(WeightMapName);
+			TArrayView<float> PatternWeights = ClothPatternFacade.GetWeightMap(WeightMapName);
 			for (int VertexID = 0; VertexID < NumInputWeightElements; ++VertexID)
 			{
 				float Val;
@@ -429,7 +410,7 @@ void UChaosClothAssetEditorMode::UpdateSimulationMeshes()
 	}
 
 	const TSet<FName> WeightMapsToRemove = ClothAssetWeightMapNames.Difference(CommonDynamicMeshWeightMapNames);
-	ChaosClothAssetEditorModeHelpers::RemoveClothWeightMaps(ClothAdapter, WeightMapsToRemove.Array());
+	ChaosClothAssetEditorModeHelpers::RemoveClothWeightMaps(ClothFacade, WeightMapsToRemove.Array());
 
 	ChaosClothAsset->Build();
 
@@ -512,11 +493,11 @@ void UChaosClothAssetEditorMode::ReinitializeDynamicMeshComponents()
 		return;
 	}
 
-	const UE::Chaos::ClothAsset::FClothConstAdapter ClothAdapter(ChaosClothAsset->GetClothCollection());
+	const UE::Chaos::ClothAsset::FCollectionClothConstFacade ClothFacade(ChaosClothAsset->GetClothCollection());
 
-	for (int32 LodIndex = 0; LodIndex < ClothAdapter.GetNumLods(); ++LodIndex)
+	for (int32 LodIndex = 0; LodIndex < ClothFacade.GetNumLods(); ++LodIndex)
 	{
-		const UE::Chaos::ClothAsset::FClothLodConstAdapter ClothLodAdapter = ClothAdapter.GetLod(LodIndex);
+		const UE::Chaos::ClothAsset::FCollectionClothLodConstFacade ClothLodFacade = ClothFacade.GetLod(LodIndex);
 
 		if (bCombineAllPatterns)
 		{
@@ -525,11 +506,11 @@ void UChaosClothAssetEditorMode::ReinitializeDynamicMeshComponents()
 		}
 		else
 		{
-			const int32 NumComponents = ClothLodAdapter.GetNumPatterns();
+			const int32 NumComponents = ClothLodFacade.GetNumPatterns();
 			DynamicMeshComponents.SetNum(NumComponents);
 			DynamicMeshComponentParentActors.SetNum(NumComponents);
 
-			for (int32 PatternIndex = 0; PatternIndex < ClothLodAdapter.GetNumPatterns(); ++PatternIndex)
+			for (int32 PatternIndex = 0; PatternIndex < ClothLodFacade.GetNumPatterns(); ++PatternIndex)
 			{
 				UE::Geometry::FDynamicMesh3 PatternMesh;
 				FClothPatternToDynamicMesh Converter;
@@ -567,7 +548,7 @@ void UChaosClothAssetEditorMode::ReinitializeDynamicMeshComponents()
 
 		// If we found any patterns, use this LOD
 		// TODO: Give the user the ability to specify which LOD to display
-		if (ClothLodAdapter.GetNumPatterns() > 0)
+		if (ClothLodFacade.GetNumPatterns() > 0)
 		{
 			break;
 		}
diff --git a/Engine/Plugins/Experimental/ChaosClothAssetEditor/Source/ChaosClothAssetEditor/Private/ChaosClothAsset/ClothEditorToolkit.cpp b/Engine/Plugins/Experimental/ChaosClothAssetEditor/Source/ChaosClothAssetEditor/Private/ChaosClothAsset/ClothEditorToolkit.cpp
--- a/Engine/Plugins/Experimental/ChaosClothAssetEditor/Source/ChaosClothAssetEditor/Private/ChaosClothAsset/ClothEditorToolkit.cpp
+++ b/Engine/Plugins/Experimental/ChaosClothAssetEditor/Source/ChaosClothAssetEditor/Private/ChaosClothAsset/ClothEditorToolkit.cpp
@@ -11,7 +11,6 @@
 #include "ChaosClothAsset/ClothEditorRestSpaceViewportClient.h"
 #include "ChaosClothAsset/ClothComponent.h"
 #include "ChaosClothAsset/ClothAsset.h"
-#include "ChaosClothAsset/ClothCollection.h"
 #include "ChaosClothAsset/ClothEditorModeToolkit.h"
 #include "Framework/Docking/LayoutExtender.h"
 #include "AssetEditorModeManager.h"
@@ -706,7 +705,7 @@ void FChaosClothAssetEditorToolkit::InitDetailsViewPanel()
 
 void FChaosClothAssetEditorToolkit::PopulateOutliner()
 {
-	TSharedPtr<UE::Chaos::ClothAsset::FClothCollection> ClothCollection;
+	TSharedPtr<FManagedArrayCollection> ClothCollection;
 	if (const TObjectPtr<UChaosClothComponent> ClothComponent = ClothPreviewScene->ClothComponent)
 	{
 		if (UChaosClothAsset* const ClothAsset = ClothComponent->GetClothAsset())
@@ -717,12 +716,15 @@ void FChaosClothAssetEditorToolkit::PopulateOutliner()
 	
 	if (Outliner.IsValid() && ClothCollection.IsValid())
 	{
-		const FName& SelectedGroupName = ClothCollection->SimVerticesGroup;
-		
-		Outliner->SetClothCollection(ClothCollection);
-		Outliner->SetSelectedGroupName(SelectedGroupName);
-
 		ClothCollectionGroupNames = ClothCollection->GroupNames();
+		if (ClothCollectionGroupNames.Num())
+		{
+			const FName& SelectedGroupName = ClothCollectionGroupNames[0];
+	
+			Outliner->SetClothCollection(ClothCollection);
+			Outliner->SetSelectedGroupName(SelectedGroupName);
+		}
+
 	}
 }
 
diff --git a/Engine/Plugins/Experimental/ChaosClothAssetEditor/Source/ChaosClothAssetEditor/Private/ChaosClothAsset/DataflowNodes/BindToRootBoneNode.cpp b/Engine/Plugins/Experimental/ChaosClothAssetEditor/Source/ChaosClothAssetEditor/Private/ChaosClothAsset/DataflowNodes/BindToRootBoneNode.cpp
--- a/Engine/Plugins/Experimental/ChaosClothAssetEditor/Source/ChaosClothAssetEditor/Private/ChaosClothAsset/DataflowNodes/BindToRootBoneNode.cpp
+++ b/Engine/Plugins/Experimental/ChaosClothAssetEditor/Source/ChaosClothAssetEditor/Private/ChaosClothAsset/DataflowNodes/BindToRootBoneNode.cpp
@@ -22,14 +22,10 @@ void FChaosClothAssetBindToRootBoneNode::Evaluate(Dataflow::FContext& Context, c
 		using namespace UE::Chaos::ClothAsset;
 
 		// Evaluate in collection
-		const FManagedArrayCollection& InCollection = GetValue<FManagedArrayCollection>(Context, &Collection);
+		FManagedArrayCollection InCollection = GetValue<FManagedArrayCollection>(Context, &Collection);
+		const TSharedRef<FManagedArrayCollection> ClothCollection = MakeShared<FManagedArrayCollection>(MoveTemp(InCollection));
 
-		// TODO: Needs to make cloth collection a facade to avoid the copy of the cloth collection to a managed array
-		//       and remove the above const reference to operate on the InCollection instead.
-		TSharedPtr<FClothCollection> ClothCollection = MakeShared<FClothCollection>();
-		InCollection.CopyTo(ClothCollection.Get());
-
-        FClothGeometryTools::BindMeshToRootBone(ClothCollection, bBindSimMesh, bBindRenderMesh);
+		FClothGeometryTools::BindMeshToRootBone(ClothCollection, bBindSimMesh, bBindRenderMesh);
 
 		SetValue<FManagedArrayCollection>(Context, *ClothCollection, &Collection);
 	}
diff --git a/Engine/Plugins/Experimental/ChaosClothAssetEditor/Source/ChaosClothAssetEditor/Private/ChaosClothAsset/DataflowNodes/CopySimulationToRenderMeshNode.cpp b/Engine/Plugins/Experimental/ChaosClothAssetEditor/Source/ChaosClothAssetEditor/Private/ChaosClothAsset/DataflowNodes/CopySimulationToRenderMeshNode.cpp
--- a/Engine/Plugins/Experimental/ChaosClothAssetEditor/Source/ChaosClothAssetEditor/Private/ChaosClothAsset/DataflowNodes/CopySimulationToRenderMeshNode.cpp
+++ b/Engine/Plugins/Experimental/ChaosClothAssetEditor/Source/ChaosClothAssetEditor/Private/ChaosClothAsset/DataflowNodes/CopySimulationToRenderMeshNode.cpp
@@ -24,12 +24,8 @@ void FChaosClothAssetCopySimulationToRenderMeshNode::Evaluate(Dataflow::FContext
 		using namespace UE::Chaos::ClothAsset;
 
 		// Evaluate in collection
-		const FManagedArrayCollection& InCollection = GetValue<FManagedArrayCollection>(Context, &Collection);
-
-		// TODO: Needs to make cloth collection a facade to avoid the copy of the cloth collection to a managed array
-		//       and remove the above const reference to operated on the InCollection instead.
-		TSharedPtr<FClothCollection> ClothCollection = MakeShared<FClothCollection>();
-		InCollection.CopyTo(ClothCollection.Get());
+		FManagedArrayCollection InCollection = GetValue<FManagedArrayCollection>(Context, &Collection);
+		const TSharedRef<FManagedArrayCollection> ClothCollection = MakeShared<FManagedArrayCollection>(MoveTemp(InCollection));
 
 		// Empty mesh and materials
 		FClothGeometryTools::DeleteRenderMesh(ClothCollection);
@@ -39,22 +35,8 @@ void FChaosClothAssetCopySimulationToRenderMeshNode::Evaluate(Dataflow::FContext
 			Material->GetPathName() :
 			FString(TEXT("/Engine/EditorMaterials/Cloth/CameraLitDoubleSided.CameraLitDoubleSided"));
 
-		// Find or add the material attribute
-		// TODO: Add to ClothAdapter functionality
-		TManagedArray<FString>* const FindAttributeResult = ClothCollection->FindAttribute<FString>("MaterialPathName", FClothCollection::MaterialsGroup);
-		TManagedArray<FString>& MaterialPathNameArray = FindAttributeResult ? *FindAttributeResult :
-			ClothCollection->AddAttribute<FString>("MaterialPathName", FClothCollection::MaterialsGroup);
-
-		// Work out the material index and set the path name
-		int32 MaterialIndex = MaterialPathNameArray.Find(MaterialPathName);
-	
-		if (MaterialIndex == INDEX_NONE)
-		{
-			MaterialIndex = ClothCollection->AddElements(1, FClothCollection::MaterialsGroup);
-			MaterialPathNameArray[MaterialIndex] = MaterialPathName;
-		}
-
-		FClothGeometryTools::CopySimMeshToRenderMesh(ClothCollection, MaterialIndex);
+		// Copy the mesh data
+		FClothGeometryTools::CopySimMeshToRenderMesh(ClothCollection, MaterialPathName);
 
 		SetValue<FManagedArrayCollection>(Context, *ClothCollection, &Collection);
 	}
diff --git a/Engine/Plugins/Experimental/ChaosClothAssetEditor/Source/ChaosClothAssetEditor/Private/ChaosClothAsset/DataflowNodes/DatasmithImportNode.cpp b/Engine/Plugins/Experimental/ChaosClothAssetEditor/Source/ChaosClothAssetEditor/Private/ChaosClothAsset/DataflowNodes/DatasmithImportNode.cpp
--- a/Engine/Plugins/Experimental/ChaosClothAssetEditor/Source/ChaosClothAssetEditor/Private/ChaosClothAsset/DataflowNodes/DatasmithImportNode.cpp
+++ b/Engine/Plugins/Experimental/ChaosClothAssetEditor/Source/ChaosClothAssetEditor/Private/ChaosClothAsset/DataflowNodes/DatasmithImportNode.cpp
@@ -2,8 +2,8 @@
 
 #include "ChaosClothAsset/DataflowNodes/DatasmithImportNode.h"
 #include "ChaosClothAsset/DataflowNodes/DataflowNodes.h"
-#include "ChaosClothAsset/ClothCollection.h"
 #include "ChaosClothAsset/ClothAsset.h"
+#include "ChaosClothAsset/CollectionClothFacade.h"
 #include "DatasmithImportContext.h"
 #include "DatasmithImportFactory.h"
 #include "ExternalSource.h"
@@ -99,9 +99,15 @@ void FChaosClothAssetDatasmithImportNode::Evaluate(Dataflow::FContext& Context,
 	}
 	else
 	{
+		using namespace UE::Chaos::ClothAsset;
+
 		// Init with an empty cloth collection
-		UE::Chaos::ClothAsset::FClothCollection ClothCollection;
-		SetValue<FManagedArrayCollection>(Context, ClothCollection, &Collection);
+		const TSharedRef<FManagedArrayCollection> ClothCollection = MakeShared<FManagedArrayCollection>();
+		FCollectionClothFacade CollectionClothFacade(ClothCollection);
+		CollectionClothFacade.DefineSchema();
+		CollectionClothFacade.AddLod();
+
+		SetValue<FManagedArrayCollection>(Context, *ClothCollection, &Collection);
 	}
 }
 
diff --git a/Engine/Plugins/Experimental/ChaosClothAssetEditor/Source/ChaosClothAssetEditor/Private/ChaosClothAsset/DataflowNodes/DeleteRenderMeshNode.cpp b/Engine/Plugins/Experimental/ChaosClothAssetEditor/Source/ChaosClothAssetEditor/Private/ChaosClothAsset/DataflowNodes/DeleteRenderMeshNode.cpp
--- a/Engine/Plugins/Experimental/ChaosClothAssetEditor/Source/ChaosClothAssetEditor/Private/ChaosClothAsset/DataflowNodes/DeleteRenderMeshNode.cpp
+++ b/Engine/Plugins/Experimental/ChaosClothAssetEditor/Source/ChaosClothAssetEditor/Private/ChaosClothAsset/DataflowNodes/DeleteRenderMeshNode.cpp
@@ -23,12 +23,8 @@ void FChaosClothAssetDeleteRenderMeshNode::Evaluate(Dataflow::FContext& Context,
 		using namespace UE::Chaos::ClothAsset;
 
 		// Evaluate in collection
-		const FManagedArrayCollection& InCollection = GetValue<FManagedArrayCollection>(Context, &Collection);
-
-		// TODO: Needs to make cloth collection a facade to avoid the copy of the cloth collection to a managed array
-		//       and remove the above const reference to operated on the InCollection instead.
-		TSharedPtr<FClothCollection> ClothCollection = MakeShared<FClothCollection>();
-		InCollection.CopyTo(ClothCollection.Get());
+		FManagedArrayCollection InCollection = GetValue<FManagedArrayCollection>(Context, &Collection);
+		const TSharedRef<FManagedArrayCollection> ClothCollection = MakeShared<FManagedArrayCollection>(MoveTemp(InCollection));
 
 		if (Patterns.IsEmpty())
 		{
diff --git a/Engine/Plugins/Experimental/ChaosClothAssetEditor/Source/ChaosClothAssetEditor/Private/ChaosClothAsset/DataflowNodes/ImportNode.cpp b/Engine/Plugins/Experimental/ChaosClothAssetEditor/Source/ChaosClothAssetEditor/Private/ChaosClothAsset/DataflowNodes/ImportNode.cpp
--- a/Engine/Plugins/Experimental/ChaosClothAssetEditor/Source/ChaosClothAssetEditor/Private/ChaosClothAsset/DataflowNodes/ImportNode.cpp
+++ b/Engine/Plugins/Experimental/ChaosClothAssetEditor/Source/ChaosClothAssetEditor/Private/ChaosClothAsset/DataflowNodes/ImportNode.cpp
@@ -3,7 +3,7 @@
 #include "ChaosClothAsset/DataflowNodes/ImportNode.h"
 #include "ChaosClothAsset/DataflowNodes/DataflowNodes.h"
 #include "ChaosClothAsset/ClothAsset.h"
-#include "ChaosClothAsset/ClothCollection.h"
+#include "ChaosClothAsset/CollectionClothFacade.h"
 #include "Engine/SkinnedAssetCommon.h"
 #include "Materials/Material.h"
 #include "PhysicsEngine/PhysicsAsset.h"
@@ -15,6 +15,8 @@
 FChaosClothAssetImportNode::FChaosClothAssetImportNode(const Dataflow::FNodeParameters& InParam, FGuid InGuid)
 	: FDataflowNode(InParam, InGuid)
 {
+	RegisterInputConnection(&ClothAsset);
+	RegisterInputConnection(&ImportLod);
 	RegisterOutputConnection(&Collection);
 }
 
@@ -52,52 +54,31 @@ void FChaosClothAssetImportNode::Evaluate(Dataflow::FContext& Context, const FDa
 			}
 		}
 
+		// Create a new cloth collection with its LOD 0
+		const TSharedRef<FManagedArrayCollection> ClothCollection = MakeShared<FManagedArrayCollection>();
+		FCollectionClothFacade ClothFacade(ClothCollection);
+		ClothFacade.DefineSchema();
+		FCollectionClothLodFacade ClothLodFacade = ClothFacade.AddGetLod();  // Always output a single LOD at LOD 0, even if empty
+
 		// Copy the cloth asset to this node's output collection
 		if (bSetValue)
 		{
-			// TODO: Needs to make cloth collection a facade to avoid the copy of the cloth collection to a managed array
-			TSharedPtr<FClothCollection> ClothCollection = MakeShared<FClothCollection>();
-
-			if (ClothAsset->GetClothCollection())
-			{
-				const TArray<FName> GroupsToSkip =
-				{
-					FClothCollection::MaterialsGroup,   // Don't copy materials, and collisions but rebuild the
-					FClothCollection::CollisionsGroup   // group from the asset for backward compatibility instead
-				};
-				ClothAsset->GetClothCollection()->CopyTo(ClothCollection.Get(), GroupsToSkip);
-			}
-
-			// Rebuild material list
-			const TArray<FSkeletalMaterial>& Materials = ClothAsset->GetMaterials();
-			const int32 NumMaterials = Materials.Num();
-
-			ClothCollection->AddElements(NumMaterials, FClothCollection::MaterialsGroup);
-
-			for (int32 MaterialIndex = 0; MaterialIndex < NumMaterials; ++MaterialIndex)
+			const TObjectPtr<const UChaosClothAsset>& InClothAsset = GetValue<TObjectPtr<const UChaosClothAsset>>(Context, &ClothAsset);
+			if (InClothAsset->GetClothCollection())
 			{
-				const FSkeletalMaterial& Material = Materials[MaterialIndex];
-				ClothCollection->MaterialPathName[MaterialIndex] = Material.MaterialInterface ?
-					Material.MaterialInterface->GetPathName() :
-					FString();
-			}
+				const int32& InImportLod = GetValue<int32>(Context, &ImportLod);
+				const FCollectionClothConstFacade InClothFacade(InClothAsset->GetClothCollection());
 
-			// Overwrite physics asset
-			if (const UPhysicsAsset* const PhysicsAsset = ClothAsset->GetPhysicsAsset())
-			{
-				ClothCollection->AddElements(1, FClothCollection::CollisionsGroup);
-				ClothCollection->PhysicsAssetPathName[0] = PhysicsAsset->GetPathName();
+				if (InImportLod >= 0 && InImportLod < InClothFacade.GetNumLods())
+				{
+					// Copy specified LOD from the input asset to the LOD 0 of the output collection
+					const FCollectionClothLodConstFacade InClothLodFacade = InClothFacade.GetLod(InImportLod);
+					ClothLodFacade.Initialize(InClothLodFacade);
+				}
 			}
-
-			SetValue<FManagedArrayCollection>(Context, *ClothCollection, &Collection);
-		}
-		else
-		{
-			// Init with an empty cloth collection
-			UE::Chaos::ClothAsset::FClothCollection ClothCollection;
-			SetValue<FManagedArrayCollection>(Context, ClothCollection, &Collection);
 		}
-	}
+		SetValue<FManagedArrayCollection>(Context, *ClothCollection, &Collection);
+}
 }
 
 #undef LOCTEXT_NAMESPACE
diff --git a/Engine/Plugins/Experimental/ChaosClothAssetEditor/Source/ChaosClothAssetEditor/Private/ChaosClothAsset/DataflowNodes/ImportNode.h b/Engine/Plugins/Experimental/ChaosClothAssetEditor/Source/ChaosClothAssetEditor/Private/ChaosClothAsset/DataflowNodes/ImportNode.h
--- a/Engine/Plugins/Experimental/ChaosClothAssetEditor/Source/ChaosClothAssetEditor/Private/ChaosClothAsset/DataflowNodes/ImportNode.h
+++ b/Engine/Plugins/Experimental/ChaosClothAssetEditor/Source/ChaosClothAssetEditor/Private/ChaosClothAsset/DataflowNodes/ImportNode.h
@@ -20,9 +20,14 @@ public:
 	UPROPERTY(Meta = (DataflowOutput, DisplayName = "Collection"))
 	FManagedArrayCollection Collection;
 
-	UPROPERTY(EditAnywhere, Category = "Dataflow")
+	/** The Cloth Asset to import into a collection. */
+	UPROPERTY(EditAnywhere, Category = "ClothAsset")
 	TObjectPtr<const UChaosClothAsset> ClothAsset;
 
+	/** The LOD to import into the collection. Only one LOD can be imported at a time. */
+	UPROPERTY(EditAnywhere, Category = "ClothAsset", Meta = (DisplayName = "Import LOD"))
+	int32 ImportLod;
+
 	FChaosClothAssetImportNode(const Dataflow::FNodeParameters& InParam, FGuid InGuid = FGuid::NewGuid());
 
 	virtual void Evaluate(Dataflow::FContext& Context, const FDataflowOutput* Out) const override;
diff --git a/Engine/Plugins/Experimental/ChaosClothAssetEditor/Source/ChaosClothAssetEditor/Private/ChaosClothAsset/DataflowNodes/ReverseNormalsNode.cpp b/Engine/Plugins/Experimental/ChaosClothAssetEditor/Source/ChaosClothAssetEditor/Private/ChaosClothAsset/DataflowNodes/ReverseNormalsNode.cpp
--- a/Engine/Plugins/Experimental/ChaosClothAssetEditor/Source/ChaosClothAssetEditor/Private/ChaosClothAsset/DataflowNodes/ReverseNormalsNode.cpp
+++ b/Engine/Plugins/Experimental/ChaosClothAssetEditor/Source/ChaosClothAssetEditor/Private/ChaosClothAsset/DataflowNodes/ReverseNormalsNode.cpp
@@ -23,12 +23,8 @@ void FChaosClothAssetReverseNormalsNode::Evaluate(Dataflow::FContext& Context, c
 		using namespace UE::Chaos::ClothAsset;
 
 		// Evaluate in collection
-		const FManagedArrayCollection& InCollection = GetValue<FManagedArrayCollection>(Context, &Collection);
-
-		// TODO: Needs to make cloth collection a facade to avoid the copy of the cloth collection to a managed array
-		//       and remove the above const reference to operated on the InCollection instead.
-		TSharedPtr<FClothCollection> ClothCollection = MakeShared<FClothCollection>();
-		InCollection.CopyTo(ClothCollection.Get());
+		FManagedArrayCollection InCollection = GetValue<FManagedArrayCollection>(Context, &Collection);
+		const TSharedRef<FManagedArrayCollection> ClothCollection = MakeShared<FManagedArrayCollection>(MoveTemp(InCollection));
 
 		FClothGeometryTools::ReverseNormals(ClothCollection, bReverseSimMeshNormals, bReverseRenderMeshNormals, Patterns);
 
diff --git a/Engine/Plugins/Experimental/ChaosClothAssetEditor/Source/ChaosClothAssetEditor/Private/ChaosClothAsset/DataflowNodes/TerminalNode.cpp b/Engine/Plugins/Experimental/ChaosClothAssetEditor/Source/ChaosClothAssetEditor/Private/ChaosClothAsset/DataflowNodes/TerminalNode.cpp
--- a/Engine/Plugins/Experimental/ChaosClothAssetEditor/Source/ChaosClothAssetEditor/Private/ChaosClothAsset/DataflowNodes/TerminalNode.cpp
+++ b/Engine/Plugins/Experimental/ChaosClothAssetEditor/Source/ChaosClothAssetEditor/Private/ChaosClothAsset/DataflowNodes/TerminalNode.cpp
@@ -3,7 +3,7 @@
 #include "ChaosClothAsset/DataflowNodes/TerminalNode.h"
 #include "ChaosClothAsset/DataflowNodes/DataflowNodes.h"
 #include "ChaosClothAsset/ClothAsset.h"
-#include "ChaosClothAsset/ClothAdapter.h"
+#include "ChaosClothAsset/CollectionClothFacade.h"
 #include "Animation/Skeleton.h"
 #include "Engine/SkinnedAssetCommon.h"
 #include "Materials/Material.h"
@@ -13,15 +13,10 @@
 
 #define LOCTEXT_NAMESPACE "ChaosClothAssetTerminalNode"
 
-namespace UE::Chaos::ClothAsset::DataflowNodes
-{
-	static const FString DefaultSkeletonPathName = TEXT("/Engine/EditorMeshes/SkeletalMesh/DefaultSkeletalMesh_Skeleton.DefaultSkeletalMesh_Skeleton");
-}
-
 FChaosClothAssetTerminalNode::FChaosClothAssetTerminalNode(const Dataflow::FNodeParameters& InParam, FGuid InGuid)
 	: FDataflowTerminalNode(InParam, InGuid)
 {
-	RegisterInputConnection(&Collection);
+	RegisterInputConnection(&CollectionLod0);
 }
 
 void FChaosClothAssetTerminalNode::SetAssetValue(TObjectPtr<UObject> Asset, Dataflow::FContext& Context) const
@@ -30,63 +25,159 @@ void FChaosClothAssetTerminalNode::SetAssetValue(TObjectPtr<UObject> Asset, Data
 	{
 		using namespace UE::Chaos::ClothAsset;
 
-		const FManagedArrayCollection& InCollection = GetValue<FManagedArrayCollection>(Context, &Collection);
+		// Reset the asset's collection
+		const TSharedPtr<FManagedArrayCollection> ClothCollection = ClothAsset->GetClothCollection();
+		ClothCollection->Reset();
 
-		FClothCollection& ClothCollection = *ClothAsset->GetClothCollection();
+		FCollectionClothFacade ClothFacade(ClothCollection);
+		ClothFacade.DefineSchema();
 
-		// Replace collection
-		ClothCollection.Reset();
-		InCollection.CopyTo(&ClothCollection);
+		// Iterate through the LODs
+		FString SkeletonPathName;
+		FString PhysicsAssetPathName;
 
-		// Set materials
-		const int32 NumMaterials = ClothCollection.MaterialPathName.Num();
+		const TArray<const FManagedArrayCollection*> CollectionLods = GetCollectionLods();
+		for (int32 LodIndex = 0; LodIndex < CollectionLods.Num(); ++LodIndex)
+		{
+			// New LOD
+			const int32 NewLodIndex = ClothFacade.AddLod();
+			check(NewLodIndex == LodIndex);
 
-		TArray<FSkeletalMaterial>& Materials = ClothAsset->GetMaterials();
-		Materials.Reset(NumMaterials);
+			// Retrieve input LOD
+			const FManagedArrayCollection& InCollectionLod = GetValue<FManagedArrayCollection>(Context, CollectionLods[LodIndex]);
+			const TSharedRef<const FManagedArrayCollection> InClothCollection = MakeShared<const FManagedArrayCollection>(InCollectionLod);
 
-		for (int32 MaterialIndex = 0; MaterialIndex < NumMaterials; ++MaterialIndex)
-		{
-			if (UMaterial* const Material = LoadObject<UMaterial>(ClothAsset, *ClothCollection.MaterialPathName[MaterialIndex], nullptr, LOAD_None, nullptr))
+			// Only use LOD 0 from input
+			const FCollectionClothConstFacade InClothFacade(InClothCollection);
+			if (InClothFacade.GetNumLods() >= 1)
 			{
-				Materials.Emplace(Material, true, false, Material->GetFName());
-			}
-			else
-			{
-				Materials.Emplace();
+				const FCollectionClothLodConstFacade InClothLodFacade = InClothFacade.GetLod(0);
+
+				// Copy input LOD 0 to current output LOD
+				FCollectionClothLodFacade ClothLodFacade = ClothFacade.GetLod(LodIndex);
+				ClothLodFacade.Initialize(InClothLodFacade);
+
+				// Add this LOD's materials to the asset
+				const int32 NumLodMaterials = InClothLodFacade.GetNumMaterials();
+
+				TArray<FSkeletalMaterial>& Materials = ClothAsset->GetMaterials();
+				Materials.Reserve(Materials.Num() + NumLodMaterials);
+
+				for (int32 LodMaterialIndex = 0; LodMaterialIndex < NumLodMaterials; ++LodMaterialIndex)
+				{
+					const FString& RenderMaterialPathName = ClothLodFacade.GetRenderMaterialPathName()[LodMaterialIndex];
+
+					auto Predicate = [&RenderMaterialPathName](const FSkeletalMaterial& SkeletalMaterial)
+						{
+							return SkeletalMaterial.MaterialInterface && SkeletalMaterial.MaterialInterface->GetPathName() == RenderMaterialPathName;
+						};
+
+					if (!Materials.FindByPredicate(Predicate))
+					{
+						if (UMaterial* const Material = LoadObject<UMaterial>(ClothAsset, *RenderMaterialPathName, nullptr, LOAD_None, nullptr))
+						{
+							Materials.Emplace(Material, true, false, Material->GetFName());
+						}
+					}
+				}
+
+				// Set skeleton and physics asset only with LOD 0 at the moment
+				if (LodIndex == 0)
+				{
+					SkeletonPathName = InClothLodFacade.GetSkeletonAssetPathName();
+					PhysicsAssetPathName = InClothLodFacade.GetPhysicsAssetPathName();
+				}
 			}
 		}
 
-		// Set reference skeleton
-		const FString SkeletonPathName = (ClothCollection.SkeletonAssetPathName.Num() && !ClothCollection.SkeletonAssetPathName[0].IsEmpty()) ?
-			ClothCollection.SkeletonAssetPathName[0] :
-			DataflowNodes::DefaultSkeletonPathName;
-
-		if (const USkeleton* const Skeleton = LoadObject<USkeleton>(nullptr, *SkeletonPathName, nullptr, LOAD_None, nullptr))
+		// Make sure that whatever happens there is always at least one empty LOD to avoid crashing the render data
+		if (ClothFacade.GetNumLods() < 1)
 		{
-			constexpr bool bRebuildClothSimulationModel = false;  // Avoid rebuilding the asset twice
-			ClothAsset->SetReferenceSkeleton(Skeleton->GetReferenceSkeleton(), bRebuildClothSimulationModel);
+			ClothFacade.AddLod();
 		}
 
+		// Set reference skeleton
+		constexpr bool bRebuildModels = false;  // Avoid rebuilding the asset twice
+		USkeleton* const Skeleton = !SkeletonPathName.IsEmpty() ?
+			LoadObject<USkeleton>(nullptr, *SkeletonPathName, nullptr, LOAD_None, nullptr) :
+			nullptr;
+		ClothAsset->SetSkeleton(Skeleton, bRebuildModels);
+
 		// Set physics asset (note: the cloth asset's physics asset is only replaced if a collection path name is found valid)
-		if (ClothCollection.PhysicsAssetPathName.Num() && !ClothCollection.PhysicsAssetPathName[0].IsEmpty())
+		UPhysicsAsset* const PhysicsAsset = !PhysicsAssetPathName.IsEmpty() ?
+			LoadObject<UPhysicsAsset>(ClothAsset, *PhysicsAssetPathName, nullptr, LOAD_None, nullptr) :
+			nullptr;
+		ClothAsset->SetPhysicsAsset(PhysicsAsset);
+
+		// Rebuild the asset static data
+		ClothAsset->Build();
+	}
+}
+
+Dataflow::FPin FChaosClothAssetTerminalNode::AddPin()
+{
+	auto AddInput = [this](const FManagedArrayCollection* Collection) -> Dataflow::FPin
 		{
-			check(ClothCollection.PhysicsAssetPathName.Num() == 1);  // Can't really deal with more than a single physics asset at the moment
-			if (UPhysicsAsset* const PhysicsAsset = LoadObject<UPhysicsAsset>(ClothAsset, *ClothCollection.PhysicsAssetPathName[0], nullptr, LOAD_None, nullptr))
-			{
-				ClothAsset->SetPhysicsAsset(PhysicsAsset);
-			}
-		}
+			RegisterInputConnection(Collection);
+			const FDataflowInput* const Input = FindInput(Collection);
+			return { Dataflow::FPin::EDirection::INPUT, Input->GetType(), Input->GetName() };
+		};
+
+	switch (NumLods)
+	{
+	case 1: ++NumLods; return AddInput(&CollectionLod1);
+	case 2: ++NumLods; return AddInput(&CollectionLod2);
+	case 3: ++NumLods; return AddInput(&CollectionLod3);
+	case 4: ++NumLods; return AddInput(&CollectionLod4);
+	case 5: ++NumLods; return AddInput(&CollectionLod5);
+	default: break;
+	}
 
-		// Check there is at least one empty LOD to avoid crashing the render data
-		FClothAdapter ClothAdapter(ClothAsset->GetClothCollection());
-		if (ClothAdapter.GetNumLods() < 1)
+	return Super::AddPin();
+}
+
+Dataflow::FPin FChaosClothAssetTerminalNode::RemovePin()
+{
+	auto RemoveInput = [this](const FManagedArrayCollection* Collection) -> Dataflow::FPin
 		{
-			ClothAdapter.AddLod();
-		}
+			const FDataflowInput* const Input = FindInput(Collection);
+			check(Input);
+			Dataflow::FPin Pin = { Dataflow::FPin::EDirection::INPUT, Input->GetType(), Input->GetName() };
+			UnregisterInputConnection(Collection);  // This will delete the input, so set the pin before that
+			return Pin;
+		};
+
+	switch (NumLods - 1)
+	{
+	case 1: --NumLods; return RemoveInput(&CollectionLod1);
+	case 2: --NumLods; return RemoveInput(&CollectionLod2);
+	case 3: --NumLods; return RemoveInput(&CollectionLod3);
+	case 4: --NumLods; return RemoveInput(&CollectionLod4);
+	case 5: --NumLods; return RemoveInput(&CollectionLod5);
+	default: break;
+	}
+	return Super::AddPin();
+}
 
-		// Rebuild the asset static data
-		ClothAsset->Build();
+TArray<const FManagedArrayCollection*> FChaosClothAssetTerminalNode::GetCollectionLods() const
+{
+	TArray<const FManagedArrayCollection*> CollectionLods;
+	CollectionLods.SetNumUninitialized(NumLods);
+
+	for (int32 LodIndex = 0; LodIndex < NumLods; ++LodIndex)
+	{
+		switch (LodIndex)
+		{
+		case 0: CollectionLods[LodIndex] = &CollectionLod0; break;
+		case 1: CollectionLods[LodIndex] = &CollectionLod1; break;
+		case 2: CollectionLods[LodIndex] = &CollectionLod2; break;
+		case 3: CollectionLods[LodIndex] = &CollectionLod3; break;
+		case 4: CollectionLods[LodIndex] = &CollectionLod4; break;
+		case 5: CollectionLods[LodIndex] = &CollectionLod5; break;
+		default: check(false); break;
+		}
 	}
+	return CollectionLods;
 }
 
 #undef LOCTEXT_NAMESPACE
diff --git a/Engine/Plugins/Experimental/ChaosClothAssetEditor/Source/ChaosClothAssetEditor/Private/ChaosClothAsset/DataflowNodes/TerminalNode.h b/Engine/Plugins/Experimental/ChaosClothAssetEditor/Source/ChaosClothAssetEditor/Private/ChaosClothAsset/DataflowNodes/TerminalNode.h
--- a/Engine/Plugins/Experimental/ChaosClothAssetEditor/Source/ChaosClothAssetEditor/Private/ChaosClothAsset/DataflowNodes/TerminalNode.h
+++ b/Engine/Plugins/Experimental/ChaosClothAssetEditor/Source/ChaosClothAssetEditor/Private/ChaosClothAsset/DataflowNodes/TerminalNode.h
@@ -14,13 +14,41 @@ struct FChaosClothAssetTerminalNode : public FDataflowTerminalNode
 	//DATAFLOW_NODE_RENDER_TYPE(FGeometryCollection::StaticType(), "Collection")  // TODO: Leave out the render type until there is something to render
 
 public:
-
-	UPROPERTY(meta = (DataflowInput, DisplayName = "Collection"))
-	FManagedArrayCollection Collection;
+	static constexpr int32 MaxLods = 6;  // Hardcoded number of LODs since it is currently not possible to use arrays for optional inputs
+
+	/** LOD 0 input, right click on the node and add pins to add more LODs. */
+	UPROPERTY(meta = (DataflowInput, DisplayName = "Collection LOD 0"))
+	FManagedArrayCollection CollectionLod0;
+	/** LOD 1 input, right click on the node and add pins to add more LODs. */
+	UPROPERTY(meta = (DisplayName = "Collection LOD 1"))
+	FManagedArrayCollection CollectionLod1;
+	/** LOD 2 input, right click on the node and add pins to add more LODs. */
+	UPROPERTY(meta = (DisplayName = "Collection LOD 2"))
+	FManagedArrayCollection CollectionLod2;
+	/** LOD 3 input, right click on the node and add pins to add more LODs. */
+	UPROPERTY(meta = (DisplayName = "Collection LOD 3"))
+	FManagedArrayCollection CollectionLod3;
+	/** LOD 4 input, right click on the node and add pins to add more LODs. */
+	UPROPERTY(meta = (DisplayName = "Collection LOD 4"))
+	FManagedArrayCollection CollectionLod4;
+	/** LOD 5 input, right click on the node and add pins to add more LODs. */
+	UPROPERTY(meta = (DisplayName = "Collection LOD 5"))
+	FManagedArrayCollection CollectionLod5;
+	/** The number of LODs currently exposed to the node UI. */
+	UPROPERTY()
+	int32 NumLods = 1;
 
 	FChaosClothAssetTerminalNode(const Dataflow::FNodeParameters& InParam, FGuid InGuid = FGuid::NewGuid());
 
 	virtual void SetAssetValue(TObjectPtr<UObject> Asset, Dataflow::FContext& Context) const override;
 
 	virtual void Evaluate(Dataflow::FContext& Context, const FDataflowOutput* Out) const override {}
+
+	virtual Dataflow::FPin AddPin() override;
+	virtual bool CanAddPin() const override { return NumLods < MaxLods; }
+	virtual bool CanRemovePin() const override { return NumLods > 1; }
+	virtual Dataflow::FPin RemovePin() override;
+
+private:
+	TArray<const FManagedArrayCollection*> GetCollectionLods() const;
 };
diff --git a/Engine/Plugins/Experimental/ChaosClothAssetEditor/Source/ChaosClothAssetEditor/Private/ChaosClothAsset/SClothCollectionOutliner.cpp b/Engine/Plugins/Experimental/ChaosClothAssetEditor/Source/ChaosClothAssetEditor/Private/ChaosClothAsset/SClothCollectionOutliner.cpp
--- a/Engine/Plugins/Experimental/ChaosClothAssetEditor/Source/ChaosClothAssetEditor/Private/ChaosClothAsset/SClothCollectionOutliner.cpp
+++ b/Engine/Plugins/Experimental/ChaosClothAssetEditor/Source/ChaosClothAssetEditor/Private/ChaosClothAsset/SClothCollectionOutliner.cpp
@@ -1,7 +1,7 @@
 // Copyright Epic Games, Inc. All Rights Reserved.
 
 #include "ChaosClothAsset/SClothCollectionOutliner.h"
-#include "ChaosClothAsset/ClothCollection.h"
+#include "GeometryCollection/ManagedArrayCollection.h"
 #include "Templates/EnableIf.h"
 
 #define LOCTEXT_NAMESPACE "ClothCollectionOutliner"
@@ -77,7 +77,7 @@ void SClothCollectionOutliner::Construct(const FArguments& InArgs)
 	];
 }
 
-void SClothCollectionOutliner::SetClothCollection(TWeakPtr<UE::Chaos::ClothAsset::FClothCollection> InClothCollection)
+void SClothCollectionOutliner::SetClothCollection(TWeakPtr<FManagedArrayCollection> InClothCollection)
 {
 	ClothCollection = InClothCollection;
 }
@@ -99,7 +99,7 @@ void SClothCollectionOutliner::RegenerateHeader()
 {
 	constexpr float CustomFillWidth = 2.0f;
 
-	TSharedPtr<UE::Chaos::ClothAsset::FClothCollection> PinnedClothCollection = ClothCollection.Pin();
+	TSharedPtr<FManagedArrayCollection> PinnedClothCollection = ClothCollection.Pin();
 
 	if (!PinnedClothCollection.IsValid())
 	{
@@ -142,6 +142,11 @@ namespace ClothCollectionOutlinerHelpers
 		return FString::FromInt(Value);
 	}
 
+	FString AttributeValueToString(const FString& Value)
+	{
+		return Value;
+	}
+
 	template<typename T>
 	FString AttributeValueToString(const TArray<T>& Array)
 	{
@@ -159,7 +164,7 @@ namespace ClothCollectionOutlinerHelpers
 	}
 
 	template<typename T>
-	FString AttributeValueToString(const UE::Chaos::ClothAsset::FClothCollection& ClothCollection, const FName& AttributeName, const FName& GroupName, int32 AttributeArrayIndex)
+	FString AttributeValueToString(const FManagedArrayCollection& ClothCollection, const FName& AttributeName, const FName& GroupName, int32 AttributeArrayIndex)
 	{
 		const TManagedArray<T>* const Array = ClothCollection.FindAttributeTyped<T>(AttributeName, GroupName);
 		if (Array == nullptr)
@@ -170,7 +175,7 @@ namespace ClothCollectionOutlinerHelpers
 		return AttributeValueToString((*Array)[AttributeArrayIndex]);
 	}
 
-	FString AttributeValueToString(const UE::Chaos::ClothAsset::FClothCollection& ClothCollection, const FName& AttributeName, const FName& GroupName, int32 AttributeArrayIndex)
+	FString AttributeValueToString(const FManagedArrayCollection& ClothCollection, const FName& AttributeName, const FName& GroupName, int32 AttributeArrayIndex)
 	{
 		const FManagedArrayCollection::EArrayType ArrayType = ClothCollection.GetAttributeType(AttributeName, GroupName);
 
@@ -188,7 +193,7 @@ namespace ClothCollectionOutlinerHelpers
 			ValueAsString = AttributeValueToString<float>(ClothCollection, AttributeName, GroupName, AttributeArrayIndex);
 			break;
 		case FManagedArrayCollection::EArrayType::FIntVectorType:
-			ValueAsString = AttributeValueToString<FIntVector>(ClothCollection, AttributeName, GroupName, AttributeArrayIndex);
+			ValueAsString = AttributeValueToString<FIntVector3>(ClothCollection, AttributeName, GroupName, AttributeArrayIndex);
 			break;
 		case FManagedArrayCollection::EArrayType::FVector2DArrayType:
 			ValueAsString = AttributeValueToString<TArray<FVector2f>>(ClothCollection, AttributeName, GroupName, AttributeArrayIndex);
@@ -205,6 +210,9 @@ namespace ClothCollectionOutlinerHelpers
 		case FManagedArrayCollection::EArrayType::FFloatArrayType:
 			ValueAsString = AttributeValueToString<TArray<float>>(ClothCollection, AttributeName, GroupName, AttributeArrayIndex);
 			break;
+		case FManagedArrayCollection::EArrayType::FStringType:
+			ValueAsString = AttributeValueToString<FString>(ClothCollection, AttributeName, GroupName, AttributeArrayIndex);
+			break;
 		default:
 			ensure(false);
 			ValueAsString = "(Unknown Data Type)";
@@ -218,7 +226,7 @@ void SClothCollectionOutliner::RepopulateListView()
 {
 	ListItems.Empty();
 
-	const TSharedPtr<const UE::Chaos::ClothAsset::FClothCollection> PinnedClothCollection = ClothCollection.Pin();
+	const TSharedPtr<const FManagedArrayCollection> PinnedClothCollection = ClothCollection.Pin();
 
 	if (!PinnedClothCollection.IsValid())
 	{
diff --git a/Engine/Plugins/Experimental/ChaosClothAssetEditor/Source/ChaosClothAssetEditor/Private/ChaosClothAsset/SClothCollectionOutliner.h b/Engine/Plugins/Experimental/ChaosClothAssetEditor/Source/ChaosClothAssetEditor/Private/ChaosClothAsset/SClothCollectionOutliner.h
--- a/Engine/Plugins/Experimental/ChaosClothAssetEditor/Source/ChaosClothAssetEditor/Private/ChaosClothAsset/SClothCollectionOutliner.h
+++ b/Engine/Plugins/Experimental/ChaosClothAssetEditor/Source/ChaosClothAssetEditor/Private/ChaosClothAsset/SClothCollectionOutliner.h
@@ -5,7 +5,7 @@
 #include "Widgets/SCompoundWidget.h"
 #include "Widgets/Views/SListView.h"
 
-namespace UE::Chaos::ClothAsset{ class FClothCollection; }
+struct FManagedArrayCollection;
 
 struct FClothCollectionHeaderData
 {
@@ -45,13 +45,13 @@ class SClothCollectionOutliner : public SCompoundWidget
 public:
 
 	SLATE_BEGIN_ARGS(SClothCollectionOutliner) {}
-		SLATE_ARGUMENT(TWeakPtr<UE::Chaos::ClothAsset::FClothCollection>, ClothCollection)
+		SLATE_ARGUMENT(TWeakPtr<FManagedArrayCollection>, ClothCollection)
 		SLATE_ARGUMENT(FName, SelectedGroupName)
 	SLATE_END_ARGS()
 
 	void Construct(const FArguments& InArgs);
 
-	void SetClothCollection(TWeakPtr<UE::Chaos::ClothAsset::FClothCollection> ClothCollection);
+	void SetClothCollection(TWeakPtr<FManagedArrayCollection> ClothCollection);
 
 	///  Changing the selected group name will also automatically rebuild the header and all rows
 	void SetSelectedGroupName(const FName& SelectedGroupName);
@@ -62,7 +62,7 @@ public:
 
 private:
 
-	TWeakPtr<UE::Chaos::ClothAsset::FClothCollection> ClothCollection;
+	TWeakPtr<FManagedArrayCollection> ClothCollection;
 	FName SelectedGroupName;
 
 	TSharedPtr<SListView<TSharedPtr<const FClothCollectionItem>>> ListView;
diff --git a/Engine/Plugins/Experimental/ChaosClothAssetEditor/Source/ChaosClothAssetEditorTools/Private/ChaosClothAsset/ClothTransferSkinWeightsTool.cpp b/Engine/Plugins/Experimental/ChaosClothAssetEditor/Source/ChaosClothAssetEditorTools/Private/ChaosClothAsset/ClothTransferSkinWeightsTool.cpp
--- a/Engine/Plugins/Experimental/ChaosClothAssetEditor/Source/ChaosClothAssetEditorTools/Private/ChaosClothAsset/ClothTransferSkinWeightsTool.cpp
+++ b/Engine/Plugins/Experimental/ChaosClothAssetEditor/Source/ChaosClothAssetEditorTools/Private/ChaosClothAsset/ClothTransferSkinWeightsTool.cpp
@@ -5,8 +5,7 @@
 #include "ChaosClothAsset/ClothComponent.h"
 #include "ChaosClothAsset/ClothComponentToolTarget.h"
 #include "ChaosClothAsset/ClothPatternToDynamicMesh.h"
-#include "ChaosClothAsset/ClothAdapter.h"
-#include "ChaosClothAsset/ClothCollection.h"
+#include "ChaosClothAsset/CollectionClothFacade.h"
 
 #include "BoneWeights.h"
 #include "SkeletalMeshAttributes.h"
@@ -74,21 +73,23 @@ namespace ClothTransferSkinWeightsToolHelpers
 
 	void ClothComponentToDynamicMesh(const UChaosClothComponent* ClothComponent, UE::Geometry::FDynamicMesh3& MeshOut)
 	{
+		using namespace UE::Chaos::ClothAsset;
+
 		const UChaosClothAsset* ChaosClothAsset = ClothComponent->GetClothAsset();
 		if (!ChaosClothAsset)
 		{
 			return;
 		}
 
-		const UE::Chaos::ClothAsset::FClothConstAdapter ClothAdapter(ChaosClothAsset->GetClothCollection());
+		const FCollectionClothConstFacade ClothFacade(ChaosClothAsset->GetClothCollection());
 		constexpr int32 LodIndex = 0;
-		const UE::Chaos::ClothAsset::FClothLodConstAdapter ClothLodAdapter = ClothAdapter.GetLod(LodIndex);
+		const FCollectionClothLodConstFacade ClothLodFacade = ClothFacade.GetLod(LodIndex);
 		constexpr bool bGet2DPattern = false;
 
 		UE::Geometry::FDynamicMeshEditor MeshEditor(&MeshOut);
 		FClothPatternToDynamicMesh Converter;
 
-		for (int32 PatternIndex = 0; PatternIndex < ClothLodAdapter.GetNumPatterns(); ++PatternIndex)
+		for (int32 PatternIndex = 0; PatternIndex < ClothLodFacade.GetNumPatterns(); ++PatternIndex)
 		{
 			UE::Geometry::FDynamicMesh3 PatternMesh;
 			Converter.Convert(ChaosClothAsset, LodIndex, PatternIndex, bGet2DPattern, PatternMesh);
@@ -353,6 +354,7 @@ void UClothTransferSkinWeightsTool::TransferWeights()
 {
 	using namespace UE::AnimationCore;
 	using namespace UE::Geometry;
+	using namespace UE::Chaos::ClothAsset;
 	
 	//TODO: for now, assume we are always transfering from LOD 0, but make this a parameter in the future...
 	constexpr int32 SourceLODIdx = 0; 
@@ -387,28 +389,28 @@ void UClothTransferSkinWeightsTool::TransferWeights()
 		return;
 	}
 
-	UE::Chaos::ClothAsset::FClothAdapter ClothAdapter(TargetClothAsset->GetClothCollection());
+	FCollectionClothFacade ClothFacade(TargetClothAsset->GetClothCollection());
     FTransformSRT3d TargetToWorld = ClothComponent->GetComponentTransform();
 
 	// Iterate over the LODs and transfer the bone weights from the source Skeletal mesh to the Cloth asset
-	for (int TargetLODIdx = 0; TargetLODIdx < ClothAdapter.GetNumLods(); ++TargetLODIdx) 
+	for (int TargetLODIdx = 0; TargetLODIdx < ClothFacade.GetNumLods(); ++TargetLODIdx) 
 	{
-		UE::Chaos::ClothAsset::FClothLodAdapter ClothLodAdapter = ClothAdapter.GetLod(TargetLODIdx);
+		FCollectionClothLodFacade ClothLodFacade = ClothFacade.GetLod(TargetLODIdx);
 
 		// Cloth collection data arrays we are writing to
-		TArrayView<int32> SimNumBoneInfluences = ClothLodAdapter.GetPatternsSimNumBoneInfluences();
-		TArrayView<TArray<int32>> SimBoneIndices = ClothLodAdapter.GetPatternsSimBoneIndices();
-		TArrayView<TArray<float>> SimBoneWeights = ClothLodAdapter.GetPatternsSimBoneWeights();
+		TArrayView<int32> SimNumBoneInfluences = ClothLodFacade.GetSimNumBoneInfluences();
+		TArrayView<TArray<int32>> SimBoneIndices = ClothLodFacade.GetSimBoneIndices();
+		TArrayView<TArray<float>> SimBoneWeights = ClothLodFacade.GetSimBoneWeights();
 
-		TArrayView<int32> RenderNumBoneInfluences = ClothLodAdapter.GetPatternsRenderNumBoneInfluences();
-		TArrayView<TArray<int32>> RenderBoneIndices = ClothLodAdapter.GetPatternsRenderBoneIndices();
-		TArrayView<TArray<float>> RenderBoneWeights = ClothLodAdapter.GetPatternsRenderBoneWeights();
+		TArrayView<int32> RenderNumBoneInfluences = ClothLodFacade.GetRenderNumBoneInfluences();
+		TArrayView<TArray<int32>> RenderBoneIndices = ClothLodFacade.GetRenderBoneIndices();
+		TArrayView<TArray<float>> RenderBoneWeights = ClothLodFacade.GetRenderBoneWeights();
 
-		const TArrayView<FVector3f> SimPositions =  ClothLodAdapter.GetPatternsSimRestPosition();
+		const TArrayView<FVector3f> SimPositions =  ClothLodFacade.GetSimRestPosition();
 		
 		checkSlow(SimPositions.Num() == SimBoneIndices.Num());
 		
-		const int32 NumVert = ClothLodAdapter.GetPatternsNumSimVertices();
+		const int32 NumVert = ClothLodFacade.GetNumSimVertices();
 		constexpr bool bUseParallel = true; 
 
 		// Iterate over each vertex and write the data from FBoneWeights into cloth collection managed arrays
diff --git a/Engine/Plugins/Experimental/ChaosClothAssetEditor/Source/ChaosClothAssetTools/Private/ChaosClothAsset/ClothAssetBuilderEditor.cpp b/Engine/Plugins/Experimental/ChaosClothAssetEditor/Source/ChaosClothAssetTools/Private/ChaosClothAsset/ClothAssetBuilderEditor.cpp
--- a/Engine/Plugins/Experimental/ChaosClothAssetEditor/Source/ChaosClothAssetTools/Private/ChaosClothAsset/ClothAssetBuilderEditor.cpp
+++ b/Engine/Plugins/Experimental/ChaosClothAssetEditor/Source/ChaosClothAssetTools/Private/ChaosClothAsset/ClothAssetBuilderEditor.cpp
@@ -2,8 +2,9 @@
 
 #include "ChaosClothAsset/ClothAssetBuilderEditor.h"
 #include "ChaosClothAsset/ClothAsset.h"
-#include "ChaosClothAsset/ClothCollection.h"
+#include "ChaosClothAsset/CollectionClothFacade.h"
 #include "Engine/SkeletalMesh.h"
+#include "Engine/SkinnedAssetCommon.h"
 #include "Modules/ModuleManager.h"
 #include "Rendering/SkeletalMeshLODModel.h"
 #include "Utils/ClothingMeshUtils.h"
@@ -15,6 +16,8 @@
 
 void UClothAssetBuilderEditor::BuildLod(FSkeletalMeshLODModel& LODModel, const UChaosClothAsset& ClothAsset, int32 LodIndex) const
 {
+	using namespace UE::Chaos::ClothAsset;
+
 	// Start from an empty LODModel
 	LODModel.Empty();
 
@@ -28,42 +31,49 @@ void UClothAssetBuilderEditor::BuildLod(FSkeletalMeshLODModel& LODModel, const U
 	// Init the size of the vertex buffer
 	LODModel.NumVertices = 0;
 
-	// Build the section/faces map from the LOD patterns
-	TMap<int32, TArray<int32>> SectionFacesMap;
-	SectionFacesMap.Reserve(ClothAsset.GetMaterials().Num());
+	// Create a table to remap the LOD materials to the asset materials
+	const TArray<FSkeletalMaterial>& Materials = ClothAsset.GetMaterials();
 
-	TArray<uint32> LodRenderIndexRemap;
+	const TSharedPtr<const FManagedArrayCollection> ClothCollection = ClothAsset.GetClothCollection();
 
-	const TSharedPtr<const UE::Chaos::ClothAsset::FClothCollection> ClothCollection = ClothAsset.GetClothCollection();
-	const int32 PatternStart = ClothCollection->PatternStart[LodIndex];
-	const int32 PatternEnd = ClothCollection->PatternEnd[LodIndex];
-	if (PatternStart != INDEX_NONE && PatternEnd != INDEX_NONE)
-	{
-		for (int32 PatternIndex = PatternStart; PatternIndex <= PatternEnd; ++PatternIndex)
-		{
-			const int32 RenderFacesStart = ClothCollection->RenderFacesStart[PatternIndex];
-			const int32 RenderFacesEnd = ClothCollection->RenderFacesEnd[PatternIndex];
+	const FCollectionClothConstFacade ClothFacade(ClothCollection);
+	check(LodIndex < ClothFacade.GetNumLods())
+	const FCollectionClothLodConstFacade ClothLodFacade = ClothFacade.GetLod(LodIndex);
 
-			if (RenderFacesStart != INDEX_NONE && RenderFacesEnd != INDEX_NONE)
+	const int32 NumLodMaterials = ClothLodFacade.GetNumMaterials();
+
+	TArray<int32> LodMaterialMap;
+	LodMaterialMap.AddUninitialized(NumLodMaterials);
+	for (int32 LodMaterialIndex = 0; LodMaterialIndex < NumLodMaterials; ++LodMaterialIndex)
+	{
+		const FString& RenderMaterialPathName = ClothLodFacade.GetRenderMaterialPathName()[LodMaterialIndex];
+		const int32 SkeletalMaterialIndex = Materials.IndexOfByPredicate([&RenderMaterialPathName](const FSkeletalMaterial& SkeletalMaterial)
 			{
-				for (int32 RenderFaceIndex = RenderFacesStart; RenderFaceIndex <= RenderFacesEnd; ++RenderFaceIndex)
-				{
-					const int32 MaterialIndex = ClothCollection->RenderMaterialIndex[RenderFaceIndex];
+				return SkeletalMaterial.MaterialInterface && SkeletalMaterial.MaterialInterface->GetPathName() == RenderMaterialPathName;
+			});
+		LodMaterialMap[LodMaterialIndex] = SkeletalMaterialIndex;
+	}
 
-					TArray<int32>& SectionFaces = SectionFacesMap.FindOrAdd(MaterialIndex);
-					SectionFaces.Add(RenderFaceIndex);
-				}
-			}
-		}
+	// Build the section/faces map from the LOD patterns
+	TMap<int32, TArray<int32>> SectionFacesMap;
+	SectionFacesMap.Reserve(Materials.Num());
 
-		// Initialize the remapping array to start at LodRenderVerticesStart with the maximum number of vertices used in this LOD
-		const int32 LodRenderVerticesStart = ClothCollection->RenderVerticesStart[PatternStart];
-		const int32 LodRenderVerticesEnd = ClothCollection->RenderVerticesEnd[PatternEnd];
-		const int32 NumLodRenderVertices = LodRenderVerticesEnd - LodRenderVerticesStart + 1;
+	TArray<uint32> LodRenderIndexRemap;
+
+	const int32 NumRenderFaces = ClothLodFacade.GetNumRenderFaces();
+	const TConstArrayView<int32> LodRenderMaterialIndex = ClothLodFacade.GetRenderMaterialIndex();
 
-		LodRenderIndexRemap.SetNumUninitialized(NumLodRenderVertices);
+	for (int32 RenderFaceIndex = 0; RenderFaceIndex < NumRenderFaces; ++RenderFaceIndex)
+	{
+		const int32 RenderMaterialIndex = LodMaterialMap[LodRenderMaterialIndex[RenderFaceIndex]];
+
+		TArray<int32>& SectionFaces = SectionFacesMap.FindOrAdd(RenderMaterialIndex);
+		SectionFaces.Add(RenderFaceIndex);
 	}
 
+	const int32 NumRenderVertices = ClothLodFacade.GetNumRenderVertices();
+	LodRenderIndexRemap.SetNumUninitialized(NumRenderVertices);
+
 	// Keep track of the active bone indices for this LOD model
 	TSet<FBoneIndexType> ActiveBoneIndices;
 	ActiveBoneIndices.Reserve(ClothAsset.RefSkeleton.GetNum());
@@ -114,9 +124,12 @@ void UClothAssetBuilderEditor::BuildLod(FSkeletalMeshLODModel& LODModel, const U
 		TArray<uint32> Indices;
 		TSet<int32> UniqueIndicesSet;
 		Indices.SetNumUninitialized(NumIndices);
+
+		const TConstArrayView<FIntVector3> LodRenderIndices = ClothLodFacade.GetRenderIndices();
+
 		for (int32 FaceIndex = 0; FaceIndex < NumFaces; ++FaceIndex)
 		{
-			const FIntVector3& RenderIndices = ClothCollection->RenderIndices[Faces[FaceIndex]];
+			const FIntVector3& RenderIndices = LodRenderIndices[Faces[FaceIndex]];
 			for (int32 VertexIndex = 0; VertexIndex < 3; ++VertexIndex)
 			{
 				const int32 RenderIndex = RenderIndices[VertexIndex];
@@ -144,6 +157,16 @@ void UClothAssetBuilderEditor::BuildLod(FSkeletalMeshLODModel& LODModel, const U
 		// Track how many bones we added to the section's bone map so far
 		int CurSectionBoneMapNum = 0; 
 
+		const TConstArrayView<FVector3f> LodRenderPosition = ClothLodFacade.GetRenderPosition();
+		const TConstArrayView<FVector3f> LodRenderTangentU = ClothLodFacade.GetRenderTangentU();
+		const TConstArrayView<FVector3f> LodRenderTangentV = ClothLodFacade.GetRenderTangentV();
+		const TConstArrayView<FVector3f> LodRenderNormal = ClothLodFacade.GetRenderNormal();
+		const TConstArrayView<FLinearColor> LodRenderColor = ClothLodFacade.GetRenderColor();
+		const TConstArrayView<TArray<FVector2f>> LodRenderUVs = ClothLodFacade.GetRenderUVs();
+		const TConstArrayView<int32> LodRenderNumBoneInfluences = ClothLodFacade.GetRenderNumBoneInfluences();
+		const TConstArrayView<TArray<int32>> LodRenderBoneIndices = ClothLodFacade.GetRenderBoneIndices();
+		const TConstArrayView<TArray<float>> LodRenderBoneWeights = ClothLodFacade.GetRenderBoneWeights();
+
 		for (int32 VertexIndex = 0; VertexIndex < NumVertices; ++VertexIndex)
 		{
 			const int32 RenderIndex = UniqueIndices[VertexIndex];
@@ -153,27 +176,27 @@ void UClothAssetBuilderEditor::BuildLod(FSkeletalMeshLODModel& LODModel, const U
 
 			FSoftSkinVertex& SoftVertex = Section.SoftVertices[VertexIndex];
 
-			SoftVertex.Position = ClothCollection->RenderPosition[RenderIndex];
-			SoftVertex.TangentX = ClothCollection->RenderTangentU[RenderIndex];
-			SoftVertex.TangentY = ClothCollection->RenderTangentV[RenderIndex];
-			SoftVertex.TangentZ = ClothCollection->RenderNormal[RenderIndex];
+			SoftVertex.Position = LodRenderPosition[RenderIndex];
+			SoftVertex.TangentX = LodRenderTangentU[RenderIndex];
+			SoftVertex.TangentY = LodRenderTangentV[RenderIndex];
+			SoftVertex.TangentZ = LodRenderNormal[RenderIndex];
 
 			constexpr bool bSRGB = false; // Avoid linear to srgb conversion
-			SoftVertex.Color = ClothCollection->RenderColor[RenderIndex].ToFColor(bSRGB);
+			SoftVertex.Color = LodRenderColor[RenderIndex].ToFColor(bSRGB);
 
-			const TArray<FVector2f>& RenderUVs = ClothCollection->RenderUVs[RenderIndex];
+			const TArray<FVector2f>& RenderUVs = LodRenderUVs[RenderIndex];
 			for (int32 TexCoord = 0; TexCoord < FMath::Min(RenderUVs.Num(), (int32)MAX_TEXCOORDS); ++TexCoord)
 			{
 				SoftVertex.UVs[TexCoord] = RenderUVs[TexCoord];
 			}
 
-			const int32 NumInfluences = ClothCollection->RenderNumBoneInfluences[RenderIndex];
+			const int32 NumInfluences = LodRenderNumBoneInfluences[RenderIndex];
 			
 			// Add all of the bones that have non-zero influence to the section's bone map and keep track of the order
 			// that we added the reference bone via CurSectionBoneMapNum
 			for (int32 Influence = 0; Influence < NumInfluences; ++Influence)
 			{
-				const FBoneIndexType InfluenceBone = (FBoneIndexType)ClothCollection->RenderBoneIndices[RenderIndex][Influence];
+				const FBoneIndexType InfluenceBone = (FBoneIndexType)LodRenderBoneIndices[RenderIndex][Influence];
 
 				if (ReferenceToSectionBoneMap.Contains(InfluenceBone) == false)
 				{
@@ -185,8 +208,8 @@ void UClothAssetBuilderEditor::BuildLod(FSkeletalMeshLODModel& LODModel, const U
 			int32 Influence = 0;
 			for (;Influence < NumInfluences; ++Influence)
 			{
-				const FBoneIndexType InfluenceBone = (FBoneIndexType)ClothCollection->RenderBoneIndices[RenderIndex][Influence];
-				const float InWeight = ClothCollection->RenderBoneWeights[RenderIndex][Influence];
+				const FBoneIndexType InfluenceBone = (FBoneIndexType)LodRenderBoneIndices[RenderIndex][Influence];
+				const float InWeight = LodRenderBoneWeights[RenderIndex][Influence];
 				const uint16 InfluenceWeight = static_cast<uint16>(InWeight * static_cast<float>(UE::AnimationCore::MaxRawBoneWeight) + 0.5f);
 				
 				// FSoftSkinVertex::InfluenceBones contain indices into the section's bone map and not the reference
diff --git a/Engine/Plugins/Experimental/ChaosClothAssetEditor/Source/ChaosClothAssetTools/Private/ChaosClothAsset/ClothAssetFactory.cpp b/Engine/Plugins/Experimental/ChaosClothAssetEditor/Source/ChaosClothAssetTools/Private/ChaosClothAsset/ClothAssetFactory.cpp
--- a/Engine/Plugins/Experimental/ChaosClothAssetEditor/Source/ChaosClothAssetTools/Private/ChaosClothAsset/ClothAssetFactory.cpp
+++ b/Engine/Plugins/Experimental/ChaosClothAssetEditor/Source/ChaosClothAssetTools/Private/ChaosClothAsset/ClothAssetFactory.cpp
@@ -2,7 +2,7 @@
 
 #include "ClothAssetFactory.h"
 #include "ChaosClothAsset/ClothAsset.h"
-#include "ChaosClothAsset/ClothAdapter.h"
+#include "ChaosClothAsset/CollectionClothFacade.h"
 #include "Animation/Skeleton.h"
 
 #include UE_INLINE_GENERATED_CPP_BY_NAME(ClothAssetFactory)
@@ -23,24 +23,14 @@ UObject* UChaosClothAssetFactory::FactoryCreateNew(UClass* Class, UObject* Paren
 	ClothAsset->MarkPackageDirty();
 
 	// Add an empty default LOD, to avoid LOD mismatch with render data
-	TSharedPtr<FClothCollection> ClothCollection = ClothAsset->GetClothCollection();
+	TSharedPtr<FManagedArrayCollection> ClothCollection = ClothAsset->GetClothCollection();
 
-	FClothAdapter ClothAdapter(ClothCollection);
-	ClothAdapter.AddLod();
+	FCollectionClothFacade ClothFacade(ClothCollection);
+	ClothFacade.DefineSchema();
+	ClothFacade.AddLod();
 
-	// Set reference skeleton in both the collection and the asset
-	ClothCollection->AddElements(1, FClothCollection::SkeletonsGroup);
-	ClothCollection->SkeletonAssetPathName[0] = TEXT("/Engine/EditorMeshes/SkeletalMesh/DefaultSkeletalMesh_Skeleton.DefaultSkeletalMesh_Skeleton");
-
-	if (const USkeleton* const Skeleton = LoadObject<USkeleton>(nullptr, *ClothCollection->SkeletonAssetPathName[0], nullptr, LOAD_None, nullptr))
-	{
-		constexpr bool bRebuildClothSimulationModel = false;  // Avoid rebuilding the asset twice
-		ClothAsset->SetReferenceSkeleton(Skeleton->GetReferenceSkeleton(), bRebuildClothSimulationModel);
-	}
-
-	// Build static data
-	ClothAsset->Build();
+	// Set the default skeleton on this new LOD and rebuild the static data models (which is done by default with this override)
+	ClothAsset->SetSkeleton(nullptr);
 
 	return ClothAsset;
 }
-
diff --git a/Engine/Plugins/Experimental/ChaosClothAssetEditor/Source/ChaosClothAssetTools/Private/ChaosClothAsset/ClothPatternToDynamicMesh.cpp b/Engine/Plugins/Experimental/ChaosClothAssetEditor/Source/ChaosClothAssetTools/Private/ChaosClothAsset/ClothPatternToDynamicMesh.cpp
--- a/Engine/Plugins/Experimental/ChaosClothAssetEditor/Source/ChaosClothAssetTools/Private/ChaosClothAsset/ClothPatternToDynamicMesh.cpp
+++ b/Engine/Plugins/Experimental/ChaosClothAssetEditor/Source/ChaosClothAssetTools/Private/ChaosClothAsset/ClothPatternToDynamicMesh.cpp
@@ -5,10 +5,7 @@
 #include "ChaosClothAsset/ClothPatternToDynamicMesh.h"
 #include "ChaosClothAsset/ClothAsset.h"
 #include "ChaosClothAsset/ClothComponent.h"
-#include "ChaosClothAsset/ClothCollection.h"
-#include "ChaosClothAsset/ClothAdapter.h"
-#include "ChaosClothAsset/ClothLodAdapter.h"
-#include "ChaosClothAsset/ClothPatternAdapter.h"
+#include "ChaosClothAsset/CollectionClothFacade.h"
 #include "ToDynamicMesh.h"
 #include "DynamicMesh/MeshNormals.h"
 
@@ -27,43 +24,28 @@ public:
 	typedef int32 NormalIDType;
 	typedef int32 ColorIDType;
 
-	FClothSimPatternWrapper(const UE::Chaos::ClothAsset::FClothPatternConstAdapter& PatternIn, bool bGet2DPatternVertices) :
+	FClothSimPatternWrapper(const UE::Chaos::ClothAsset::FCollectionClothConstFacade& ClothFacade, int32 LodIndex, int32 PatternIndex, bool bGet2DPatternVertices) :
 		bGet2DPatternVertices(bGet2DPatternVertices),
-		Pattern(PatternIn),
-		Lod(PatternIn.GetClothCollection(), PatternIn.GetLodIndex())
+		Pattern(ClothFacade.GetLod(LodIndex).GetPattern(PatternIndex))
 	{
-		using UE::Chaos::ClothAsset::FClothCollection;
-
-		const int32 PatternIndex = Pattern.GetPatternIndex();
-
-		VertIDs.SetNum(Pattern.GetNumSimVertices());
-		for (int32 VtxIndex = 0; VtxIndex < Pattern.GetNumSimVertices(); ++VtxIndex)
+		const int32 NumSimVertices = Pattern.GetNumSimVertices();
+		VertIDs.SetNum(NumSimVertices);
+		for (int32 VtxIndex = 0; VtxIndex < NumSimVertices; ++VtxIndex)
 		{
-			const int32 PatternOffset = Lod.GetSimVerticesStart()[PatternIndex];
-			VertIDs[VtxIndex] = PatternOffset + VtxIndex;
+			VertIDs[VtxIndex] = VtxIndex;
 		}
 
-		TriIDs.SetNum(Pattern.GetNumSimFaces());
-		for (int32 TriIndex = 0; TriIndex < Pattern.GetNumSimFaces(); ++TriIndex)
+		const int32 NumSimFaces = Pattern.GetNumSimFaces();
+		TriIDs.SetNum(NumSimFaces);
+		for (int32 TriIndex = 0; TriIndex < NumSimFaces; ++TriIndex)
 		{
 			TriIDs[TriIndex] = TriIndex;
 		}
 
-
 		//
 		// Weight map layers precomputation
 		//
-		const FClothCollection* ClothCollection = Pattern.GetClothCollection().Get();
-		const TArray<FName> SimVerticesAttributeNames = ClothCollection->AttributeNames(FClothCollection::SimVerticesGroup);
-
-		// Of the SimVerticesGroup attributes, which ones refer to scalar weight maps?
-		for (const FName& AttributeName : SimVerticesAttributeNames)
-		{
-			if (ClothCollection->FindAttributeTyped<float>(AttributeName, FClothCollection::SimVerticesGroup))
-			{
-				WeightMapNames.Add(AttributeName);
-			}
-		}
+		WeightMapNames = ClothFacade.GetWeightMapNames();
 	}
 
 	int32 NumTris() const
@@ -93,14 +75,7 @@ public:
 
 	float GetVertexWeight(int32 LayerIndex, int32 VertexIndex) const
 	{
-		using UE::Chaos::ClothAsset::FClothCollection;
-
-		const FName& AttributeName = WeightMapNames[LayerIndex];
-		const FClothCollection* ClothCollection = Pattern.GetClothCollection().Get();
-		ensure(ClothCollection->FindAttributeTyped<float>(AttributeName, FClothCollection::SimVerticesGroup));
-		
-		const TManagedArray<float>& ManagedArray = ClothCollection->GetAttribute<float>(AttributeName, FClothCollection::SimVerticesGroup);
-		return ManagedArray[VertexIndex];
+		return Pattern.GetWeightMap(WeightMapNames[LayerIndex])[VertexIndex];
 	}
 
 	// --"Vertex Buffer" info
@@ -113,13 +88,13 @@ public:
 	{
 		if (bGet2DPatternVertices)
 		{
-			TConstArrayView<FVector2f> SimPositions = Lod.GetPatternsSimPosition();
+			const TConstArrayView<FVector2f> SimPositions = Pattern.GetSimPosition();
 			const FVector2f& Pos = SimPositions[VtxID];
 			return FVector3d(Pos[0], Pos[1], 0.0);
 		} 
 		else
 		{
-			TConstArrayView<FVector3f> RestPositions = Lod.GetPatternsSimRestPosition();
+			const TConstArrayView<FVector3f> RestPositions = Pattern.GetSimRestPosition();
 			const FVector3f& Pos = RestPositions[VtxID];
 			return FVector3d(Pos[0], Pos[1], Pos[2]);
 		}
@@ -133,10 +108,11 @@ public:
 
 	bool GetTri(TriIDType TriID, VertIDType& VID0, VertIDType& VID1, VertIDType& VID2) const
 	{
-		const FIntVector Face = Pattern.GetSimIndices()[TriID];
-		VID0 = Face[0];
-		VID1 = Face[1];
-		VID2 = Face[2];
+		const FIntVector Face = Pattern.GetSimIndices()[TriID];          // Indices are using an LOD based indexation
+		const int32 SimVerticesOffset = Pattern.GetSimVerticesOffset();  // and need to have the pattern offset removed
+		VID0 = Face[0] - SimVerticesOffset;
+		VID1 = Face[1] - SimVerticesOffset;
+		VID2 = Face[2] - SimVerticesOffset;
 
 		return true;
 	}
@@ -207,7 +183,7 @@ public:
 	}
 
 	const TArray<int32>& GetNormalIDs() const { return VertIDs; }
-	FVector3f GetNormal(NormalIDType ID) const { return Lod.GetPatternsSimRestNormal()[ID]; }
+	FVector3f GetNormal(NormalIDType ID) const { return Pattern.GetSimRestNormal()[ID]; }
 	bool GetNormalTri(const TriIDType& TriID, NormalIDType& ID0, NormalIDType& ID1, NormalIDType& ID2) const { return GetTri(TriID, ID0, ID1, ID2); }
 	
 	const TArray<int32>& GetUVIDs(int32 LayerID) const { return EmptyArray; }
@@ -230,13 +206,12 @@ private:
 
 	bool bGet2DPatternVertices = false;
 
-	TArray<TriIDType> TriIDs;		// indices into Pattern.GetSimIndices() 
-	TArray<VertIDType> VertIDs;		// indices into Lod.GetPatternsSimPosition(), etc.
+	TArray<TriIDType> TriIDs;		// indices into Pattern.GetSimFaces()
+	TArray<VertIDType> VertIDs;		// indices into Pattern.GetSimVertices()
 	
 	TArray<FName> WeightMapNames;
 
-	const UE::Chaos::ClothAsset::FClothPatternConstAdapter Pattern;
-	const UE::Chaos::ClothAsset::FClothLodConstAdapter Lod;
+	const UE::Chaos::ClothAsset::FCollectionClothPatternConstFacade Pattern;
 
 	TArray<int32> EmptyArray;
 };
@@ -244,16 +219,15 @@ private:
 
 void FClothPatternToDynamicMesh::Convert(const UChaosClothAsset* ClothAssetMeshIn, int32 LODIndex, int32 PatternIndex, bool bGet2DPattern, UE::Geometry::FDynamicMesh3& MeshOut)
 {
-	const TSharedPtr<const UE::Chaos::ClothAsset::FClothCollection> ClothCollection = ClothAssetMeshIn->GetClothCollection();
-	check(ClothCollection.IsValid());
-	const UE::Chaos::ClothAsset::FClothConstAdapter ClothAdapter(ClothCollection);
+	using namespace UE::Chaos::ClothAsset;
 
-	const UE::Chaos::ClothAsset::FClothLodConstAdapter ClothLodAdapter = ClothAdapter.GetLod(LODIndex);
-	const UE::Chaos::ClothAsset::FClothPatternConstAdapter ClothPatternAdapter = ClothLodAdapter.GetPattern(PatternIndex);
+	const TSharedPtr<const FManagedArrayCollection> ClothCollection = ClothAssetMeshIn->GetClothCollection();
+	check(ClothCollection.IsValid());
+	const FCollectionClothConstFacade ClothFacade(ClothCollection);
 
 	// Actual conversion
 	UE::Geometry::TToDynamicMesh<FClothSimPatternWrapper> PatternToDynamicMesh; 
-	FClothSimPatternWrapper PatternWrapper(ClothPatternAdapter, bGet2DPattern);
+	FClothSimPatternWrapper PatternWrapper(ClothFacade, LODIndex, PatternIndex, bGet2DPattern);
 
 	const bool bDisableAttributes = false;
 	auto TriangleToGroupFunction = [](FClothSimPatternWrapper::TriIDType) { return 0; };
diff --git a/Engine/Plugins/Experimental/ChaosClothAssetEditor/Source/ChaosClothAssetTools/Private/ChaosClothAsset/ClothingAssetToClothAssetExporter.cpp b/Engine/Plugins/Experimental/ChaosClothAssetEditor/Source/ChaosClothAssetTools/Private/ChaosClothAsset/ClothingAssetToClothAssetExporter.cpp
--- a/Engine/Plugins/Experimental/ChaosClothAssetEditor/Source/ChaosClothAssetTools/Private/ChaosClothAsset/ClothingAssetToClothAssetExporter.cpp
+++ b/Engine/Plugins/Experimental/ChaosClothAssetEditor/Source/ChaosClothAssetTools/Private/ChaosClothAsset/ClothingAssetToClothAssetExporter.cpp
@@ -2,7 +2,7 @@
 
 #include "ClothingAssetToClothAssetExporter.h"
 #include "ChaosClothAsset/ClothAsset.h"
-#include "ChaosClothAsset/ClothAdapter.h"
+#include "ChaosClothAsset/CollectionClothFacade.h"
 #include "Animation/Skeleton.h"
 #include "Engine/SkinnedAssetCommon.h"
 #include "Engine/SkeletalMesh.h"
@@ -34,37 +34,36 @@ void UClothingAssetToChaosClothAssetExporter::Export(const UClothingAssetBase* C
 
 	UChaosClothAsset* const ClothAsset = CastChecked<UChaosClothAsset>(ExportedAsset);
 	check(ClothAsset);
+	FCollectionClothFacade Cloth(ClothAsset->GetClothCollection());
 
+	// Create the LODs
 	for (const FClothLODDataCommon& ClothLODData : ClothingAssetCommon->LodData)
 	{
 		const FClothPhysicalMeshData& PhysicalMeshData = ClothLODData.PhysicalMeshData;
 
 		// Unwrap the physical mesh data into the pattern and rest meshes
-		FClothAdapter Cloth(ClothAsset->GetClothCollection());
-		FClothLodAdapter ClothLod = Cloth.AddGetLod();
+		FCollectionClothLodFacade ClothLod = Cloth.AddGetLod();
 		ClothLod.Initialize(PhysicalMeshData.Vertices, PhysicalMeshData.Indices);
 	}
 
-	// Add a default material
-	ClothAsset->GetMaterials().Reset(1);
-	UMaterialInterface* const DefaultMaterial = LoadObject<UMaterialInterface>(nullptr, TEXT("/Engine/EditorMaterials/Cloth/CameraLitDoubleSided.CameraLitDoubleSided"), nullptr, LOAD_None, nullptr);
-	const int32 MaterialId = ClothAsset->GetMaterials().Emplace(DefaultMaterial, true, false, DefaultMaterial->GetFName());
+	if (Cloth.GetNumLods())
+	{
+		// Set the render mesh to duplicate the sim mesh
+		ClothAsset->CopySimMeshToRenderMesh();
+	}
+	else
+	{
+		// Make sure that at least one empty LOD is always created
+		Cloth.AddLod();
+	}
 
-	// Assign the physics asset if any
+	// Assign the physics asset if any (must be done after having added the LODs)
 	ClothAsset->SetPhysicsAsset(ClothingAssetCommon->PhysicsAsset);
 
-	// Set the skeleton from the skeletal mesh, or create a default skeleton if it can't find one
-	USkeletalMesh* const SkeletalMesh = Cast<USkeletalMesh>(ClothingAssetCommon->GetOuter());
-
-	ClothAsset->SetReferenceSkeleton(SkeletalMesh ?
-		SkeletalMesh->GetRefSkeleton() :
-		LoadObject<USkeleton>(nullptr, TEXT("/Engine/EditorMeshes/SkeletalMesh/DefaultSkeletalMesh_Skeleton.DefaultSkeletalMesh_Skeleton"), nullptr, LOAD_None, nullptr)->GetReferenceSkeleton());
-
-	// Bind the mesh to the root bone.
-	ClothAsset->BindSimMeshToRootBone(); 
-
-	// Set the render mesh to duplicate the sim mesh
-	ClothAsset->CopySimMeshToRenderMesh(MaterialId);
+	// Set the skeleton from the skeletal mesh (must be done after having added the LODs)
+	constexpr bool bRebuildModels = false;  // Build is called last
+	USkeletalMesh* const SkeletalMesh = CastChecked<USkeletalMesh>(ClothingAssetCommon->GetOuter());
+	ClothAsset->SetSkeleton(SkeletalMesh->GetSkeleton(), bRebuildModels);
 
 	// Build the asset, since it is already loaded, it won't rebuild on load
 	ClothAsset->Build();
