CIS - Fixed static analysis in CollectionClothLodFacade.cpp on an ill defined for loop.

#rnx

diff --git a/Engine/Plugins/Experimental/ChaosClothAsset/Source/ChaosClothAsset/Private/ChaosClothAsset/CollectionClothLodFacade.cpp b/Engine/Plugins/Experimental/ChaosClothAsset/Source/ChaosClothAsset/Private/ChaosClothAsset/CollectionClothLodFacade.cpp
--- a/Engine/Plugins/Experimental/ChaosClothAsset/Source/ChaosClothAsset/Private/ChaosClothAsset/CollectionClothLodFacade.cpp
+++ b/Engine/Plugins/Experimental/ChaosClothAsset/Source/ChaosClothAsset/Private/ChaosClothAsset/CollectionClothLodFacade.cpp
@@ -904,7 +904,7 @@ namespace UE::Chaos::ClothAsset
 		// Tether Batches Group
 		const int32 NumTetherBatches = Other.GetNumPatterns();
 		SetNumTetherBatches(NumTetherBatches);
-		for (int32 TetherBatchIndex = 0; TetherBatchIndex < TetherBatchIndex; ++TetherBatchIndex)
+		for (int32 TetherBatchIndex = 0; TetherBatchIndex < NumTetherBatches; ++TetherBatchIndex)
 		{
 			// GetTetherBatch(TetherBatchIndex).Initialize(Other.GetTetherBatch(TetherBatchIndex));  // TODO: Tether Batches facade
 		}
