Chaos Cloth Dataflow: Add StaticMeshImportNode.

This will can be used to import a static mesh to start a new cloth asset. Optionally use the UVs to define the pattern data.

Internally using conversion from Mesh Description to FDynamicMesh3, and then importing the FDynamicMesh3. The existing FCollectionClothLodFacade::Initialize has been updated to also internally create an FDynamicMesh3 so the code can be shared. There were various bugs related to pattern edge and seam generation in the previous Initialize code (in particular it was possible for bowtie vertices to not be seamed back together, and no seams were generated between separate islands that were created as part of the unwrap process). Hopefully they've been fixed.

Of note, FDynamicMesh3 does not allow non-manifold edges. Edges are duplicated in the FDynamicMesh3 in the conversion process if a non-manifold edge is found. This split is reflected as a new Seam in the ClothAsset.

diff --git a/Engine/Plugins/Experimental/ChaosClothAsset/Source/ChaosClothAsset/ChaosClothAsset.Build.cs b/Engine/Plugins/Experimental/ChaosClothAsset/Source/ChaosClothAsset/ChaosClothAsset.Build.cs
--- a/Engine/Plugins/Experimental/ChaosClothAsset/Source/ChaosClothAsset/ChaosClothAsset.Build.cs
+++ b/Engine/Plugins/Experimental/ChaosClothAsset/Source/ChaosClothAsset/ChaosClothAsset.Build.cs
@@ -11,7 +11,9 @@ public class ChaosClothAsset : ModuleRules
 		PrivateDependencyModuleNames.AddRange(
 			new string[]
 			{
-				"Core"
+				"Core",
+				"GeometryCore",
+				"MeshConversion",
 			}
 		);
 		PublicDependencyModuleNames.AddRange(
diff --git a/Engine/Plugins/Experimental/ChaosClothAsset/Source/ChaosClothAsset/Private/ChaosClothAsset/CollectionClothLodFacade.cpp b/Engine/Plugins/Experimental/ChaosClothAsset/Source/ChaosClothAsset/Private/ChaosClothAsset/CollectionClothLodFacade.cpp
--- a/Engine/Plugins/Experimental/ChaosClothAsset/Source/ChaosClothAsset/Private/ChaosClothAsset/CollectionClothLodFacade.cpp
+++ b/Engine/Plugins/Experimental/ChaosClothAsset/Source/ChaosClothAsset/Private/ChaosClothAsset/CollectionClothLodFacade.cpp
@@ -3,17 +3,73 @@
 #include "ChaosClothAsset/CollectionClothLodFacade.h"
 #include "ChaosClothAsset/CollectionClothFacade.h"
 #include "ChaosClothAsset/ClothCollection.h"
+#include "DynamicMesh/DynamicMesh3.h"
+#include "DynamicMesh/DynamicMeshAttributeSet.h"
+#include "DynamicMesh/NonManifoldMappingSupport.h"
+#include "ToDynamicMesh.h"
+#include "Util/IndexUtil.h"
 
 // Utility functions to unwrap a 3d sim mesh into a tailored cloth
 namespace UE::Chaos::ClothAsset::Private
 {
+	struct FSimpleSrcMeshInterface
+	{
+ 		typedef int32     VertIDType;
+ 		typedef int32     TriIDType;
+ 
+		FSimpleSrcMeshInterface(const TArray<FVector3f>& InPositions, const TArray<uint32>& InIndices)
+			:Positions(InPositions), Indices(InIndices)
+		{
+			VertIDs.SetNumUninitialized(Positions.Num());
+			for (int32 VtxIndex = 0; VtxIndex < Positions.Num(); ++VtxIndex)
+			{
+				VertIDs[VtxIndex] = VtxIndex;
+			}
+			
+			check(Indices.Num() % 3 == 0);
+			const int32 NumFaces = Indices.Num() / 3;
+			TriIDs.SetNumUninitialized(NumFaces);
+			for (int32 TriIndex = 0; TriIndex < NumFaces; ++TriIndex)
+			{
+				TriIDs[TriIndex] = 3*TriIndex;
+			}
+		}
+
+		// accounting.
+		int32 NumTris() const { return TriIDs.Num(); }
+		int32 NumVerts() const { return VertIDs.Num(); }
+ 
+		// --"Vertex Buffer" info
+		const TArray<VertIDType>& GetVertIDs() const { return VertIDs; }
+		const FVector GetPosition(const VertIDType VtxID) const { return FVector(Positions[VtxID]); }
+ 
+		// --"Index Buffer" info
+		const TArray<TriIDType>& GetTriIDs() const { return TriIDs; }
+		// return false if this TriID is not contained in mesh.
+		bool GetTri(const TriIDType TriID, VertIDType& VID0, VertIDType& VID1, VertIDType& VID2) const
+		{
+			VID0 = Indices[TriID + 0];
+			VID1 = Indices[TriID + 1];
+			VID2 = Indices[TriID + 2];
+
+			return true;
+		}
+
+	private:
+		const TArray<FVector3f>& Positions;
+		const TArray<uint32>& Indices;
+
+		TArray<TriIDType> TriIDs; // TriID = first index in flat Indices array
+		TArray<VertIDType> VertIDs;
+	};
+
 	// Triangle islands to become patterns, although in this case all the seams are internal (same pattern)
 	struct FIsland
 	{
-		TArray<uint32> Indices;  // 3x number of triangles
+		TArray<int32> Indices;  // 3x number of triangles
 		TArray<FVector2f> Positions;
 		TArray<FVector3f> RestPositions;  // Same size as Positions
-		TArray<uint32> SourceIndices;  // Index in the original welded position array
+		TArray<int32> PositionToSourceIndex; // Same size as Positions. Index in the original welded position array
 	};
 
 	enum class EIntersectCirclesResult
@@ -63,70 +119,41 @@ namespace UE::Chaos::ClothAsset::Private
 
 		return EIntersectCirclesResult::DoubleIntersect;
 	}
-	;
-	static FUintVector2 MakeSortedUintVector2(uint32 Index0, uint32 Index1)
+;
+	static FIntVector2 MakeSortedIntVector2(int32 Index0, int32 Index1)
 	{
-		return Index0 < Index1 ? FUintVector2(Index0, Index1) : FUintVector2(Index1, Index0);
+		return Index0 < Index1 ? FIntVector2(Index0, Index1) : FIntVector2(Index1, Index0);
 	}
 
-	static void BuildEdgeMap(const TArray<uint32>& Indices, TMap<FUintVector2, TArray<uint32>>& OutEdgeToTriangles)
+	template<bool bWeldNearlyCoincidentVertices>
+	static void UnwrapDynamicMesh(const UE::Geometry::FDynamicMesh3& DynamicMesh, TArray<FIsland>& OutIslands)
 	{
-		ensure(Indices.Num() % 3 == 0);
-		const uint32 NumTriangles = (uint32)Indices.Num() / 3;
-		OutEdgeToTriangles.Empty(NumTriangles * 2);  // Rough estimate for the number of edges
-		if (!NumTriangles)
-		{
-			return;
-		}
-
-		for (uint32 Triangle = 0; Triangle < NumTriangles; ++Triangle)
-		{
-			const uint32 Index0 = Indices[Triangle * 3 + 0];
-			const uint32 Index1 = Indices[Triangle * 3 + 1];
-			const uint32 Index2 = Indices[Triangle * 3 + 2];
+		using namespace UE::Geometry;
 
-			const FUintVector2 Edge0 = MakeSortedUintVector2(Index0, Index1);
-			const FUintVector2 Edge1 = MakeSortedUintVector2(Index1, Index2);
-			const FUintVector2 Edge2 = MakeSortedUintVector2(Index2, Index0);
-
-			OutEdgeToTriangles.FindOrAdd(Edge0).Add(Triangle);
-			OutEdgeToTriangles.FindOrAdd(Edge1).Add(Triangle);
-			OutEdgeToTriangles.FindOrAdd(Edge2).Add(Triangle);
-		}
-	}
-
-	static void UnwrapMesh(const TArray<FVector3f>& Positions, const TArray<uint32>& Indices, TArray<FIsland>& OutIslands)
-	{
 		OutIslands.Reset();
+		constexpr float SquaredWeldingDistance = FMath::Square(0.01f);  // 0.1 mm
 
-		ensure(Indices.Num() % 3 == 0);
-		const uint32 NumTriangles = (uint32)Indices.Num() / 3;
-		if (!NumTriangles)
-		{
-			return;
-		}
-
-		// Gather edge information
-		TMap<FUintVector2, TArray<uint32>> EdgeToTrianglesMap;
-		BuildEdgeMap(Indices, EdgeToTrianglesMap);
-
-		// Build pattern islands
-		TSet<uint32> VisitedTriangles;
+		// Build pattern islands. 
+		const int32 NumTriangles = DynamicMesh.TriangleCount();
+		TSet<int32> VisitedTriangles;
 		VisitedTriangles.Reserve(NumTriangles);
 
-		constexpr float SquaredWeldingDistance = FMath::Square(0.01f);  // 0.1 mm
-
-		for (uint32 SeedTriangle = 0; SeedTriangle < NumTriangles; ++SeedTriangle)
+		for (int32 SeedTriangle : DynamicMesh.TriangleIndicesItr())
 		{
 			if (VisitedTriangles.Contains(SeedTriangle))
 			{
 				continue;
 			}
+			const FIndex3i TriangleIndices = DynamicMesh.GetTriangle(SeedTriangle);
 
-			const uint32 SeedIndex0 = Indices[SeedTriangle * 3 + 0];
-			const uint32 SeedIndex1 = Indices[SeedTriangle * 3 + 1];
+			const int32 SeedIndex0 = TriangleIndices[0];
+			const int32 SeedIndex1 = TriangleIndices[1];
 
-			if (FVector3f::DistSquared(Positions[SeedIndex0], Positions[SeedIndex1]) <= SquaredWeldingDistance)
+			const FVector3f Position0(DynamicMesh.GetVertex(SeedIndex0));
+			const FVector3f Position1(DynamicMesh.GetVertex(SeedIndex1));
+			const float Position01DistSq = FVector3f::DistSquared(Position0, Position1);
+
+			if (Position01DistSq <= SquaredWeldingDistance)
 			{
 				continue;  // A degenerated triangle edge is not a good start
 			}
@@ -134,24 +161,26 @@ namespace UE::Chaos::ClothAsset::Private
 			// Setup first visitor from seed, and add the first two points
 			FIsland& Island = OutIslands.AddDefaulted_GetRef();
 
-			Island.RestPositions.Add(Positions[SeedIndex0]);
-			Island.RestPositions.Add(Positions[SeedIndex1]);
+			Island.RestPositions.Add(Position0);
+			Island.RestPositions.Add(Position1);
+			Island.PositionToSourceIndex.Add(SeedIndex0);
+			Island.PositionToSourceIndex.Add(SeedIndex1);
 
-			const uint32 SeedIndex2D0 = Island.Positions.Add(FVector2f::ZeroVector);
-			const uint32 SeedIndex2D1 = Island.Positions.Add(FVector2f(FVector3f::Dist(Positions[SeedIndex0], Positions[SeedIndex1]), 0.f));
+			const int32 SeedIndex2D0 = Island.Positions.Add(FVector2f::ZeroVector);
+			const int32 SeedIndex2D1 = Island.Positions.Add(FVector2f(FMath::Sqrt(Position01DistSq), 0.f));
 
 			struct FVisitor
 			{
-				uint32 Triangle;
-				FUintVector2 OldEdge;
-				FUintVector2 NewEdge;
-				uint32 CrossEdgePoint;  // Keep the opposite point to orientate degenerate cases
+				int32 Triangle;
+				FIndex2i OldEdge;
+				FIndex2i NewEdge;
+				int32 CrossEdgePoint;  // Keep the opposite point to orientate degenerate cases
 			} Visitor =
 			{
 				SeedTriangle,
-				FUintVector2(SeedIndex0, SeedIndex1),
-				FUintVector2(SeedIndex2D0, SeedIndex2D1),
-				(uint32)INDEX_NONE
+				FIndex2i(SeedIndex0, SeedIndex1),
+				FIndex2i(SeedIndex2D0, SeedIndex2D1),
+				INDEX_NONE
 			};
 
 			VisitedTriangles.Add(SeedTriangle);
@@ -159,31 +188,26 @@ namespace UE::Chaos::ClothAsset::Private
 			TQueue<FVisitor> Visitors;
 			do
 			{
-				const uint32 Triangle = Visitor.Triangle;
-				const uint32 CrossEdgePoint = Visitor.CrossEdgePoint;
-				const uint32 OldIndex0 = Visitor.OldEdge.X;
-				const uint32 OldIndex1 = Visitor.OldEdge.Y;
-				const uint32 NewIndex0 = Visitor.NewEdge.X;
-				const uint32 NewIndex1 = Visitor.NewEdge.Y;
+				const int32 Triangle = Visitor.Triangle;
+				const int32 CrossEdgePoint = Visitor.CrossEdgePoint;
+				const int32 OldIndex0 = Visitor.OldEdge.A;
+				const int32 OldIndex1 = Visitor.OldEdge.B;
+				const int32 NewIndex0 = Visitor.NewEdge.A;
+				const int32 NewIndex1 = Visitor.NewEdge.B;
 
 				// Find opposite index from this triangle edge
-				const uint32 TriangleIndex0 = Indices[Triangle * 3 + 0];
-				const uint32 TriangleIndex1 = Indices[Triangle * 3 + 1];
-				const uint32 TriangleIndex2 = Indices[Triangle * 3 + 2];
 
-				const uint32 OldIndex2 =
-					(OldIndex0 != TriangleIndex0 && OldIndex1 != TriangleIndex0) ? TriangleIndex0 :
-					(OldIndex0 != TriangleIndex1 && OldIndex1 != TriangleIndex1) ? TriangleIndex1 : TriangleIndex2;
+				const int32 OldIndex2 = IndexUtil::FindTriOtherVtxUnsafe(OldIndex0, OldIndex1, DynamicMesh.GetTriangle(Triangle));
 
 				// Find the 2D intersection of the two connecting adjacent edges using the 3D reference length
-				const FVector3f& P0 = Positions[OldIndex0];
-				const FVector3f& P1 = Positions[OldIndex1];
-				const FVector3f& P2 = Positions[OldIndex2];
+				const FVector3f P0(DynamicMesh.GetVertexRef(OldIndex0));
+				const FVector3f P1(DynamicMesh.GetVertexRef(OldIndex1));
+				const FVector3f P2(DynamicMesh.GetVertexRef(OldIndex2));
 
 				const float R0 = FVector3f::Dist(P0, P2);
 				const float R1 = FVector3f::Dist(P1, P2);
-				const FVector2f C0 = Island.Positions[NewIndex0];
-				const FVector2f C1 = Island.Positions[NewIndex1];
+				const FVector2f& C0 = Island.Positions[NewIndex0];
+				const FVector2f& C1 = Island.Positions[NewIndex1];
 
 				FVector2f I0, I1;
 				const EIntersectCirclesResult IntersectCirclesResult = IntersectCircles(C0, R0, C1, R1, I0, I1);
@@ -208,20 +232,34 @@ namespace UE::Chaos::ClothAsset::Private
 				}
 
 				// Add the new position found for the opposite point
-				uint32 NewIndex2 = INDEX_NONE;
-				for (uint32 UsedIndex = 0; UsedIndex < (uint32)Island.Positions.Num(); ++UsedIndex)
+				int32 NewIndex2 = INDEX_NONE;
+				for (int32 UsedIndex = 0; UsedIndex < Island.Positions.Num(); ++UsedIndex)
 				{
-					if (FVector2f::DistSquared(Island.Positions[UsedIndex], C2) <= SquaredWeldingDistance &&
-						FVector3f::DistSquared(Island.RestPositions[UsedIndex], Positions[OldIndex2]) <= SquaredWeldingDistance)
+					if (bWeldNearlyCoincidentVertices)
 					{
-						NewIndex2 = UsedIndex;  // Both Rest and 2D positions match, reuse this index
-						break;
+						if (FVector2f::DistSquared(Island.Positions[UsedIndex], C2) <= SquaredWeldingDistance &&
+							FVector3f::DistSquared(Island.RestPositions[UsedIndex], P2) <= SquaredWeldingDistance)
+						{
+							NewIndex2 = UsedIndex;  // Both Rest and 2D positions match, reuse this index
+							break;
+						}
+					}
+					else
+					{
+						if(Island.PositionToSourceIndex[UsedIndex] == OldIndex2 &&
+							FVector2f::DistSquared(Island.Positions[UsedIndex], C2) <= SquaredWeldingDistance)
+						{
+							NewIndex2 = UsedIndex;  // Both Rest and 2D positions match, reuse this index
+							break;
+						}
 					}
 				}
+
 				if (NewIndex2 == INDEX_NONE)
 				{
 					NewIndex2 = Island.Positions.Add(C2);
-					Island.RestPositions.Add(Positions[OldIndex2]);
+					Island.RestPositions.Add(P2);
+					Island.PositionToSourceIndex.Add(OldIndex2);
 				}
 
 				// Add triangle to list of indices, unless it is degenerated to a segment
@@ -230,32 +268,29 @@ namespace UE::Chaos::ClothAsset::Private
 					Island.Indices.Add(NewIndex0);
 					Island.Indices.Add(NewIndex1);
 					Island.Indices.Add(NewIndex2);
-					Island.SourceIndices.Add(OldIndex0);
-					Island.SourceIndices.Add(OldIndex1);
-					Island.SourceIndices.Add(OldIndex2);
 				}
 
 				// Add neighbor triangles to the queue
-				const FUintVector2 OldEdgeList[3] =
+				const FIndex2i OldEdgeList[3] =
 				{
-					FUintVector2(OldIndex1, OldIndex0),  // Reversed as to keep the correct winding order
-					FUintVector2(OldIndex2, OldIndex1),
-					FUintVector2(OldIndex0, OldIndex2)
+					FIndex2i(OldIndex1, OldIndex0),  // Reversed as to keep the correct winding order
+					FIndex2i(OldIndex2, OldIndex1),
+					FIndex2i(OldIndex0, OldIndex2)
 				};
-				const FUintVector3 NewEdgeList[3] =
+				const FIndex3i NewEdgeList[3] =
 				{
-					FUintVector3(NewIndex1, NewIndex0, NewIndex2),  // Adds opposite point index
-					FUintVector3(NewIndex2, NewIndex1, NewIndex0),
-					FUintVector3(NewIndex0, NewIndex2, NewIndex1)
+					FIndex3i(NewIndex1, NewIndex0, NewIndex2),  // Adds opposite point index
+					FIndex3i(NewIndex2, NewIndex1, NewIndex0),
+					FIndex3i(NewIndex0, NewIndex2, NewIndex1)
 				};
 				for (int32 Edge = 0; Edge < 3; ++Edge)
 				{
-					const uint32 EdgeIndex0 = OldEdgeList[Edge].X;
-					const uint32 EdgeIndex1 = OldEdgeList[Edge].Y;
-
-					const TArray<uint32>& NeighborTriangles = EdgeToTrianglesMap.FindChecked(MakeSortedUintVector2(EdgeIndex0, EdgeIndex1));
+					const int32 EdgeIndex0 = OldEdgeList[Edge].A;
+					const int32 EdgeIndex1 = OldEdgeList[Edge].B;
 
-					for (const uint32 NeighborTriangle : NeighborTriangles)
+					const FIndex2i EdgeT = DynamicMesh.GetEdgeT(DynamicMesh.FindEdgeFromTri(EdgeIndex0, EdgeIndex1, Triangle));
+					const int32 NeighborTriangle = EdgeT.OtherElement(Triangle);
+					if (NeighborTriangle != IndexConstants::InvalidID)
 					{
 						if (!VisitedTriangles.Contains(NeighborTriangle))
 						{
@@ -266,126 +301,176 @@ namespace UE::Chaos::ClothAsset::Private
 							Visitors.Enqueue(FVisitor
 								{
 									NeighborTriangle,
-									FUintVector2(EdgeIndex0, EdgeIndex1),
-									FUintVector2(NewEdgeList[Edge].X, NewEdgeList[Edge].Y),
-									NewEdgeList[Edge].Z,  // Pass the cross edge 2D opposite point to help define orientation of any degenerated triangles
+									OldEdgeList[Edge],
+									FIndex2i(NewEdgeList[Edge].A, NewEdgeList[Edge].B),
+									NewEdgeList[Edge].C,  // Pass the cross edge 2D opposite point to help define orientation of any degenerated triangles
 								});
 						}
+
 					}
 				}
 			} while (Visitors.Dequeue(Visitor));
 		}
 	}
 
-	struct FSeam
+	static void BuildIslandsFromDynamicMeshUVs(const UE::Geometry::FDynamicMeshUVOverlay& UVOverlay, TArray<FIsland>& OutIslands)
 	{
-		TSet<FIntVector2> Stitches;
-		FIntVector2 Patterns;
-	};
+		using namespace UE::Geometry;
 
-	// Rebuild the seam information from the torn/unwrapped mesh islands data
-	// Note that the isolated mesh islands are not technically patterns despite being considered so, 
-	// since they aren't sewed together in the source welded mesh.
-	// The algorithm will have to be slightly modified to be used with provided UV panels.
-	static void BuildSeams(const TArray<FIsland>& Islands, TArray<FSeam>& OutSeams)
-	{
-		OutSeams.Reset();
+		const FDynamicMesh3* DynamicMesh = UVOverlay.GetParentMesh();
+		check(DynamicMesh);
 
-		for (int32 IslandIndex = 0; IslandIndex < Islands.Num(); ++IslandIndex)
+		OutIslands.Reset();
+
+		// Build pattern islands. 
+		const int32 NumTriangles = DynamicMesh->TriangleCount();
+		TSet<int32> VisitedTriangles;
+		VisitedTriangles.Reserve(NumTriangles);
+
+		// This is reused for each island, but only allocate once.
+		TArray<int32> SourceElementIndexToNewIndex;
+
+		for (int32 SeedTriangle : DynamicMesh->TriangleIndicesItr())
 		{
-			const FIsland& Island = Islands[IslandIndex];
+			if (VisitedTriangles.Contains(SeedTriangle))
+			{
+				continue;
+			}
+
+			// Setup first visitor from seed
+			FIsland& Island = OutIslands.AddDefaulted_GetRef();
 
-			FSeam Seam;
-			Seam.Patterns = FIntVector2(IslandIndex);  // Just patching an unwrap here, will only do pattern to same
+			SourceElementIndexToNewIndex.Init(INDEX_NONE, UVOverlay.MaxElementID());
+
+			struct FVisitor
+			{
+				int32 Triangle;
+			} Visitor =
+			{
+				SeedTriangle
+			};
 
-			// Gather edge information for the source mesh
-			TMap<FUintVector2, TArray<uint32>> SourceEdgeToTrianglesMap;
-			BuildEdgeMap(Island.SourceIndices, SourceEdgeToTrianglesMap);
+			VisitedTriangles.Add(SeedTriangle);
 
-			auto GetTriangleEdgeMatchingSourceEdge = [&Island](uint32 Triangle, const FUintVector2& SourceEdge) -> FUintVector2
+			TQueue<FVisitor> Visitors;
+			do
 			{
-				const uint32 TriangleBase = Triangle * 3;
-				const uint32 TriangleIndex0 = Island.SourceIndices[TriangleBase + 0];
-				const uint32 TriangleIndex1 = Island.SourceIndices[TriangleBase + 1];
-				const uint32 TriangleIndex2 = Island.SourceIndices[TriangleBase + 2];
+				const int32 Triangle = Visitor.Triangle;
+				const FIndex3i TriangleIndices = DynamicMesh->GetTriangle(Triangle);
+				const FIndex3i TriangleUVElements = UVOverlay.GetTriangle(Triangle);
 
-				if (SourceEdge[0] == TriangleIndex0)
+				auto GetOrAddNewIndex = [&UVOverlay, &Island, &SourceElementIndexToNewIndex, &DynamicMesh](int32 ElementId, int32 VertexId)
 				{
-					if (SourceEdge[1] == TriangleIndex1)
+					int32& NewIndex = SourceElementIndexToNewIndex[ElementId];
+					if (NewIndex == INDEX_NONE)
 					{
-						return FUintVector2(Island.Indices[TriangleBase + 0], Island.Indices[TriangleBase + 1]);  // Edge 0 1
+						NewIndex = Island.RestPositions.Add(FVector3f(DynamicMesh->GetVertexRef(VertexId)));
+						Island.Positions.Add(UVOverlay.GetElement(ElementId));
+						Island.PositionToSourceIndex.Add(VertexId);
 					}
-					else if (SourceEdge[1] == TriangleIndex2)
-					{
-						return FUintVector2(Island.Indices[TriangleBase + 0], Island.Indices[TriangleBase + 2]);  // Edge 0 2
-					}
-				}
-				else if (SourceEdge[0] == TriangleIndex1)
+					return NewIndex;
+				};
+
+				const int32 NewIndex0 = GetOrAddNewIndex(TriangleUVElements[0], TriangleIndices[0]);
+				const int32 NewIndex1 = GetOrAddNewIndex(TriangleUVElements[1], TriangleIndices[1]);
+				const int32 NewIndex2 = GetOrAddNewIndex(TriangleUVElements[2], TriangleIndices[2]);
+				Island.Indices.Add(NewIndex0);
+				Island.Indices.Add(NewIndex1);
+				Island.Indices.Add(NewIndex2);
+
+				TArray<int32> NeighborTriangles;
+				for (int32 LocalVertexId = 0; LocalVertexId < 3; ++LocalVertexId)
 				{
-					if (SourceEdge[1] == TriangleIndex0)
+					NeighborTriangles.Reset();
+					UVOverlay.GetElementTriangles(TriangleUVElements[LocalVertexId], NeighborTriangles);
+					for (const int32 NeighborTriangle : NeighborTriangles)
 					{
-						return FUintVector2(Island.Indices[TriangleBase + 1], Island.Indices[TriangleBase + 0]);  // Edge 1 0
-					}
-					else if (SourceEdge[1] == TriangleIndex2)
-					{
-						return FUintVector2(Island.Indices[TriangleBase + 1], Island.Indices[TriangleBase + 2]);  // Edge 1 2
+						if (!VisitedTriangles.Contains(NeighborTriangle))
+						{
+							// Mark neighboring triangle as visited
+							VisitedTriangles.Add(NeighborTriangle);
+
+							// Enqueue next triangle
+							Visitors.Enqueue(FVisitor({ NeighborTriangle }));
+						}
 					}
 				}
-				else if (SourceEdge[0] == TriangleIndex2)
+			} while (Visitors.Dequeue(Visitor));
+		}
+	}
+
+	struct FSeam
+	{
+		TSet<FIntVector2> Stitches;
+		FIntVector2 Patterns;
+	};
+
+	// Stitch together any vertices that were split, either via DynamicMesh NonManifoldMapping or UV Unwrap
+	static void BuildSeams(const TArray<FIsland>& Islands, const UE::Geometry::FDynamicMesh3& DynamicMesh, TArray<FSeam>& OutSeams)
+	{
+		OutSeams.Reset();
+
+		const UE::Geometry::FNonManifoldMappingSupport NonManifoldMapping(DynamicMesh);
+
+		TArray<TMap<int32, TArray<int32>>> IslandSourceIndexToPositions;
+		IslandSourceIndexToPositions.SetNum(Islands.Num());
+
+		for (int32 IslandIndex = 0; IslandIndex < Islands.Num(); ++IslandIndex)
+		{
+			const FIsland& Island = Islands[IslandIndex];
+			TMap<int32, TArray<int32>>& SourceIndexToPositions = IslandSourceIndexToPositions[IslandIndex];
+
+			// Build reverse lookup to PositionToSourceIndex
+			SourceIndexToPositions.Reserve(Island.PositionToSourceIndex.Num());
+			for (int32 PositionIndex = 0; PositionIndex < Island.PositionToSourceIndex.Num(); ++PositionIndex)
+			{
+				const int32 SourceIndex = NonManifoldMapping.GetOriginalNonManifoldVertexID(Island.PositionToSourceIndex[PositionIndex]);
+				SourceIndexToPositions.FindOrAdd(SourceIndex).Add(PositionIndex);
+			}
+
+			// Find all internal seams
+			FSeam InternalSeam;
+			InternalSeam.Patterns = FIntVector2(IslandIndex);
+			for (const TPair<int32, TArray<int32>>& Source : SourceIndexToPositions)
+			{
+				for (int32 FirstSourceArrayIdx = 0; FirstSourceArrayIdx < Source.Value.Num() - 1; ++FirstSourceArrayIdx)
 				{
-					if (SourceEdge[1] == TriangleIndex0)
-					{
-						return FUintVector2(Island.Indices[TriangleBase + 2], Island.Indices[TriangleBase + 0]);  // Edge 2 0
-					}
-					else if (SourceEdge[1] == TriangleIndex1)
+					for (int32 SecondSourceArrayIdx = FirstSourceArrayIdx + 1; SecondSourceArrayIdx < Source.Value.Num(); ++SecondSourceArrayIdx)
 					{
-						return FUintVector2(Island.Indices[TriangleBase + 2], Island.Indices[TriangleBase + 1]);  // Edge 2 1
+						InternalSeam.Stitches.Emplace(MakeSortedIntVector2(Source.Value[FirstSourceArrayIdx], Source.Value[SecondSourceArrayIdx]));
 					}
 				}
-				check(false);  // Sanity check
-				return FUintVector2(0);
-			};
-
-			// Look for disconnected triangles
-			for (const TPair<FUintVector2, TArray<uint32>>& SourceEdgeToTriangles : SourceEdgeToTrianglesMap)
+			}
+			if (InternalSeam.Stitches.Num())
 			{
-				const TArray<uint32>& Triangles = SourceEdgeToTriangles.Value;
-				const int32 NumTriangles = Triangles.Num();
-				if (NumTriangles > 1)
-				{
-					const FUintVector2& SourceEdge = SourceEdgeToTriangles.Key;
+				OutSeams.Emplace(MoveTemp(InternalSeam));
+			}
 
-					for (int32 Index0 = 0; Index0 < NumTriangles - 1; ++Index0)
-					{
-						const uint32 TriangleIndex0 = Triangles[Index0];
 
-						const FUintVector2 Edge0 = GetTriangleEdgeMatchingSourceEdge(TriangleIndex0, SourceEdge);
-						check(Edge0[0] <= (uint32)TNumericLimits<int32>::Max() && Edge0[1] <= (uint32)TNumericLimits<int32>::Max());
+			for (int32 OtherIslandIndex = 0; OtherIslandIndex < IslandIndex; ++OtherIslandIndex)
+			{
+				// Find all seams between the two islands
+				const TMap<int32, TArray<int32>>& OtherSourceIndexToPositions = IslandSourceIndexToPositions[OtherIslandIndex];
 
-						for (int32 Index1 = Index0 + 1; Index1 < NumTriangles; ++Index1)
+				FSeam Seam;
+				Seam.Patterns = FIntVector2(OtherIslandIndex, IslandIndex);
+				for (const TPair<int32, TArray<int32>>& FirstSource : SourceIndexToPositions)
+				{
+					if (const TArray<int32>* OtherSource = OtherSourceIndexToPositions.Find(FirstSource.Key))
+					{
+						for (const int32 FirstSourceVert : FirstSource.Value)
 						{
-							const uint32 TriangleIndex1 = Triangles[Index1];
-
-							const FUintVector2 Edge1 = GetTriangleEdgeMatchingSourceEdge(TriangleIndex1, SourceEdge);
-							check(Edge1[0] <= (uint32)TNumericLimits<int32>::Max() && Edge1[1] <= (uint32)TNumericLimits<int32>::Max());
-
-							if (Edge0[0] != Edge1[0])
+							for (const int32 OtherSourceVert : *OtherSource)
 							{
-								Seam.Stitches.Emplace(FIntVector2(MakeSortedUintVector2(Edge0[0], Edge1[0])));
-							}
-							if (Edge0[1] != Edge1[1])
-							{
-								Seam.Stitches.Emplace(FIntVector2(MakeSortedUintVector2(Edge0[1], Edge1[1])));
+								Seam.Stitches.Emplace(FIntVector2(OtherSourceVert, FirstSourceVert));
 							}
 						}
 					}
 				}
-			}
-
-			// Add this island's seams
-			if (Seam.Stitches.Num())
-			{
-				OutSeams.Emplace(Seam);
+				if (Seam.Stitches.Num())
+				{
+					OutSeams.Emplace(MoveTemp(Seam));
+				}
 			}
 		}
 	}
@@ -840,13 +925,24 @@ namespace UE::Chaos::ClothAsset
 		SetNumPatterns(0);
 	}
 
-	void FCollectionClothLodFacade::Initialize(const TArray<FVector3f>& Positions, const TArray<uint32>& Indices)
+	template<bool bWeldNearlyCoincidentVertices>
+	void FCollectionClothLodFacade::InitializeFromDynamicMeshInternal(const UE::Geometry::FDynamicMesh3& DynamicMesh, int32 UVChannelIndex)
 	{
 		using namespace UE::Chaos::ClothAsset::Private;
 		Reset();
 
+		const UE::Geometry::FDynamicMeshAttributeSet* const AttributeSet = DynamicMesh.Attributes();
+		const UE::Geometry::FDynamicMeshUVOverlay* const UVOverlay = AttributeSet ? AttributeSet->GetUVLayer(UVChannelIndex) : nullptr;
+
 		TArray<FIsland> Islands;
-		UnwrapMesh(Positions, Indices, Islands);  // Unwrap to 2D and reconstruct indices on 3D mesh
+		if (UVOverlay)
+		{
+			BuildIslandsFromDynamicMeshUVs(*UVOverlay, Islands);
+		}
+		else
+		{
+			UnwrapDynamicMesh<bWeldNearlyCoincidentVertices>(DynamicMesh, Islands);
+		}
 
 		for (FIsland& Island : Islands)
 		{
@@ -858,7 +954,7 @@ namespace UE::Chaos::ClothAsset
 		}
 
 		TArray<FSeam> Seams;
-		BuildSeams(Islands, Seams);  // Build the seam information as to be able to re-weld the mesh for simulation
+		BuildSeams(Islands, DynamicMesh, Seams);  // Build the seam information as to be able to re-weld the mesh for simulation
 
 		SetNumSeams(Seams.Num());
 
@@ -878,6 +974,27 @@ namespace UE::Chaos::ClothAsset
 				SeamStitches[SeamIndex].Emplace(Stitch[0] + (*SimVerticesStart)[Patterns[0]], Stitch[1] + (*SimVerticesStart)[Patterns[1]]);
 			}
 		}
+
+	}
+
+	void FCollectionClothLodFacade::Initialize(const TArray<FVector3f>& Positions, const TArray<uint32>& Indices)
+	{
+		using namespace UE::Chaos::ClothAsset::Private;
+
+		// Build a DynamicMesh from Positions and Indices
+		UE::Geometry::TToDynamicMeshBase<FSimpleSrcMeshInterface> ToDynamicMesh;
+		FSimpleSrcMeshInterface SimpleSrc(Positions, Indices);
+
+		UE::Geometry::FDynamicMesh3 DynamicMesh;
+		ToDynamicMesh.Convert(DynamicMesh, SimpleSrc, [](FSimpleSrcMeshInterface::TriIDType) {return 0; });
+		UE::Geometry::FNonManifoldMappingSupport::AttachNonManifoldVertexMappingData(ToDynamicMesh.ToSrcVertIDMap, DynamicMesh);
+
+		InitializeFromDynamicMeshInternal<true>(DynamicMesh, INDEX_NONE);	
+	}
+
+	void FCollectionClothLodFacade::Initialize(const UE::Geometry::FDynamicMesh3& DynamicMesh, int32 UVChannelIndex)
+	{
+		InitializeFromDynamicMeshInternal<false>(DynamicMesh, UVChannelIndex);
 	}
 
 	void FCollectionClothLodFacade::Initialize(const FCollectionClothLodConstFacade& Other)
diff --git a/Engine/Plugins/Experimental/ChaosClothAsset/Source/ChaosClothAsset/Private/ChaosClothAsset/CollectionClothPatternFacade.cpp b/Engine/Plugins/Experimental/ChaosClothAsset/Source/ChaosClothAsset/Private/ChaosClothAsset/CollectionClothPatternFacade.cpp
--- a/Engine/Plugins/Experimental/ChaosClothAsset/Source/ChaosClothAsset/Private/ChaosClothAsset/CollectionClothPatternFacade.cpp
+++ b/Engine/Plugins/Experimental/ChaosClothAsset/Source/ChaosClothAsset/Private/ChaosClothAsset/CollectionClothPatternFacade.cpp
@@ -271,7 +271,8 @@ namespace UE::Chaos::ClothAsset
 		SetDefaults();
 	}
 
-	void FCollectionClothPatternFacade::Initialize(const TArray<FVector2f>& Positions, const TArray<FVector3f>& RestPositions, const TArray<uint32>& Indices)
+	template<typename IndexType>
+	void FCollectionClothPatternFacade::Initialize(const TArray<FVector2f>& Positions, const TArray<FVector3f>& RestPositions, const TArray<IndexType>& Indices)
 	{
 		Reset();
 
@@ -335,6 +336,8 @@ namespace UE::Chaos::ClothAsset
 			SimRestNormal[SimVertexIndex] = SimRestNormal[SimVertexIndex].GetSafeNormal(UE_SMALL_NUMBER, FVector3f::XAxisVector);
 		}
 	}
+	template CHAOSCLOTHASSET_API void FCollectionClothPatternFacade::Initialize(const TArray<FVector2f>& Positions, const TArray<FVector3f>& RestPositions, const TArray<int32>& Indices);
+	template CHAOSCLOTHASSET_API void FCollectionClothPatternFacade::Initialize(const TArray<FVector2f>& Positions, const TArray<FVector3f>& RestPositions, const TArray<uint32>& Indices);
 
 	void FCollectionClothPatternFacade::Initialize(const FCollectionClothPatternConstFacade& Other)
 	{
diff --git a/Engine/Plugins/Experimental/ChaosClothAsset/Source/ChaosClothAsset/Public/ChaosClothAsset/CollectionClothLodFacade.h b/Engine/Plugins/Experimental/ChaosClothAsset/Source/ChaosClothAsset/Public/ChaosClothAsset/CollectionClothLodFacade.h
--- a/Engine/Plugins/Experimental/ChaosClothAsset/Source/ChaosClothAsset/Public/ChaosClothAsset/CollectionClothLodFacade.h
+++ b/Engine/Plugins/Experimental/ChaosClothAsset/Source/ChaosClothAsset/Public/ChaosClothAsset/CollectionClothLodFacade.h
@@ -4,6 +4,11 @@
 
 #include "ChaosClothAsset/CollectionClothPatternFacade.h"
 
+namespace UE::Geometry
+{
+class FDynamicMesh3;
+}
+
 namespace UE::Chaos::ClothAsset
 {
 	class FClothCollection;
@@ -133,6 +138,9 @@ namespace UE::Chaos::ClothAsset
 		/** Initialize the cloth LOD using another cloth collection. */
 		void Initialize(const FCollectionClothLodConstFacade& Other);
 
+		/** Initialize the cloth LOD from a DynamicMesh (possibly using UV islands to define Patterns)*/
+		void Initialize(const UE::Geometry::FDynamicMesh3& DynamicMesh, int32 UVChannelIndex);
+
 		/** Set the new number of materials to this cloth LOD. */
 		void SetNumMaterials(int32 NumMaterials);
 		/** Set the new number of tether batches to this cloth LOD. */
@@ -202,6 +210,9 @@ namespace UE::Chaos::ClothAsset
 
 		void SetDefaults();
 
+		template<bool bWeldNearlyCoincidentVertices>
+		void InitializeFromDynamicMeshInternal(const UE::Geometry::FDynamicMesh3& DynamicMesh, int32 UVChannelIndex);
+
 		TSharedPtr<FClothCollection> GetClothCollection() { return ConstCastSharedPtr<FClothCollection>(ClothCollection); }
 	};
 }  // End namespace UE::Chaos::ClothAsset
diff --git a/Engine/Plugins/Experimental/ChaosClothAsset/Source/ChaosClothAsset/Public/ChaosClothAsset/CollectionClothPatternFacade.h b/Engine/Plugins/Experimental/ChaosClothAsset/Source/ChaosClothAsset/Public/ChaosClothAsset/CollectionClothPatternFacade.h
--- a/Engine/Plugins/Experimental/ChaosClothAsset/Source/ChaosClothAsset/Public/ChaosClothAsset/CollectionClothPatternFacade.h
+++ b/Engine/Plugins/Experimental/ChaosClothAsset/Source/ChaosClothAsset/Public/ChaosClothAsset/CollectionClothPatternFacade.h
@@ -125,7 +125,8 @@ namespace UE::Chaos::ClothAsset
 		void Reset();
 
 		/** Initialize the cloth pattern using the specified 3D and 2D positions, and topology. */
-		void Initialize(const TArray<FVector2f>& Positions, const TArray<FVector3f>& RestPositions, const TArray<uint32>& Indices);
+		template<typename IndexType>
+		void Initialize(const TArray<FVector2f>& Positions, const TArray<FVector3f>& RestPositions, const TArray<IndexType>& Indices);
 
 		/** Initialize the cloth pattern using another pattern. */
 		void Initialize(const FCollectionClothPatternConstFacade& Other);
diff --git a/Engine/Plugins/Experimental/ChaosClothAssetEditor/Source/ChaosClothAssetEditor/ChaosClothAssetEditor.Build.cs b/Engine/Plugins/Experimental/ChaosClothAssetEditor/Source/ChaosClothAssetEditor/ChaosClothAssetEditor.Build.cs
--- a/Engine/Plugins/Experimental/ChaosClothAssetEditor/Source/ChaosClothAssetEditor/ChaosClothAssetEditor.Build.cs
+++ b/Engine/Plugins/Experimental/ChaosClothAssetEditor/Source/ChaosClothAssetEditor/ChaosClothAssetEditor.Build.cs
@@ -35,6 +35,8 @@ public class ChaosClothAssetEditor : ModuleRules
 				"EditorInteractiveToolsFramework",
 				"ModelingComponentsEditorOnly", // Static/skeletal mesh tool targets
 				"ModelingComponents",
+				"MeshConversion",
+				"MeshDescription",
 				"MeshModelingToolsExp",
 				"MeshModelingToolsEditorOnlyExp",
 				"EditorStyle",
diff --git a/Engine/Plugins/Experimental/ChaosClothAssetEditor/Source/ChaosClothAssetEditor/Private/ChaosClothAsset/DataflowNodes/DataflowNodes.cpp b/Engine/Plugins/Experimental/ChaosClothAssetEditor/Source/ChaosClothAssetEditor/Private/ChaosClothAsset/DataflowNodes/DataflowNodes.cpp
--- a/Engine/Plugins/Experimental/ChaosClothAssetEditor/Source/ChaosClothAssetEditor/Private/ChaosClothAsset/DataflowNodes/DataflowNodes.cpp
+++ b/Engine/Plugins/Experimental/ChaosClothAssetEditor/Source/ChaosClothAssetEditor/Private/ChaosClothAsset/DataflowNodes/DataflowNodes.cpp
@@ -4,6 +4,7 @@
 #include "ChaosClothAsset/DataflowNodes/TerminalNode.h"
 #include "ChaosClothAsset/DataflowNodes/DatasmithImportNode.h"
 #include "ChaosClothAsset/DataflowNodes/ImportNode.h"
+#include "ChaosClothAsset/DataflowNodes/StaticMeshImportNode.h"
 #include "ChaosClothAsset/DataflowNodes/PatternSelectionNode.h"
 #include "ChaosClothAsset/DataflowNodes/DeleteRenderMeshNode.h"
 #include "ChaosClothAsset/DataflowNodes/CopySimulationToRenderMeshNode.h"
@@ -23,6 +24,7 @@ namespace UE::Chaos::ClothAsset::DataflowNodes
 		DATAFLOW_NODE_REGISTER_CREATION_FACTORY(FChaosClothAssetTerminalNode);
 		DATAFLOW_NODE_REGISTER_CREATION_FACTORY(FChaosClothAssetDatasmithImportNode);
 		DATAFLOW_NODE_REGISTER_CREATION_FACTORY(FChaosClothAssetImportNode);
+		DATAFLOW_NODE_REGISTER_CREATION_FACTORY(FChaosClothAssetStaticMeshImportNode);
 		DATAFLOW_NODE_REGISTER_CREATION_FACTORY(FChaosClothAssetPatternSelectionNode);
 		DATAFLOW_NODE_REGISTER_CREATION_FACTORY(FChaosClothAssetDeleteRenderMeshNode);
 		DATAFLOW_NODE_REGISTER_CREATION_FACTORY(FChaosClothAssetCopySimulationToRenderMeshNode);
diff --git a/Engine/Plugins/Experimental/ChaosClothAssetEditor/Source/ChaosClothAssetEditor/Private/ChaosClothAsset/DataflowNodes/StaticMeshImportNode.cpp b/Engine/Plugins/Experimental/ChaosClothAssetEditor/Source/ChaosClothAssetEditor/Private/ChaosClothAsset/DataflowNodes/StaticMeshImportNode.cpp
new file mode 100644
--- /dev/null
+++ b/Engine/Plugins/Experimental/ChaosClothAssetEditor/Source/ChaosClothAssetEditor/Private/ChaosClothAsset/DataflowNodes/StaticMeshImportNode.cpp
@@ -0,0 +1,65 @@
+// Copyright Epic Games, Inc. All Rights Reserved.
+
+#include "ChaosClothAsset/DataflowNodes/StaticMeshImportNode.h"
+#include "ChaosClothAsset/DataflowNodes/DataflowNodes.h"
+#include "ChaosClothAsset/CollectionClothFacade.h"
+#include "ChaosClothAsset/ClothAsset.h"
+#include "ChaosClothAsset/ClothCollection.h"
+#include "Engine/StaticMesh.h"
+#include "MeshDescriptionToDynamicMesh.h"
+
+#include UE_INLINE_GENERATED_CPP_BY_NAME(StaticMeshImportNode)
+
+#define LOCTEXT_NAMESPACE "ChaosClothAssetStaticMeshImportNode"
+
+FChaosClothAssetStaticMeshImportNode::FChaosClothAssetStaticMeshImportNode(const Dataflow::FNodeParameters& InParam, FGuid InGuid)
+	: FDataflowNode(InParam, InGuid)
+{
+	RegisterOutputConnection(&Collection);
+}
+
+void FChaosClothAssetStaticMeshImportNode::Evaluate(Dataflow::FContext& Context, const FDataflowOutput* Out) const
+{
+	using namespace UE::Chaos::ClothAsset;
+	using namespace UE::Chaos::ClothAsset::DataflowNodes;
+
+	if (Out->IsA<FManagedArrayCollection>(&Collection))
+	{
+
+		if (StaticMesh)
+		{
+			// Evaluate in collection
+			const TSharedRef<FManagedArrayCollection> ClothCollection = MakeShared<FManagedArrayCollection>();
+			FCollectionClothFacade ClothFacade(ClothCollection);
+			ClothFacade.DefineSchema();
+
+			const int32 NumLods = StaticMesh->GetNumSourceModels();
+			for (int32 LodIndex = 0; LodIndex < NumLods; ++LodIndex)
+			{
+				const FMeshDescription* MeshDescription = StaticMesh->GetMeshDescription(LodIndex);
+				check(MeshDescription);
+
+				FMeshDescriptionToDynamicMesh Converter;
+				Converter.bPrintDebugMessages = true; // TODO: testing only
+				Converter.bEnableOutputGroups = false;
+				Converter.bVIDsFromNonManifoldMeshDescriptionAttr = true;
+				UE::Geometry::FDynamicMesh3 DynamicMesh;
+
+				Converter.Convert(MeshDescription, DynamicMesh);
+
+				FCollectionClothLodFacade ClothLodFacade = ClothFacade.AddGetLod();
+				ClothLodFacade.Initialize(DynamicMesh, UVChannel);
+			}
+			// Make sure that whatever happens there is always at least one empty LOD to avoid crashing the render data
+			if (ClothFacade.GetNumLods() < 1)
+			{
+				ClothFacade.AddLod();
+			}
+			SetValue<FManagedArrayCollection>(Context, *ClothCollection, &Collection);
+		}
+
+
+	}
+}
+
+#undef LOCTEXT_NAMESPACE
diff --git a/Engine/Plugins/Experimental/ChaosClothAssetEditor/Source/ChaosClothAssetEditor/Private/ChaosClothAsset/DataflowNodes/StaticMeshImportNode.h b/Engine/Plugins/Experimental/ChaosClothAssetEditor/Source/ChaosClothAssetEditor/Private/ChaosClothAsset/DataflowNodes/StaticMeshImportNode.h
new file mode 100644
--- /dev/null
+++ b/Engine/Plugins/Experimental/ChaosClothAssetEditor/Source/ChaosClothAssetEditor/Private/ChaosClothAsset/DataflowNodes/StaticMeshImportNode.h
@@ -0,0 +1,30 @@
+// Copyright Epic Games, Inc. All Rights Reserved.
+
+#pragma once 
+
+#include "Dataflow/DataflowNode.h"
+#include "GeometryCollection/ManagedArrayCollection.h"
+#include "StaticMeshImportNode.generated.h"
+
+class UStaticMesh;
+
+USTRUCT(Meta = (DataflowCloth))
+struct FChaosClothAssetStaticMeshImportNode : public FDataflowNode
+{
+	GENERATED_USTRUCT_BODY()
+	DATAFLOW_NODE_DEFINE_INTERNAL(FChaosClothAssetStaticMeshImportNode, "StaticMeshImport", "Cloth", "Cloth Static Mesh Import")
+public:
+
+	UPROPERTY(Meta = (DataflowOutput, DisplayName = "Collection"))
+	FManagedArrayCollection Collection;
+
+	UPROPERTY(EditAnywhere, Category = "Static Mesh")
+	TObjectPtr<const UStaticMesh> StaticMesh;
+
+	UPROPERTY(EditAnywhere, Category = "Static Mesh")
+	int32 UVChannel = INDEX_NONE;
+
+	FChaosClothAssetStaticMeshImportNode(const Dataflow::FNodeParameters& InParam, FGuid InGuid = FGuid::NewGuid());
+
+	virtual void Evaluate(Dataflow::FContext& Context, const FDataflowOutput* Out) const override;
+};
