Fix ChaosCloth cache playback from CL 24765088.

No longer storing TrackToParticle data, so can't use that to validate cache. Keeping both checks so old and new caches still work.

diff --git a/Engine/Plugins/Experimental/ChaosCloth/Source/Chaos/Private/ChaosCloth/ChaosClothingCacheSchema.cpp b/Engine/Plugins/Experimental/ChaosCloth/Source/Chaos/Private/ChaosCloth/ChaosClothingCacheSchema.cpp
--- a/Engine/Plugins/Experimental/ChaosCloth/Source/Chaos/Private/ChaosCloth/ChaosClothingCacheSchema.cpp
+++ b/Engine/Plugins/Experimental/ChaosCloth/Source/Chaos/Private/ChaosCloth/ChaosClothingCacheSchema.cpp
@@ -207,4 +207,10 @@ void FClothingCacheSchema::LoadCacheData(UChaosCache* InCache, Chaos::FReal InTi
 	}
 }
 
+bool FClothingCacheSchema::CacheIsValidForPlayback(UChaosCache* InCache)
+{
+	// Old caches use TrackToParticle. New caches use ChannelCurveToParticle
+	return InCache->TrackToParticle.Num() > 0 || InCache->ChannelCurveToParticle.Num() > 0; 
+}
+
 } // namespace Chaos
\ No newline at end of file
diff --git a/Engine/Plugins/Experimental/ChaosCloth/Source/Chaos/Private/ChaosCloth/SkeletalMeshComponentCacheAdapter.cpp b/Engine/Plugins/Experimental/ChaosCloth/Source/Chaos/Private/ChaosCloth/SkeletalMeshComponentCacheAdapter.cpp
--- a/Engine/Plugins/Experimental/ChaosCloth/Source/Chaos/Private/ChaosCloth/SkeletalMeshComponentCacheAdapter.cpp
+++ b/Engine/Plugins/Experimental/ChaosCloth/Source/Chaos/Private/ChaosCloth/SkeletalMeshComponentCacheAdapter.cpp
@@ -64,7 +64,7 @@ namespace Chaos
 	{
 		// If we have a skel mesh we can play back any cache as long as it has one or more tracks
 		const USkeletalMeshComponent* MeshComp = CastChecked<USkeletalMeshComponent>(InComponent);
-		return MeshComp && MeshComp->GetSkinnedAsset() && InCache->TrackToParticle.Num() > 0;
+		return MeshComp && MeshComp->GetSkinnedAsset() && Chaos::FClothingCacheSchema::CacheIsValidForPlayback(InCache);
 	}
 
 	FClothingSimulationSolver* FSkeletalMeshCacheAdapter::GetClothSolver(UPrimitiveComponent* InComponent) const
diff --git a/Engine/Plugins/Experimental/ChaosCloth/Source/Chaos/Public/ChaosCloth/ChaosClothingCacheSchema.h b/Engine/Plugins/Experimental/ChaosCloth/Source/Chaos/Public/ChaosCloth/ChaosClothingCacheSchema.h
--- a/Engine/Plugins/Experimental/ChaosCloth/Source/Chaos/Public/ChaosCloth/ChaosClothingCacheSchema.h
+++ b/Engine/Plugins/Experimental/ChaosCloth/Source/Chaos/Public/ChaosCloth/ChaosClothingCacheSchema.h
@@ -23,6 +23,8 @@ struct FClothingCacheSchema
 	static void CHAOSCLOTH_API PlaybackPreSolve(UChaosCache& InCache, FReal InTime, FPlaybackTickRecord& TickRecord, FClothingSimulationSolver& ClothSolver);
 	static void CHAOSCLOTH_API LoadCacheData(UChaosCache* InCache, FReal InTime, FClothingSimulationCacheData& CacheData);
 
+	static bool CHAOSCLOTH_API CacheIsValidForPlayback(UChaosCache* InCache);
+
 private:
 	inline static const FName VelocityXName = TEXT("VelocityX");
 	inline static const FName VelocityYName = TEXT("VelocityY");
diff --git a/Engine/Plugins/Experimental/ChaosClothAsset/Source/ChaosClothAssetEngine/Private/ChaosClothAsset/ClothComponentCacheAdapter.cpp b/Engine/Plugins/Experimental/ChaosClothAsset/Source/ChaosClothAssetEngine/Private/ChaosClothAsset/ClothComponentCacheAdapter.cpp
--- a/Engine/Plugins/Experimental/ChaosClothAsset/Source/ChaosClothAssetEngine/Private/ChaosClothAsset/ClothComponentCacheAdapter.cpp
+++ b/Engine/Plugins/Experimental/ChaosClothAsset/Source/ChaosClothAssetEngine/Private/ChaosClothAsset/ClothComponentCacheAdapter.cpp
@@ -59,7 +59,7 @@ namespace UE::Chaos::ClothAsset
 	bool FClothComponentCacheAdapter::ValidForPlayback(UPrimitiveComponent* InComponent, UChaosCache* InCache) const
 	{
 		const UChaosClothComponent* ClothComp = GetClothComponent(InComponent);
-		return ClothComp && InCache->ChannelCurveToParticle.Num() > 0;
+		return ClothComp && ::Chaos::FClothingCacheSchema::CacheIsValidForPlayback(InCache);
 	}
 
 	::Chaos::FPhysicsSolverEvents* FClothComponentCacheAdapter::BuildEventsSolver(UPrimitiveComponent* InComponent) const
