Reverting the change (CL24283502) to the updating of the bone matrices logic in FClothSimulationContext.
Instead added a line that resets the RefToLocalMatrices array to mimic the logic of the FClothingSimulationContextCommon::FillRefToLocals function.

#rnx

diff --git a/Engine/Plugins/Experimental/ChaosClothAsset/Source/ChaosClothAssetEngine/Private/ChaosClothAsset/ClothSimulationContext.cpp b/Engine/Plugins/Experimental/ChaosClothAsset/Source/ChaosClothAssetEngine/Private/ChaosClothAsset/ClothSimulationContext.cpp
--- a/Engine/Plugins/Experimental/ChaosClothAsset/Source/ChaosClothAssetEngine/Private/ChaosClothAsset/ClothSimulationContext.cpp
+++ b/Engine/Plugins/Experimental/ChaosClothAsset/Source/ChaosClothAssetEngine/Private/ChaosClothAsset/ClothSimulationContext.cpp
@@ -92,20 +92,20 @@ namespace UE::Chaos::ClothAsset
 		}
 
 		// Update bone matrices
-		if (bIsInitialization)
+		RefToLocalMatrices.Reset(NumBones);
+		if (!bIsInitialization)
 		{
-			RefToLocalMatrices.Reset(NumBones);
-			RefToLocalMatrices.AddUninitialized(NumBones);
+			ClothComponent.GetCurrentRefToLocalMatrices(RefToLocalMatrices, LodIndex);
 		}
-		
-		// We must set the matrices to identity on every call to Fill, otherwise every call to GetCurrentRefToLocalMatrices 
-		// below will multiply the content of the RefToLocalMatrices by another matrix and save the result back to
-		// RefToLocalMatrices which leads to the translation component increasing
-		for (int32 BoneIndex = 0; BoneIndex < NumBones; ++BoneIndex)
+		else
 		{
-			RefToLocalMatrices[BoneIndex] = FMatrix44f::Identity;
+			RefToLocalMatrices.AddUninitialized(NumBones);
+
+			for (int32 BoneIndex = 0; BoneIndex < NumBones; ++BoneIndex)
+			{
+				RefToLocalMatrices[BoneIndex] = FMatrix44f::Identity;
+			}
 		}
-		ClothComponent.GetCurrentRefToLocalMatrices(RefToLocalMatrices, LodIndex);
 		
 		// Update gravity
 		const UWorld* const World = ClothComponent.GetWorld();
