Reset the cloth after its added to the solver.

#rnx

diff --git a/Engine/Plugins/Experimental/ChaosClothAsset/Source/ChaosClothAssetEngine/Private/ChaosClothAsset/ClothSimulationProxy.cpp b/Engine/Plugins/Experimental/ChaosClothAsset/Source/ChaosClothAssetEngine/Private/ChaosClothAsset/ClothSimulationProxy.cpp
--- a/Engine/Plugins/Experimental/ChaosClothAsset/Source/ChaosClothAssetEngine/Private/ChaosClothAsset/ClothSimulationProxy.cpp
+++ b/Engine/Plugins/Experimental/ChaosClothAsset/Source/ChaosClothAssetEngine/Private/ChaosClothAsset/ClothSimulationProxy.cpp
@@ -146,6 +146,7 @@ PRAGMA_ENABLE_DEPRECATION_WARNINGS
 			GroupId));
 
 		Solver->AddCloth(Cloths[ClothIndex].Get());
+		Cloths[ClothIndex]->Reset();
 
 		// Create solver config simulation thread object
 		const int32 SolverConfigIndex = Configs.Emplace(MakeUnique<FClothingSimulationConfig>(ClothComponent.GetPropertyCollection()));  // TODO: Use a separate solver config for outfits
