Chaos Cloth - Added a new function to update the skeleton from the collection.

#rnx

diff --git a/Engine/Plugins/Experimental/ChaosClothAsset/Source/ChaosClothAssetEngine/Private/ChaosClothAsset/ClothAsset.cpp b/Engine/Plugins/Experimental/ChaosClothAsset/Source/ChaosClothAssetEngine/Private/ChaosClothAsset/ClothAsset.cpp
--- a/Engine/Plugins/Experimental/ChaosClothAsset/Source/ChaosClothAssetEngine/Private/ChaosClothAsset/ClothAsset.cpp
+++ b/Engine/Plugins/Experimental/ChaosClothAsset/Source/ChaosClothAssetEngine/Private/ChaosClothAsset/ClothAsset.cpp
@@ -625,6 +625,36 @@ void UChaosClothAsset::SetSkeleton(USkeleton* InSkeleton, bool bRebuildModels)
 	}
 }
 
+void UChaosClothAsset::UpdateSkeletonFromCollection(bool bRebuildModels)
+{
+	using namespace UE::Chaos::ClothAsset;
+	check(ClothCollection.IsValid());
+
+	FCollectionClothConstFacade ClothFacade(ClothCollection);
+
+	USkeleton* InSkeleton = nullptr;
+	if (ClothFacade.GetNumLods())
+	{
+		const FString& SkeletonAssetPathName = ClothFacade.GetLod(0).GetSkeletonAssetPathName();
+		InSkeleton = LoadObject<USkeleton>(nullptr, *SkeletonAssetPathName, nullptr, LOAD_None, nullptr);
+	}
+
+	if (!InSkeleton)
+	{
+		SetSkeleton(nullptr, true);
+	}
+	else
+	{
+		Skeleton = InSkeleton;
+		RefSkeleton = Skeleton->GetReferenceSkeleton();
+
+		if (bRebuildModels)
+		{
+			Build();
+		}
+	}
+}
+
 void UChaosClothAsset::CopySimMeshToRenderMesh(UMaterialInterface* Material)
 {
 	using namespace UE::Chaos::ClothAsset;
diff --git a/Engine/Plugins/Experimental/ChaosClothAsset/Source/ChaosClothAssetEngine/Public/ChaosClothAsset/ClothAsset.h b/Engine/Plugins/Experimental/ChaosClothAsset/Source/ChaosClothAssetEngine/Public/ChaosClothAsset/ClothAsset.h
--- a/Engine/Plugins/Experimental/ChaosClothAsset/Source/ChaosClothAssetEngine/Public/ChaosClothAsset/ClothAsset.h
+++ b/Engine/Plugins/Experimental/ChaosClothAsset/Source/ChaosClothAssetEngine/Public/ChaosClothAsset/ClothAsset.h
@@ -122,10 +122,16 @@ public:
 	/**
 	 * Set the skeleton asset and bone hierachy to use for this cloth.
 	 * @param InSkeleton The skeleton asset to use for this cloth, or nullptr to setup a default skeleton asset.
-	 * @param bRebuildClothModel Whether to rebuild the simulation and rendering models (editor builds only).
+	 * @param bRebuildModels Whether to rebuild the simulation and rendering models (editor builds only).
 	 */
 	void SetSkeleton(USkeleton* InSkeleton, bool bRebuildModels);
 
+	/**
+	 * Set the skeleton asset using the LOD0 skeleton specified in the cloth collection, or load a default skeleton if it is invalid.
+	 * @param bRebuildModels Whether to rebuild the simulation and rendering models (editor builds only).
+	 */
+	void UpdateSkeletonFromCollection(bool bRebuildModels);
+
 	/** Set the bone hierachy to use for this cloth. */
 	UE_DEPRECATED(5.3, "Use SetSkeleton instead")
 	PRAGMA_DISABLE_DEPRECATION_WARNINGS
