ClothAssetEngine: remove ProfileVisualizer

diff --git a/Engine/Plugins/Experimental/ChaosClothAsset/Source/ChaosClothAssetEngine/ChaosClothAssetEngine.Build.cs b/Engine/Plugins/Experimental/ChaosClothAsset/Source/ChaosClothAssetEngine/ChaosClothAssetEngine.Build.cs
--- a/Engine/Plugins/Experimental/ChaosClothAsset/Source/ChaosClothAssetEngine/ChaosClothAssetEngine.Build.cs
+++ b/Engine/Plugins/Experimental/ChaosClothAsset/Source/ChaosClothAssetEngine/ChaosClothAssetEngine.Build.cs
@@ -19,9 +19,6 @@ public class ChaosClothAssetEngine : ModuleRules
 				"Chaos",
 				"ChaosClothAsset",
 				"ChaosCloth",
-				"ChaosCaching",
-				"ProfileVisualizer",
-				"ClothingSystemRuntimeCommon",
 				"ClothingSystemRuntimeInterface",
 				"DataflowEngine"
 			}
