ChaosClothAssetEngine: fix module dependencies

diff --git a/Engine/Plugins/Experimental/ChaosClothAsset/Source/ChaosClothAssetEngine/ChaosClothAssetEngine.Build.cs b/Engine/Plugins/Experimental/ChaosClothAsset/Source/ChaosClothAssetEngine/ChaosClothAssetEngine.Build.cs
--- a/Engine/Plugins/Experimental/ChaosClothAsset/Source/ChaosClothAssetEngine/ChaosClothAssetEngine.Build.cs
+++ b/Engine/Plugins/Experimental/ChaosClothAsset/Source/ChaosClothAssetEngine/ChaosClothAssetEngine.Build.cs
@@ -19,6 +19,8 @@ public class ChaosClothAssetEngine : ModuleRules
 				"Chaos",
 				"ChaosClothAsset",
 				"ChaosCloth",
+				"ChaosCaching",
+				"ClothingSystemRuntimeCommon",
 				"ClothingSystemRuntimeInterface",
 				"DataflowEngine"
 			}
