Chaos Cloth:
- Added a CreateSimulationProxy virtual function in the cloth component for the cloth training tool.
- Renamed the cloth component name and tool tip to Chaos cloth component instead of Cloth simulation component.

#rnx

diff --git a/Engine/Plugins/Experimental/ChaosClothAsset/Source/ChaosClothAssetEngine/Private/ChaosClothAsset/ClothComponent.cpp b/Engine/Plugins/Experimental/ChaosClothAsset/Source/ChaosClothAssetEngine/Private/ChaosClothAsset/ClothComponent.cpp
--- a/Engine/Plugins/Experimental/ChaosClothAsset/Source/ChaosClothAssetEngine/Private/ChaosClothAsset/ClothComponent.cpp
+++ b/Engine/Plugins/Experimental/ChaosClothAsset/Source/ChaosClothAssetEngine/Private/ChaosClothAsset/ClothComponent.cpp
@@ -111,11 +111,17 @@ void UChaosClothComponent::OnRegister()
 			ResetConfigProperties();
 
 			// Create simulation proxy
-			ClothSimulationProxy = MakeUnique<FClothSimulationProxy>(*this);
+			ClothSimulationProxy = CreateClothSimulationProxy();
 		}
 	}
 }
 
+TUniquePtr<UE::Chaos::ClothAsset::FClothSimulationProxy> UChaosClothComponent::CreateClothSimulationProxy()
+{
+	using namespace UE::Chaos::ClothAsset;
+	return MakeUnique<FClothSimulationProxy>(*this);
+}
+
 bool UChaosClothComponent::IsSimulationSuspended() const
 {
 	static IConsoleVariable* const CVarClothPhysics = IConsoleManager::Get().FindConsoleVariable(TEXT("p.ClothPhysics"));
diff --git a/Engine/Plugins/Experimental/ChaosClothAsset/Source/ChaosClothAssetEngine/Public/ChaosClothAsset/ClothComponent.h b/Engine/Plugins/Experimental/ChaosClothAsset/Source/ChaosClothAssetEngine/Public/ChaosClothAsset/ClothComponent.h
--- a/Engine/Plugins/Experimental/ChaosClothAsset/Source/ChaosClothAssetEngine/Public/ChaosClothAsset/ClothComponent.h
+++ b/Engine/Plugins/Experimental/ChaosClothAsset/Source/ChaosClothAssetEngine/Public/ChaosClothAsset/ClothComponent.h
@@ -25,8 +25,8 @@ namespace UE::Chaos::ClothAsset
  */
 UCLASS(
 	ClassGroup = Physics, 
-	Meta = (BlueprintSpawnableComponent, ToolTip = "Chaos cloth simulation component."),
-	DisplayName = "Cloth Simulation",
+	Meta = (BlueprintSpawnableComponent, ToolTip = "Chaos cloth component."),
+	DisplayName = "Chaos cloth component",
 	HideCategories = (Object, "Mesh|SkeletalAsset", Constraints, Advanced, Cooking, Collision, Navigation))
 class CHAOSCLOTHASSETENGINE_API UChaosClothComponent : public USkinnedMeshComponent
 {
@@ -117,6 +117,9 @@ protected:
 	virtual void GetUpdateClothSimulationData_AnyThread(TMap<int32, FClothSimulData>& OutClothSimulData, FMatrix& OutLocalToWorld, float& OutBlendWeight) override;
 	//~ End USkinnedMeshComponent Interface
 
+	/** Override this function for setting up custom simulation proxies when the component is registered. */
+	virtual TUniquePtr<UE::Chaos::ClothAsset::FClothSimulationProxy> CreateClothSimulationProxy();
+
 private:
 	void StartNewParallelSimulation(float DeltaTime);
 	void HandleExistingParallelSimulation();
