Chaos Cloth:
- Added a reverse winding order functionality to the ReverseNormal node, to allow for some imported meshes to correctly display themselves in wireframe in the editor view/paint tool view.
- Disabled the pattern selection on the ReverseNormal node until we have any use for it.

#rnx

diff --git a/Engine/Plugins/Experimental/ChaosClothAsset/Source/ChaosClothAsset/Private/ChaosClothAsset/ClothGeometryTools.cpp b/Engine/Plugins/Experimental/ChaosClothAsset/Source/ChaosClothAsset/Private/ChaosClothAsset/ClothGeometryTools.cpp
--- a/Engine/Plugins/Experimental/ChaosClothAsset/Source/ChaosClothAsset/Private/ChaosClothAsset/ClothGeometryTools.cpp
+++ b/Engine/Plugins/Experimental/ChaosClothAsset/Source/ChaosClothAsset/Private/ChaosClothAsset/ClothGeometryTools.cpp
@@ -214,10 +214,12 @@ namespace UE::Chaos::ClothAsset
 		}
 	}
 
-	void FClothGeometryTools::ReverseNormals(
+	void FClothGeometryTools::ReverseMesh(
 		const TSharedPtr<FManagedArrayCollection>& ClothCollection,
 		bool bReverseSimMeshNormals,
+		bool bReverseSimMeshWindingOrder,
 		bool bReverseRenderMeshNormals,
+		bool bReverseRenderMeshWindingOrder,
 		const TArray<int32>& PatternSelection)
 	{
 		auto ReverseSimNormals = [](const TArrayView<FVector3f>& SimRestNormal)
@@ -228,14 +230,21 @@ namespace UE::Chaos::ClothAsset
 				}
 			};
 		auto ReverseRenderNormals = [](const TArrayView<FVector3f>& RenderNormal, const TArrayView<FVector3f>& RenderTangentU)
-		{
-			check(RenderNormal.Num() == RenderTangentU.Num())
-			for (int32 VertexIndex = 0; VertexIndex < RenderNormal.Num(); ++VertexIndex)
 			{
-				RenderNormal[VertexIndex] = -RenderNormal[VertexIndex];      // Equivalent of rotating the normal basis
-				RenderTangentU[VertexIndex] = -RenderTangentU[VertexIndex];  // around tangent V
-			}
-		};
+				check(RenderNormal.Num() == RenderTangentU.Num())
+				for (int32 VertexIndex = 0; VertexIndex < RenderNormal.Num(); ++VertexIndex)
+				{
+					RenderNormal[VertexIndex] = -RenderNormal[VertexIndex];      // Equivalent of rotating the normal basis
+					RenderTangentU[VertexIndex] = -RenderTangentU[VertexIndex];  // around tangent V
+				}
+			};
+		auto ReverseWindingOrder = [](const TArrayView<FIntVector3>& Indices)
+			{
+				for (int32 FaceIndex = 0; FaceIndex < Indices.Num(); ++FaceIndex)
+				{
+					Swap(Indices[FaceIndex][1], Indices[FaceIndex][2]);
+				}
+			};
 
 		FCollectionClothFacade ClothFacade(ClothCollection);
 
@@ -249,10 +258,18 @@ namespace UE::Chaos::ClothAsset
 				{
 					ReverseSimNormals(ClothLodFacade.GetSimRestNormal());
 				}
+				if (bReverseSimMeshWindingOrder)
+				{
+					ReverseWindingOrder(ClothLodFacade.GetSimIndices());
+				}
 				if (bReverseRenderMeshNormals)
 				{
 					ReverseRenderNormals(ClothLodFacade.GetRenderNormal(), ClothLodFacade.GetRenderTangentU());
 				}
+				if (bReverseRenderMeshWindingOrder)
+				{
+					ReverseWindingOrder(ClothLodFacade.GetRenderIndices());
+				}
 			}
 			else
 			{
@@ -266,10 +283,18 @@ namespace UE::Chaos::ClothAsset
 						{
 							ReverseSimNormals(ClothPatternFacade.GetSimRestNormal());
 						}
+						if (bReverseSimMeshWindingOrder)
+						{
+							ReverseWindingOrder(ClothPatternFacade.GetSimIndices());
+						}
 						if (bReverseRenderMeshNormals)
 						{
 							ReverseRenderNormals(ClothPatternFacade.GetRenderNormal(), ClothPatternFacade.GetRenderTangentU());
 						}
+						if (bReverseRenderMeshWindingOrder)
+						{
+							ReverseWindingOrder(ClothPatternFacade.GetRenderIndices());
+						}
 					}
 				}
 			}
diff --git a/Engine/Plugins/Experimental/ChaosClothAsset/Source/ChaosClothAsset/Public/ChaosClothAsset/ClothGeometryTools.h b/Engine/Plugins/Experimental/ChaosClothAsset/Source/ChaosClothAsset/Public/ChaosClothAsset/ClothGeometryTools.h
--- a/Engine/Plugins/Experimental/ChaosClothAsset/Source/ChaosClothAsset/Public/ChaosClothAsset/ClothGeometryTools.h
+++ b/Engine/Plugins/Experimental/ChaosClothAsset/Source/ChaosClothAsset/Public/ChaosClothAsset/ClothGeometryTools.h
@@ -32,10 +32,12 @@ namespace UE::Chaos::ClothAsset
 		static void CopySimMeshToRenderMesh(const TSharedPtr<FManagedArrayCollection>& ClothCollection, const FString& RenderMaterialPathName);
 
 		/** Reverse the mesh normals. Will reverse all normals if pattern selection is empty. */
-		static void ReverseNormals(
+		static void ReverseMesh(
 			const TSharedPtr<FManagedArrayCollection>& ClothCollection,
 			bool bReverseSimMeshNormals,
+			bool bReverseSimMeshWindingOrder,
 			bool bReverseRenderMeshNormals,
+			bool bReverseRenderMeshWindingOrder,
 			const TArray<int32>& PatternSelection);
 
 		/**
diff --git a/Engine/Plugins/Experimental/ChaosClothAssetEditor/Source/ChaosClothAssetDataflowNodes/Private/ChaosClothAsset/ReverseNormalsNode.cpp b/Engine/Plugins/Experimental/ChaosClothAssetEditor/Source/ChaosClothAssetDataflowNodes/Private/ChaosClothAsset/ReverseNormalsNode.cpp
--- a/Engine/Plugins/Experimental/ChaosClothAssetEditor/Source/ChaosClothAssetDataflowNodes/Private/ChaosClothAsset/ReverseNormalsNode.cpp
+++ b/Engine/Plugins/Experimental/ChaosClothAssetEditor/Source/ChaosClothAssetDataflowNodes/Private/ChaosClothAsset/ReverseNormalsNode.cpp
@@ -26,7 +26,13 @@ void FChaosClothAssetReverseNormalsNode::Evaluate(Dataflow::FContext& Context, c
 		FManagedArrayCollection InCollection = GetValue<FManagedArrayCollection>(Context, &Collection);
 		const TSharedRef<FManagedArrayCollection> ClothCollection = MakeShared<FManagedArrayCollection>(MoveTemp(InCollection));
 
-		FClothGeometryTools::ReverseNormals(ClothCollection, bReverseSimMeshNormals, bReverseRenderMeshNormals, Patterns);
+		FClothGeometryTools::ReverseMesh(
+			ClothCollection,
+			bReverseSimMeshNormals,
+			bReverseSimMeshWindingOrder,
+			bReverseRenderMeshNormals,
+			bReverseRenderMeshWindingOrder,
+			Patterns);
 
 		SetValue<FManagedArrayCollection>(Context, *ClothCollection, &Collection);
 	}
diff --git a/Engine/Plugins/Experimental/ChaosClothAssetEditor/Source/ChaosClothAssetDataflowNodes/Public/ChaosClothAsset/ReverseNormalsNode.h b/Engine/Plugins/Experimental/ChaosClothAssetEditor/Source/ChaosClothAssetDataflowNodes/Public/ChaosClothAsset/ReverseNormalsNode.h
--- a/Engine/Plugins/Experimental/ChaosClothAssetEditor/Source/ChaosClothAssetDataflowNodes/Public/ChaosClothAsset/ReverseNormalsNode.h
+++ b/Engine/Plugins/Experimental/ChaosClothAssetEditor/Source/ChaosClothAssetDataflowNodes/Public/ChaosClothAsset/ReverseNormalsNode.h
@@ -19,17 +19,25 @@ public:
 	FManagedArrayCollection Collection;
 
 	/** List of patterns to apply the operation on. All patterns will be used if left empty. */
-	UPROPERTY(EditAnywhere, Category = "Pattern Selection")
+	UPROPERTY(EditAnywhere, Category = "Reverse Normals", Meta = (EditCondition = "false"))
 	TArray<int32> Patterns;
 
 	/** Whether to reverse the simulation mesh normals. */
 	UPROPERTY(EditAnywhere, Category = "Reverse Normals")
 	bool bReverseSimMeshNormals = true;
 
+	/** Whether to reverse the simulation mesh triangles' winding order. */
+	UPROPERTY(EditAnywhere, Category = "Reverse Normals")
+	bool bReverseSimMeshWindingOrder = false;
+
 	/** Whether to reverse the render mesh normals. */
 	UPROPERTY(EditAnywhere, Category = "Reverse Normals")
 	bool bReverseRenderMeshNormals = true;
 
+	/** Whether to reverse the render mesh triangles' winding order. */
+	UPROPERTY(EditAnywhere, Category = "Reverse Normals")
+	bool bReverseRenderMeshWindingOrder = false;
+
 	FChaosClothAssetReverseNormalsNode(const Dataflow::FNodeParameters& InParam, FGuid InGuid = FGuid::NewGuid());
 
 	virtual void Evaluate(Dataflow::FContext& Context, const FDataflowOutput* Out) const override;
