Chaos Cloth:
- Added merge node.
- Added check on map names.
- Added UVScale to static mesh import.
- Added SetPhysicsAsset node.

#rnx

diff --git a/Engine/Plugins/Experimental/ChaosClothAsset/Source/ChaosClothAsset/Private/ChaosClothAsset/CollectionClothLodFacade.cpp b/Engine/Plugins/Experimental/ChaosClothAsset/Source/ChaosClothAsset/Private/ChaosClothAsset/CollectionClothLodFacade.cpp
--- a/Engine/Plugins/Experimental/ChaosClothAsset/Source/ChaosClothAsset/Private/ChaosClothAsset/CollectionClothLodFacade.cpp
+++ b/Engine/Plugins/Experimental/ChaosClothAsset/Source/ChaosClothAsset/Private/ChaosClothAsset/CollectionClothLodFacade.cpp
@@ -313,7 +313,7 @@ namespace UE::Chaos::ClothAsset::Private
 		}
 	}
 
-	static void BuildIslandsFromDynamicMeshUVs(const UE::Geometry::FDynamicMeshUVOverlay& UVOverlay, TArray<FIsland>& OutIslands)
+	static void BuildIslandsFromDynamicMeshUVs(const UE::Geometry::FDynamicMeshUVOverlay& UVOverlay, TArray<FIsland>& OutIslands, const FVector2f& UVScale)
 	{
 		using namespace UE::Geometry;
 
@@ -359,13 +359,13 @@ namespace UE::Chaos::ClothAsset::Private
 				const FIndex3i TriangleIndices = DynamicMesh->GetTriangle(Triangle);
 				const FIndex3i TriangleUVElements = UVOverlay.GetTriangle(Triangle);
 
-				auto GetOrAddNewIndex = [&UVOverlay, &Island, &SourceElementIndexToNewIndex, &DynamicMesh](int32 ElementId, int32 VertexId)
+				auto GetOrAddNewIndex = [&UVOverlay, &Island, &SourceElementIndexToNewIndex, &DynamicMesh, &UVScale](int32 ElementId, int32 VertexId)
 				{
 					int32& NewIndex = SourceElementIndexToNewIndex[ElementId];
 					if (NewIndex == INDEX_NONE)
 					{
 						NewIndex = Island.RestPositions.Add(FVector3f(DynamicMesh->GetVertexRef(VertexId)));
-						Island.Positions.Add(UVOverlay.GetElement(ElementId));
+						Island.Positions.Add((FVector2f(1.f) - UVOverlay.GetElement(ElementId)) * UVScale);  // The static mesh import uses 1 - UV for some reason
 						Island.PositionToSourceIndex.Add(VertexId);
 					}
 					return NewIndex;
@@ -929,7 +929,7 @@ namespace UE::Chaos::ClothAsset
 	}
 
 	template<bool bWeldNearlyCoincidentVertices>
-	void FCollectionClothLodFacade::InitializeFromDynamicMeshInternal(const UE::Geometry::FDynamicMesh3& DynamicMesh, int32 UVChannelIndex)
+	void FCollectionClothLodFacade::InitializeFromDynamicMeshInternal(const UE::Geometry::FDynamicMesh3& DynamicMesh, int32 UVChannelIndex, const FVector2f& UVScale)
 	{
 		using namespace UE::Chaos::ClothAsset::Private;
 		Reset();
@@ -940,7 +940,7 @@ namespace UE::Chaos::ClothAsset
 		TArray<FIsland> Islands;
 		if (UVOverlay)
 		{
-			BuildIslandsFromDynamicMeshUVs(*UVOverlay, Islands);
+			BuildIslandsFromDynamicMeshUVs(*UVOverlay, Islands, UVScale);
 		}
 		else
 		{
@@ -992,12 +992,12 @@ namespace UE::Chaos::ClothAsset
 		ToDynamicMesh.Convert(DynamicMesh, SimpleSrc, [](FSimpleSrcMeshInterface::TriIDType) {return 0; });
 		UE::Geometry::FNonManifoldMappingSupport::AttachNonManifoldVertexMappingData(ToDynamicMesh.ToSrcVertIDMap, DynamicMesh);
 
-		InitializeFromDynamicMeshInternal<true>(DynamicMesh, INDEX_NONE);	
+		InitializeFromDynamicMeshInternal<true>(DynamicMesh, INDEX_NONE, FVector2f(1.f));
 	}
 
-	void FCollectionClothLodFacade::Initialize(const UE::Geometry::FDynamicMesh3& DynamicMesh, int32 UVChannelIndex)
+	void FCollectionClothLodFacade::Initialize(const UE::Geometry::FDynamicMesh3& DynamicMesh, int32 UVChannelIndex, const FVector2f& UVScale)
 	{
-		InitializeFromDynamicMeshInternal<false>(DynamicMesh, UVChannelIndex);
+		InitializeFromDynamicMeshInternal<false>(DynamicMesh, UVChannelIndex, UVScale);
 	}
 
 	void FCollectionClothLodFacade::Initialize(const FCollectionClothLodConstFacade& Other)
@@ -1042,6 +1042,70 @@ namespace UE::Chaos::ClothAsset
 		SetSkeletonAssetPathName(Other.GetSkeletonAssetPathName());
 	}
 
+	void FCollectionClothLodFacade::Append(const FCollectionClothLodConstFacade& Other)
+	{
+		// Expand the cloth collection
+		const int32 StartMaterial = GetNumMaterials();
+		const int32 StartTetherBatch = GetNumTetherBatches();
+		const int32 StartSeam = GetNumSeams();
+		const int32 StartPattern = GetNumPatterns();
+		const int32 StartVertex = GetNumSimVertices();
+
+		const int32 OtherNumMaterials = Other.GetNumMaterials();
+		const int32 OtherNumTetherBatches = Other.GetNumTetherBatches();
+		const int32 OtherNumSeams = Other.GetNumSeams();
+		const int32 OtherNumPatterns = Other.GetNumPatterns();
+
+		SetNumMaterials(StartMaterial + OtherNumMaterials);
+		SetNumTetherBatches(StartTetherBatch + OtherNumTetherBatches);
+		SetNumSeams(StartSeam + OtherNumSeams);
+		SetNumPatterns(StartPattern + OtherNumPatterns);
+
+		// Materials Group
+		TArrayView<FString> RenderMaterialPathNameArray = GetRenderMaterialPathName();
+		TConstArrayView<FString> OtherRenderMaterialPathNameArray = Other.GetRenderMaterialPathName();
+		for (int32 MaterialIndex = 0; MaterialIndex < OtherNumMaterials; ++MaterialIndex)
+		{
+			RenderMaterialPathNameArray[StartMaterial + MaterialIndex] = OtherRenderMaterialPathNameArray[MaterialIndex];
+		}
+
+		// Patterns Group, make sure to do this before the Tether and Seam groups so that the sim vertex indices are already added
+		for (int32 PatternIndex = 0; PatternIndex < OtherNumPatterns; ++PatternIndex)
+		{
+			FCollectionClothPatternFacade ThisClothPattern = GetPattern(StartPattern + PatternIndex);
+			const FCollectionClothPatternConstFacade OtherClothPattern = Other.GetPattern(PatternIndex);
+			ThisClothPattern.Initialize(OtherClothPattern, StartMaterial);
+		}
+
+		// Tether Batches Group
+		for (int32 TetherBatchIndex = 0; TetherBatchIndex < OtherNumTetherBatches; ++TetherBatchIndex)
+		{
+			//FCollectionClothTetherBatchFacade ThisClothTetherBatch = GetTetherBatch(StartTetherBatch + TetherBatchIndex);  // TODO: Tether Batches facade
+			//const FCollectionClothTetherBatchConstFacade OtherClothTetherBatch = Other.GetTetherBatch(TetherBatchIndex);
+			//ThisClothTetherBatch.Initialize(OtherClothTetherBatch, StartVertex);
+		}
+
+		// Seam Group
+		for (int32 SeamIndex = 0; SeamIndex < OtherNumSeams; ++SeamIndex)
+		{
+			FIntVector2 OtherSeamPattern = Other.GetSeamPatterns()[SeamIndex];
+			OtherSeamPattern.X += StartPattern;  // Reindex the seam pattern indices
+			OtherSeamPattern.Y += StartPattern;  //
+			GetSeamPatterns()[StartSeam + SeamIndex] = OtherSeamPattern;
+
+			TArray<FIntVector2> OtherSeamStitches = Other.GetSeamStitches()[SeamIndex];
+			for (FIntVector2& OtherSeamStich : OtherSeamStitches)
+			{
+				OtherSeamStich.X += StartVertex;  // Reindex the seam vertex indices
+				OtherSeamStich.Y += StartVertex;  //
+			}
+			GetSeamStitches()[StartSeam + SeamIndex] = MoveTemp(OtherSeamStitches);
+		}
+
+		// LODs Group
+		// Note: It is non intuitive to replace the current LOD's physics or skeleton asset even if it's empty
+	}
+
 	void FCollectionClothLodFacade::SetNumMaterials(int32 InNumMaterials)
 	{
 		GetClothCollection()->SetNumElements(
diff --git a/Engine/Plugins/Experimental/ChaosClothAsset/Source/ChaosClothAsset/Private/ChaosClothAsset/CollectionClothPatternFacade.cpp b/Engine/Plugins/Experimental/ChaosClothAsset/Source/ChaosClothAsset/Private/ChaosClothAsset/CollectionClothPatternFacade.cpp
--- a/Engine/Plugins/Experimental/ChaosClothAsset/Source/ChaosClothAsset/Private/ChaosClothAsset/CollectionClothPatternFacade.cpp
+++ b/Engine/Plugins/Experimental/ChaosClothAsset/Source/ChaosClothAsset/Private/ChaosClothAsset/CollectionClothPatternFacade.cpp
@@ -339,7 +339,7 @@ namespace UE::Chaos::ClothAsset
 	template CHAOSCLOTHASSET_API void FCollectionClothPatternFacade::Initialize(const TArray<FVector2f>& Positions, const TArray<FVector3f>& RestPositions, const TArray<int32>& Indices);
 	template CHAOSCLOTHASSET_API void FCollectionClothPatternFacade::Initialize(const TArray<FVector2f>& Positions, const TArray<FVector3f>& RestPositions, const TArray<uint32>& Indices);
 
-	void FCollectionClothPatternFacade::Initialize(const FCollectionClothPatternConstFacade& Other)
+	void FCollectionClothPatternFacade::Initialize(const FCollectionClothPatternConstFacade& Other, int32 RenderMaterialOffset)
 	{
 		Reset();
 
@@ -440,16 +440,16 @@ namespace UE::Chaos::ClothAsset
 			GetBaseElementIndex(),
 			GetElementIndex());
 		const int32 OtherLodRenderVerticesOffset = Other.ClothCollection->GetElementsOffset(
-			ClothCollection->GetRenderVerticesStart(),
-			GetBaseElementIndex(),
-			GetElementIndex());
+			Other.ClothCollection->GetRenderVerticesStart(),
+			Other.GetBaseElementIndex(),
+			Other.GetElementIndex());
 
 		const FIntVector3 OtherToLodRenderVerticesOffset(LodRenderVerticesOffset - OtherLodRenderVerticesOffset);
 
 		for (int32 RenderFaceIndex = 0; RenderFaceIndex < NumRenderFaces; ++RenderFaceIndex)
 		{
 			GetRenderIndices()[RenderFaceIndex] = Other.GetRenderIndices()[RenderFaceIndex] + OtherToLodRenderVerticesOffset;
-			GetRenderMaterialIndex()[RenderFaceIndex] = Other.GetRenderMaterialIndex()[RenderFaceIndex];
+			GetRenderMaterialIndex()[RenderFaceIndex] = Other.GetRenderMaterialIndex()[RenderFaceIndex] + RenderMaterialOffset;
 		}
 	}
 
diff --git a/Engine/Plugins/Experimental/ChaosClothAsset/Source/ChaosClothAsset/Public/ChaosClothAsset/CollectionClothLodFacade.h b/Engine/Plugins/Experimental/ChaosClothAsset/Source/ChaosClothAsset/Public/ChaosClothAsset/CollectionClothLodFacade.h
--- a/Engine/Plugins/Experimental/ChaosClothAsset/Source/ChaosClothAsset/Public/ChaosClothAsset/CollectionClothLodFacade.h
+++ b/Engine/Plugins/Experimental/ChaosClothAsset/Source/ChaosClothAsset/Public/ChaosClothAsset/CollectionClothLodFacade.h
@@ -139,8 +139,11 @@ namespace UE::Chaos::ClothAsset
 		/** Initialize the cloth LOD using another cloth collection. */
 		void Initialize(const FCollectionClothLodConstFacade& Other);
 
-		/** Initialize the cloth LOD from a DynamicMesh (possibly using UV islands to define Patterns)*/
-		void Initialize(const UE::Geometry::FDynamicMesh3& DynamicMesh, int32 UVChannelIndex);
+		/** Initialize the cloth LOD from a DynamicMesh (possibly using UV islands to define patterns). */
+		void Initialize(const UE::Geometry::FDynamicMesh3& DynamicMesh, int32 UVChannelIndex, const FVector2f& UVScale);
+
+		/** Append to this cloth LOD another cloth LOD. */
+		void Append(const FCollectionClothLodConstFacade& Other);
 
 		/** Set the new number of materials to this cloth LOD. */
 		void SetNumMaterials(int32 NumMaterials);
@@ -212,7 +215,7 @@ namespace UE::Chaos::ClothAsset
 		void SetDefaults();
 
 		template<bool bWeldNearlyCoincidentVertices>
-		void InitializeFromDynamicMeshInternal(const UE::Geometry::FDynamicMesh3& DynamicMesh, int32 UVChannelIndex);
+		void InitializeFromDynamicMeshInternal(const UE::Geometry::FDynamicMesh3& DynamicMesh, int32 UVChannelIndex, const FVector2f& UVScale);
 
 		TSharedPtr<FClothCollection> GetClothCollection() { return ConstCastSharedPtr<FClothCollection>(ClothCollection); }
 	};
diff --git a/Engine/Plugins/Experimental/ChaosClothAsset/Source/ChaosClothAsset/Public/ChaosClothAsset/CollectionClothPatternFacade.h b/Engine/Plugins/Experimental/ChaosClothAsset/Source/ChaosClothAsset/Public/ChaosClothAsset/CollectionClothPatternFacade.h
--- a/Engine/Plugins/Experimental/ChaosClothAsset/Source/ChaosClothAsset/Public/ChaosClothAsset/CollectionClothPatternFacade.h
+++ b/Engine/Plugins/Experimental/ChaosClothAsset/Source/ChaosClothAsset/Public/ChaosClothAsset/CollectionClothPatternFacade.h
@@ -128,8 +128,8 @@ namespace UE::Chaos::ClothAsset
 		template<typename IndexType>
 		void Initialize(const TArray<FVector2f>& Positions, const TArray<FVector3f>& RestPositions, const TArray<IndexType>& Indices);
 
-		/** Initialize the cloth pattern using another pattern. */
-		void Initialize(const FCollectionClothPatternConstFacade& Other);
+		/** Initialize the cloth pattern using another pattern, offsetting the RenderMaterialIndex in the process if needed. */
+		void Initialize(const FCollectionClothPatternConstFacade& Other, int32 RenderMaterialOffset = 0);
 
 		/** Grow or shrink the space reserved for simulation vertices for this pattern within the cloth collection and return its start index. */
 		void SetNumSimVertices(int32 NumSimVertices);
diff --git a/Engine/Plugins/Experimental/ChaosClothAssetEditor/Source/ChaosClothAssetDataflowNodes/Private/ChaosClothAsset/DataflowNodes.cpp b/Engine/Plugins/Experimental/ChaosClothAssetEditor/Source/ChaosClothAssetDataflowNodes/Private/ChaosClothAsset/DataflowNodes.cpp
--- a/Engine/Plugins/Experimental/ChaosClothAssetEditor/Source/ChaosClothAssetDataflowNodes/Private/ChaosClothAsset/DataflowNodes.cpp
+++ b/Engine/Plugins/Experimental/ChaosClothAssetEditor/Source/ChaosClothAssetDataflowNodes/Private/ChaosClothAsset/DataflowNodes.cpp
@@ -7,8 +7,10 @@
 #include "ChaosClothAsset/DatasmithImportNode.h"
 #include "ChaosClothAsset/DeleteRenderMeshNode.h"
 #include "ChaosClothAsset/ImportNode.h"
+#include "ChaosClothAsset/MergeClothCollectionsNode.h"
 #include "ChaosClothAsset/PatternSelectionNode.h"
 #include "ChaosClothAsset/ReverseNormalsNode.h"
+#include "ChaosClothAsset/SetPhysicsAssetNode.h"
 #include "ChaosClothAsset/SimulationAerodynamicsConfigNode.h"
 #include "ChaosClothAsset/SimulationAnimDriveConfigNode.h"
 #include "ChaosClothAsset/SimulationBackstopConfigNode.h"
@@ -52,8 +54,10 @@ namespace UE::Chaos::ClothAsset::DataflowNodes
 		DATAFLOW_NODE_REGISTER_CREATION_FACTORY(FChaosClothAssetDatasmithImportNode);
 		DATAFLOW_NODE_REGISTER_CREATION_FACTORY(FChaosClothAssetDeleteRenderMeshNode);
 		DATAFLOW_NODE_REGISTER_CREATION_FACTORY(FChaosClothAssetImportNode);
+		DATAFLOW_NODE_REGISTER_CREATION_FACTORY(FChaosClothAssetMergeClothCollectionsNode);
 		DATAFLOW_NODE_REGISTER_CREATION_FACTORY(FChaosClothAssetPatternSelectionNode);
 		DATAFLOW_NODE_REGISTER_CREATION_FACTORY(FChaosClothAssetReverseNormalsNode);
+		DATAFLOW_NODE_REGISTER_CREATION_FACTORY(FChaosClothAssetSetPhysicsAssetNode);
 		DATAFLOW_NODE_REGISTER_CREATION_FACTORY(FChaosClothAssetSimulationAerodynamicsConfigNode);
 		DATAFLOW_NODE_REGISTER_CREATION_FACTORY(FChaosClothAssetSimulationAnimDriveConfigNode);
 		DATAFLOW_NODE_REGISTER_CREATION_FACTORY(FChaosClothAssetSimulationBackstopConfigNode);
diff --git a/Engine/Plugins/Experimental/ChaosClothAssetEditor/Source/ChaosClothAssetDataflowNodes/Private/ChaosClothAsset/MergeClothCollectionsNode.cpp b/Engine/Plugins/Experimental/ChaosClothAssetEditor/Source/ChaosClothAssetDataflowNodes/Private/ChaosClothAsset/MergeClothCollectionsNode.cpp
new file mode 100644
--- /dev/null
+++ b/Engine/Plugins/Experimental/ChaosClothAssetEditor/Source/ChaosClothAssetDataflowNodes/Private/ChaosClothAsset/MergeClothCollectionsNode.cpp
@@ -0,0 +1,137 @@
+// Copyright Epic Games, Inc. All Rights Reserved.
+
+#include "ChaosClothAsset/MergeClothCollectionsNode.h"
+#include "ChaosClothAsset/DataflowNodes.h"
+#include "ChaosClothAsset/CollectionClothFacade.h"
+
+#include UE_INLINE_GENERATED_CPP_BY_NAME(MergeClothCollectionsNode)
+
+#define LOCTEXT_NAMESPACE "ChaosClothAssetMergeClothCollectionsNode"
+
+FChaosClothAssetMergeClothCollectionsNode::FChaosClothAssetMergeClothCollectionsNode(const Dataflow::FNodeParameters& InParam, FGuid InGuid)
+	: FDataflowNode(InParam, InGuid)
+{
+	RegisterInputConnection(&Collection);
+	RegisterOutputConnection(&Collection, &Collection);
+}
+
+void FChaosClothAssetMergeClothCollectionsNode::Evaluate(Dataflow::FContext& Context, const FDataflowOutput* Out) const
+{
+	if (Out->IsA<FManagedArrayCollection>(&Collection))
+	{
+		using namespace UE::Chaos::ClothAsset;
+
+		// Evaluate in collection
+		FManagedArrayCollection InCollection = GetValue<FManagedArrayCollection>(Context, &Collection);
+		const TSharedRef<FManagedArrayCollection> ClothCollection = MakeShared<FManagedArrayCollection>(MoveTemp(InCollection));
+		
+		// Make it a valid cloth collection if needed
+		FCollectionClothFacade ClothFacade(ClothCollection);
+		ClothFacade.DefineSchema();
+
+		if (ClothFacade.GetNumLods() < 1)
+		{
+			// Make sure that whatever happens there is always at least one empty LOD to avoid crashing the render data
+			ClothFacade.AddLod();
+		}
+		FCollectionClothLodFacade ClothLodFacade = ClothFacade.GetLod(0);
+
+		// Iterate through the inputs and append them to LOD 0
+		const TArray<const FManagedArrayCollection*> Collections = GetCollections();
+		for (int32 InputIndex = 1; InputIndex < Collections.Num(); ++InputIndex)
+		{
+			const FManagedArrayCollection& OtherCollection = GetValue<FManagedArrayCollection>(Context, Collections[InputIndex]);
+			const TSharedRef<const FManagedArrayCollection> OtherClothCollection = MakeShared<const FManagedArrayCollection>(OtherCollection);
+			const FCollectionClothConstFacade OtherClothFacade(OtherClothCollection);  // TODO: Remove Shared pointer from facade and skip duplication
+
+			if (OtherClothFacade.GetNumLods() > 0)
+			{
+				ClothLodFacade.Append(OtherClothFacade.GetLod(0));
+			}
+		}
+
+		SetValue<FManagedArrayCollection>(Context, *ClothCollection, &Collection);
+	}
+}
+
+Dataflow::FPin FChaosClothAssetMergeClothCollectionsNode::AddPin()
+{
+	auto AddInput = [this](const FManagedArrayCollection* Collection) -> Dataflow::FPin
+	{
+		RegisterInputConnection(Collection);
+		const FDataflowInput* const Input = FindInput(Collection);
+		return { Dataflow::FPin::EDirection::INPUT, Input->GetType(), Input->GetName() };
+	};
+
+	switch (NumInputs)
+	{
+	case 1: ++NumInputs; return AddInput(&Collection1);
+	case 2: ++NumInputs; return AddInput(&Collection2);
+	case 3: ++NumInputs; return AddInput(&Collection3);
+	case 4: ++NumInputs; return AddInput(&Collection4);
+	case 5: ++NumInputs; return AddInput(&Collection5);
+	default: break;
+	}
+
+	return Super::AddPin();
+}
+
+Dataflow::FPin FChaosClothAssetMergeClothCollectionsNode::RemovePin()
+{
+	auto RemoveInput = [this](const FManagedArrayCollection* Collection) -> Dataflow::FPin
+	{
+		const FDataflowInput* const Input = FindInput(Collection);
+		check(Input);
+		Dataflow::FPin Pin = { Dataflow::FPin::EDirection::INPUT, Input->GetType(), Input->GetName() };
+		UnregisterInputConnection(Collection);  // This will delete the input, so set the pin before that
+		return Pin;
+	};
+
+	switch (NumInputs - 1)
+	{
+	case 1: --NumInputs; return RemoveInput(&Collection1);
+	case 2: --NumInputs; return RemoveInput(&Collection2);
+	case 3: --NumInputs; return RemoveInput(&Collection3);
+	case 4: --NumInputs; return RemoveInput(&Collection4);
+	case 5: --NumInputs; return RemoveInput(&Collection5);
+	default: break;
+	}
+	return Super::AddPin();
+}
+
+TArray<const FManagedArrayCollection*> FChaosClothAssetMergeClothCollectionsNode::GetCollections() const
+{
+	TArray<const FManagedArrayCollection*> Collections;
+	Collections.SetNumUninitialized(NumInputs);
+
+	for (int32 InputIndex = 0; InputIndex < NumInputs; ++InputIndex)
+	{
+		switch (InputIndex)
+		{
+		case 0: Collections[InputIndex] = &Collection; break;
+		case 1: Collections[InputIndex] = &Collection1; break;
+		case 2: Collections[InputIndex] = &Collection2; break;
+		case 3: Collections[InputIndex] = &Collection3; break;
+		case 4: Collections[InputIndex] = &Collection4; break;
+		case 5: Collections[InputIndex] = &Collection5; break;
+		default: Collections[InputIndex] = nullptr; check(false); break;
+		}
+	}
+	return Collections;
+}
+
+void FChaosClothAssetMergeClothCollectionsNode::Serialize(FArchive& Ar)
+{
+	if (Ar.IsLoading())
+	{
+		const int32 NumInputsToAdd = NumInputs - 1;
+		NumInputs = 1;  // AddPin will increment it again
+		for (int32 InputIndex = 0; InputIndex < NumInputsToAdd; ++InputIndex)
+		{
+			AddPin();
+		}
+		check(NumInputsToAdd == NumInputs - 1);
+	}
+}
+
+#undef LOCTEXT_NAMESPACE
diff --git a/Engine/Plugins/Experimental/ChaosClothAssetEditor/Source/ChaosClothAssetDataflowNodes/Private/ChaosClothAsset/SetPhysicsAssetNode.cpp b/Engine/Plugins/Experimental/ChaosClothAssetEditor/Source/ChaosClothAssetDataflowNodes/Private/ChaosClothAsset/SetPhysicsAssetNode.cpp
new file mode 100644
--- /dev/null
+++ b/Engine/Plugins/Experimental/ChaosClothAssetEditor/Source/ChaosClothAssetDataflowNodes/Private/ChaosClothAsset/SetPhysicsAssetNode.cpp
@@ -0,0 +1,41 @@
+// Copyright Epic Games, Inc. All Rights Reserved.
+
+#include "ChaosClothAsset/SetPhysicsAssetNode.h"
+#include "ChaosClothAsset/DataflowNodes.h"
+#include "ChaosClothAsset/ClothGeometryTools.h"
+#include "ChaosClothAsset/CollectionClothFacade.h"
+
+#include UE_INLINE_GENERATED_CPP_BY_NAME(SetPhysicsAssetNode)
+
+#define LOCTEXT_NAMESPACE "ChaosClothAssetSetPhysicsAssetNode"
+
+FChaosClothAssetSetPhysicsAssetNode::FChaosClothAssetSetPhysicsAssetNode(const Dataflow::FNodeParameters& InParam, FGuid InGuid)
+	: FDataflowNode(InParam, InGuid)
+{
+	RegisterInputConnection(&Collection);
+	RegisterInputConnection(&PhysicsAsset);
+	RegisterOutputConnection(&Collection, &Collection);
+}
+
+void FChaosClothAssetSetPhysicsAssetNode::Evaluate(Dataflow::FContext& Context, const FDataflowOutput* Out) const
+{
+	if (Out->IsA<FManagedArrayCollection>(&Collection))
+	{
+		using namespace UE::Chaos::ClothAsset;
+
+		// Evaluate in collection
+		FManagedArrayCollection InCollection = GetValue<FManagedArrayCollection>(Context, &Collection);
+		const TSharedRef<FManagedArrayCollection> ClothCollection = MakeShared<FManagedArrayCollection>(MoveTemp(InCollection));
+
+		FCollectionClothFacade CollectionClothFacade(ClothCollection);
+		if (CollectionClothFacade.GetNumLods() > 0)
+		{
+			FCollectionClothLodFacade CollectionClothLodFacade = CollectionClothFacade.GetLod(0);
+			CollectionClothLodFacade.SetPhysicsAssetPathName(PhysicsAsset.GetPathName());
+		}
+
+		SetValue<FManagedArrayCollection>(Context, *ClothCollection, &Collection);
+	}
+}
+
+#undef LOCTEXT_NAMESPACE
diff --git a/Engine/Plugins/Experimental/ChaosClothAssetEditor/Source/ChaosClothAssetDataflowNodes/Private/ChaosClothAsset/StaticMeshImportNode.cpp b/Engine/Plugins/Experimental/ChaosClothAssetEditor/Source/ChaosClothAssetDataflowNodes/Private/ChaosClothAsset/StaticMeshImportNode.cpp
--- a/Engine/Plugins/Experimental/ChaosClothAssetEditor/Source/ChaosClothAssetDataflowNodes/Private/ChaosClothAsset/StaticMeshImportNode.cpp
+++ b/Engine/Plugins/Experimental/ChaosClothAssetEditor/Source/ChaosClothAssetDataflowNodes/Private/ChaosClothAsset/StaticMeshImportNode.cpp
@@ -48,7 +48,11 @@ void FChaosClothAssetStaticMeshImportNode::Evaluate(Dataflow::FContext& Context,
 				Converter.Convert(MeshDescription, DynamicMesh);
 
 				FCollectionClothLodFacade ClothLodFacade = ClothFacade.AddGetLod();
-				ClothLodFacade.Initialize(DynamicMesh, UVChannel);
+				ClothLodFacade.Initialize(DynamicMesh, UVChannel, UVScale);
+
+				// Set a default skeleton
+				static const TCHAR* const DefaultSkeletonPathName = TEXT("/Engine/EditorMeshes/SkeletalMesh/DefaultSkeletalMesh_Skeleton.DefaultSkeletalMesh_Skeleton");
+				ClothLodFacade.SetSkeletonAssetPathName(DefaultSkeletonPathName);
 			}
 			// Make sure that whatever happens there is always at least one empty LOD to avoid crashing the render data
 			if (ClothFacade.GetNumLods() < 1)
diff --git a/Engine/Plugins/Experimental/ChaosClothAssetEditor/Source/ChaosClothAssetDataflowNodes/Public/ChaosClothAsset/MergeClothCollectionsNode.h b/Engine/Plugins/Experimental/ChaosClothAssetEditor/Source/ChaosClothAssetDataflowNodes/Public/ChaosClothAsset/MergeClothCollectionsNode.h
new file mode 100644
--- /dev/null
+++ b/Engine/Plugins/Experimental/ChaosClothAssetEditor/Source/ChaosClothAssetDataflowNodes/Public/ChaosClothAsset/MergeClothCollectionsNode.h
@@ -0,0 +1,52 @@
+// Copyright Epic Games, Inc. All Rights Reserved.
+
+#pragma once 
+
+#include "Dataflow/DataflowNode.h"
+#include "GeometryCollection/ManagedArrayCollection.h"
+#include "MergeClothCollectionsNode.generated.h"
+
+USTRUCT(Meta = (DataflowCloth))
+struct FChaosClothAssetMergeClothCollectionsNode : public FDataflowNode
+{
+	GENERATED_USTRUCT_BODY()
+	DATAFLOW_NODE_DEFINE_INTERNAL(FChaosClothAssetMergeClothCollectionsNode, "MergeClothCollections", "Cloth", "Cloth Merge Collection")
+
+public:
+	static constexpr int32 MaxInputs = 6;  // Hardcoded number of nputs since it is currently not possible to use arrays for optional inputs
+
+	/** Input 0, right click on the node and add pins to add more merge inputs. */
+	UPROPERTY(Meta = (DataflowInput, DataflowOutput, DataflowPassthrough = "Collection"))
+	FManagedArrayCollection Collection;
+	/** Input 1, right click on the node and add pins to add more merge inputs. */
+	UPROPERTY(Meta = (DisplayName = "Collection 1"))
+	FManagedArrayCollection Collection1;
+	/** Input 2, right click on the node and add pins to add more merge inputs. */
+	UPROPERTY(Meta = (DisplayName = "Collection 2"))
+	FManagedArrayCollection Collection2;
+	/** Input 3, right click on the node and add pins to add more merge inputs. */
+	UPROPERTY(Meta = (DisplayName = "Collection 3"))
+	FManagedArrayCollection Collection3;
+	/** Input 4, right click on the node and add pins to add more merge inputs. */
+	UPROPERTY(Meta = (DisplayName = "Collection 4"))
+	FManagedArrayCollection Collection4;
+	/** Input 5, right click on the node and add pins to add more merge inputs. */
+	UPROPERTY(Meta = (DisplayName = "Collection 5"))
+	FManagedArrayCollection Collection5;
+	/** The number of inputs currently exposed to the node UI. */
+	UPROPERTY()
+	int32 NumInputs = 1;
+
+	FChaosClothAssetMergeClothCollectionsNode(const Dataflow::FNodeParameters& InParam, FGuid InGuid = FGuid::NewGuid());
+
+	virtual void Evaluate(Dataflow::FContext& Context, const FDataflowOutput* Out) const override;
+
+	virtual Dataflow::FPin AddPin() override;
+	virtual bool CanAddPin() const override { return NumInputs < MaxInputs; }
+	virtual bool CanRemovePin() const override { return NumInputs > 1; }
+	virtual Dataflow::FPin RemovePin() override;
+	virtual void Serialize(FArchive& Ar);
+
+private:
+	TArray<const FManagedArrayCollection*> GetCollections() const;
+};
diff --git a/Engine/Plugins/Experimental/ChaosClothAssetEditor/Source/ChaosClothAssetDataflowNodes/Public/ChaosClothAsset/SetPhysicsAssetNode.h b/Engine/Plugins/Experimental/ChaosClothAssetEditor/Source/ChaosClothAssetDataflowNodes/Public/ChaosClothAsset/SetPhysicsAssetNode.h
new file mode 100644
--- /dev/null
+++ b/Engine/Plugins/Experimental/ChaosClothAssetEditor/Source/ChaosClothAssetDataflowNodes/Public/ChaosClothAsset/SetPhysicsAssetNode.h
@@ -0,0 +1,29 @@
+// Copyright Epic Games, Inc. All Rights Reserved.
+
+#pragma once 
+
+#include "Dataflow/DataflowNode.h"
+#include "GeometryCollection/ManagedArrayCollection.h"
+#include "SetPhysicsAssetNode.generated.h"
+
+class UPhysicsAsset;
+
+USTRUCT(Meta = (DataflowCloth))
+struct FChaosClothAssetSetPhysicsAssetNode : public FDataflowNode
+{
+	GENERATED_USTRUCT_BODY()
+	DATAFLOW_NODE_DEFINE_INTERNAL(FChaosClothAssetSetPhysicsAssetNode, "SetPhysicsAsset", "Cloth", "Cloth Set Physics Asset")
+
+public:
+
+	UPROPERTY(Meta = (Dataflowinput, DataflowOutput, DisplayName = "Collection", DataflowPassthrough = "Collection"))
+	FManagedArrayCollection Collection;
+
+	/** The physics asset to assign to the Cloth Collection. */
+	UPROPERTY(EditAnywhere, Category = "Physics Asset")
+	TObjectPtr<UPhysicsAsset> PhysicsAsset;
+
+	FChaosClothAssetSetPhysicsAssetNode(const Dataflow::FNodeParameters& InParam, FGuid InGuid = FGuid::NewGuid());
+
+	virtual void Evaluate(Dataflow::FContext& Context, const FDataflowOutput* Out) const override;
+};
diff --git a/Engine/Plugins/Experimental/ChaosClothAssetEditor/Source/ChaosClothAssetDataflowNodes/Public/ChaosClothAsset/StaticMeshImportNode.h b/Engine/Plugins/Experimental/ChaosClothAssetEditor/Source/ChaosClothAssetDataflowNodes/Public/ChaosClothAsset/StaticMeshImportNode.h
--- a/Engine/Plugins/Experimental/ChaosClothAssetEditor/Source/ChaosClothAssetDataflowNodes/Public/ChaosClothAsset/StaticMeshImportNode.h
+++ b/Engine/Plugins/Experimental/ChaosClothAssetEditor/Source/ChaosClothAssetDataflowNodes/Public/ChaosClothAsset/StaticMeshImportNode.h
@@ -18,12 +18,15 @@ public:
 	UPROPERTY(Meta = (DataflowOutput, DisplayName = "Collection"))
 	FManagedArrayCollection Collection;
 
-	UPROPERTY(EditAnywhere, Category = "Static Mesh")
+	UPROPERTY(EditAnywhere, Category = "Static Mesh Import")
 	TObjectPtr<const UStaticMesh> StaticMesh;
 
-	UPROPERTY(EditAnywhere, Category = "Static Mesh")
+	UPROPERTY(EditAnywhere, Category = "Static Mesh Import")
 	int32 UVChannel = INDEX_NONE;
 
+	UPROPERTY(EditAnywhere, Category = "Static Mesh Import", Meta = (AllowPreserveRatio))
+	FVector2f UVScale = { 1.f, 1.f };
+
 	FChaosClothAssetStaticMeshImportNode(const Dataflow::FNodeParameters& InParam, FGuid InGuid = FGuid::NewGuid());
 
 	virtual void Evaluate(Dataflow::FContext& Context, const FDataflowOutput* Out) const override;
