ChaosClothAsset: Add tethers (long range attachments).

Several parts to this:

* Add new CollectionClothTetherBatchFacade for tethers at the cloth collection level. The fields already existed in the schema, so it was just a matter of adding the right methods to the facade.
* Add ability to generate tethers for cloth collection data. This was added to a new utility "ClothEngineTools" (rather than ClothGeometryTools because it has an Engine dependence by reusing the existing ClothTetherData).
* Add creating tethers as part of the default sim config node
* Create new atomic Long Range Attachment sim config node
* Add ClothTetherData to the FChaosClothSimulationLodModel and populate it from the cloth collection
* Access ClothTetherData via the ClothSimulationMesh. NOTE: there is only one type of tether data stored on the new  cloth assets, so the request for geometric vs euclidean tethers is unused.

Also had to add a PostSerialize to the CollectionClothFacade that is run on the ClothAsset to fixup old assets with corrupted tether data.

At Kriss's request, I also deprecated the use of TetherMode in favor of a simple bool bUseGeodesicTethers everywhere.

diff --git a/Engine/Plugins/Experimental/ChaosCloth/Source/Chaos/Private/ChaosCloth/ChaosClothingSimulationCloth.cpp b/Engine/Plugins/Experimental/ChaosCloth/Source/Chaos/Private/ChaosCloth/ChaosClothingSimulationCloth.cpp
--- a/Engine/Plugins/Experimental/ChaosCloth/Source/Chaos/Private/ChaosCloth/ChaosClothingSimulationCloth.cpp
+++ b/Engine/Plugins/Experimental/ChaosCloth/Source/Chaos/Private/ChaosCloth/ChaosClothingSimulationCloth.cpp
@@ -20,7 +20,7 @@ namespace Chaos
 namespace ClothingSimulationClothDefault
 {
 	constexpr int32 MassMode = (int32)FClothingSimulationCloth::EMassMode::Density;
-	constexpr int32 TetherMode = (int32)FClothingSimulationCloth::ETetherMode::Geodesic;
+	constexpr bool bUseGeodesicTethers = true;
 	constexpr float MassValue = 0.35f;
 	constexpr float MinPerParticleMass = 0.0001f;
 	constexpr float CollisionThickness = 1.0f;
@@ -321,7 +321,9 @@ PRAGMA_ENABLE_DEPRECATION_WARNINGS
 	bool /*bInUseThinShellVolumeConstraints*/,  // Deprecated
 	const TVec2<FRealSingle>& InTetherStiffness,
 	const TVec2<FRealSingle>& InTetherScale,
+PRAGMA_DISABLE_DEPRECATION_WARNINGS
 	ETetherMode InTetherMode,
+PRAGMA_ENABLE_DEPRECATION_WARNINGS
 	FRealSingle InMaxDistancesMultiplier,
 	const TVec2<FRealSingle>& InAnimDriveStiffness,
 	const TVec2<FRealSingle>& InAnimDriveDamping,
@@ -418,7 +420,10 @@ PRAGMA_ENABLE_DEPRECATION_WARNINGS
 	// Long range attachment
 	if (InTetherStiffness[0] > 0.f || InTetherStiffness[1] > 0.f)
 	{
-		Properties.AddValue(TEXT("TetherMode"), (int32)InTetherMode);
+		PRAGMA_DISABLE_DEPRECATION_WARNINGS
+		const bool bUseGeodesicTethers = InTetherMode == ETetherMode::Geodesic;
+		PRAGMA_ENABLE_DEPRECATION_WARNINGS
+		Properties.AddValue(TEXT("UseGeodesicTethers"), bUseGeodesicTethers);
 
 		const int32 TetherStiffnessIndex = Properties.AddProperty(TEXT("TetherStiffness"), bEnable , bAnimatable);
 		Properties.SetWeightedValue(TetherStiffnessIndex, InTetherStiffness[0], InTetherStiffness[1]);
@@ -539,8 +544,7 @@ void FClothingSimulationCloth::SetMesh(FClothingSimulationMesh* InMesh)
 	const int32 NumLODs = Mesh ? Mesh->GetNumLODs() : 0;
 	LODData.Reset(NumLODs);
 
-	const ETetherMode TetherMode = (ETetherMode)Config->GetProperties().GetValue<int32>(TEXT("TetherMode"), ClothingSimulationClothDefault::TetherMode);
-	const bool bUseGeodesicTethers = (TetherMode == ETetherMode::Geodesic);
+	const bool bUseGeodesicTethers = Config->GetProperties().GetValue<bool>(TEXT("UseGeodesicTethers"), ClothingSimulationClothDefault::bUseGeodesicTethers);
 	for (int32 LODIndex = 0; LODIndex < NumLODs; ++LODIndex)
 	{
 		// Regenerate LOD weight maps lookup map
diff --git a/Engine/Plugins/Experimental/ChaosCloth/Source/Chaos/Private/ChaosCloth/ChaosClothingSimulationConfig.cpp b/Engine/Plugins/Experimental/ChaosCloth/Source/Chaos/Private/ChaosCloth/ChaosClothingSimulationConfig.cpp
--- a/Engine/Plugins/Experimental/ChaosCloth/Source/Chaos/Private/ChaosCloth/ChaosClothingSimulationConfig.cpp
+++ b/Engine/Plugins/Experimental/ChaosCloth/Source/Chaos/Private/ChaosCloth/ChaosClothingSimulationConfig.cpp
@@ -109,10 +109,7 @@ namespace Chaos
 			// Long range attachment
 			if (ClothConfig->TetherStiffness.Low > 0.f || ClothConfig->TetherStiffness.High > 0.f)
 			{
-				const Softs::FPBDLongRangeConstraints::EMode TetherMode = ClothConfig->bUseGeodesicDistance ?
-					Softs::FPBDLongRangeConstraints::EMode::Geodesic :
-					Softs::FPBDLongRangeConstraints::EMode::Euclidean;
-				Properties->AddValue(TEXT("TetherMode"), (int32)TetherMode);
+				Properties->AddValue(TEXT("UseGeodesicTethers"), ClothConfig->bUseGeodesicDistance);
 
 				const int32 TetherStiffnessIndex = Properties->AddProperty(TEXT("TetherStiffness"), bEnable, bAnimatable);
 				Properties->SetWeightedValue(TetherStiffnessIndex, ClothConfig->TetherStiffness.Low, ClothConfig->TetherStiffness.High);
diff --git a/Engine/Plugins/Experimental/ChaosCloth/Source/Chaos/Public/ChaosCloth/ChaosClothConstraints.h b/Engine/Plugins/Experimental/ChaosCloth/Source/Chaos/Public/ChaosCloth/ChaosClothConstraints.h
--- a/Engine/Plugins/Experimental/ChaosCloth/Source/Chaos/Public/ChaosCloth/ChaosClothConstraints.h
+++ b/Engine/Plugins/Experimental/ChaosCloth/Source/Chaos/Public/ChaosCloth/ChaosClothConstraints.h
@@ -16,7 +16,10 @@ namespace Chaos
 	class FClothConstraints final
 	{
 	public:
+		UE_DEPRECATED(5.3, "ETetherMode has been replaced with bUseGeodesicTethers.")
+PRAGMA_DISABLE_DEPRECATION_WARNINGS
 		typedef Softs::FPBDLongRangeConstraints::EMode ETetherMode;
+PRAGMA_ENABLE_DEPRECATION_WARNINGS
 
 		FClothConstraints();
 		~FClothConstraints();
diff --git a/Engine/Plugins/Experimental/ChaosCloth/Source/Chaos/Public/ChaosCloth/ChaosClothingSimulationCloth.h b/Engine/Plugins/Experimental/ChaosCloth/Source/Chaos/Public/ChaosCloth/ChaosClothingSimulationCloth.h
--- a/Engine/Plugins/Experimental/ChaosCloth/Source/Chaos/Public/ChaosCloth/ChaosClothingSimulationCloth.h
+++ b/Engine/Plugins/Experimental/ChaosCloth/Source/Chaos/Public/ChaosCloth/ChaosClothingSimulationCloth.h
@@ -28,7 +28,10 @@ namespace Chaos
 			Density
 		};
 
+		UE_DEPRECATED(5.3, "ETetherMode has been replaced with bUseGeodesicTethers.")
+PRAGMA_DISABLE_DEPRECATION_WARNINGS
 		typedef FClothConstraints::ETetherMode ETetherMode;
+PRAGMA_ENABLE_DEPRECATION_WARNINGS
 
 		FClothingSimulationCloth(
 			FClothingSimulationConfig* InConfig,
@@ -58,7 +61,9 @@ PRAGMA_ENABLE_DEPRECATION_WARNINGS
 			bool bInUseThinShellVolumeConstraints,
 			const TVec2<FRealSingle>& InTetherStiffness,
 			const TVec2<FRealSingle>& InTetherScale,
+PRAGMA_DISABLE_DEPRECATION_WARNINGS
 			ETetherMode InTetherMode,
+PRAGMA_ENABLE_DEPRECATION_WARNINGS
 			FRealSingle InMaxDistancesMultiplier,
 			const TVec2<FRealSingle>& InAnimDriveStiffness,
 			const TVec2<FRealSingle>& InAnimDriveDamping,
diff --git a/Engine/Plugins/Experimental/ChaosClothAsset/Source/ChaosClothAsset/Private/ChaosClothAsset/ClothGeometryTools.cpp b/Engine/Plugins/Experimental/ChaosClothAsset/Source/ChaosClothAsset/Private/ChaosClothAsset/ClothGeometryTools.cpp
--- a/Engine/Plugins/Experimental/ChaosClothAsset/Source/ChaosClothAsset/Private/ChaosClothAsset/ClothGeometryTools.cpp
+++ b/Engine/Plugins/Experimental/ChaosClothAsset/Source/ChaosClothAsset/Private/ChaosClothAsset/ClothGeometryTools.cpp
@@ -373,4 +373,34 @@ namespace UE::Chaos::ClothAsset
 			}
 		}
 	}
+
+	TArray<float> FClothGeometryTools::BuildWeldedWeightMapForLod(
+		const TSharedPtr<const FManagedArrayCollection>& ClothCollection,
+		int32 LodIndex,
+		const FName& WeightMapName,
+		const TArray<TArray<int32>>& WeldedToUnweldedMap)
+	{
+		FCollectionClothConstFacade ClothFacade(ClothCollection);
+		FCollectionClothLodConstFacade ClothLodFacade = ClothFacade.GetLod(LodIndex);
+
+		TConstArrayView<float> ClothLodWeightMap = ClothLodFacade.GetWeightMap(WeightMapName);
+
+		TArray<float> WeldedWeightMap;
+		WeldedWeightMap.SetNumZeroed(WeldedToUnweldedMap.Num());
+
+		for (int32 WeldedIndex = 0; WeldedIndex < WeldedToUnweldedMap.Num(); ++WeldedIndex)
+		{
+			const TArray<int32>& SourceVertices = WeldedToUnweldedMap[WeldedIndex];
+			check(SourceVertices.Num());
+			for (const int32 SourceVertex : SourceVertices)
+			{
+				WeldedWeightMap[WeldedIndex] += ClothLodWeightMap[SourceVertex];
+			}
+
+			WeldedWeightMap[WeldedIndex] /= (float)SourceVertices.Num();
+		}
+
+		return WeldedWeightMap;
+	}
+
 }  // End namespace UE::Chaos::ClothAsset
diff --git a/Engine/Plugins/Experimental/ChaosClothAsset/Source/ChaosClothAsset/Private/ChaosClothAsset/CollectionClothFacade.cpp b/Engine/Plugins/Experimental/ChaosClothAsset/Source/ChaosClothAsset/Private/ChaosClothAsset/CollectionClothFacade.cpp
--- a/Engine/Plugins/Experimental/ChaosClothAsset/Source/ChaosClothAsset/Private/ChaosClothAsset/CollectionClothFacade.cpp
+++ b/Engine/Plugins/Experimental/ChaosClothAsset/Source/ChaosClothAsset/Private/ChaosClothAsset/CollectionClothFacade.cpp
@@ -3,6 +3,8 @@
 #include "ChaosClothAsset/CollectionClothFacade.h"
 #include "ChaosClothAsset/CollectionClothLodFacade.h"
 #include "ChaosClothAsset/ClothCollection.h"
+#include "Chaos/ChaosArchive.h"
+#include "UObject/UE5MainStreamObjectVersion.h"
 
 namespace UE::Chaos::ClothAsset
 {
@@ -60,6 +62,30 @@ namespace UE::Chaos::ClothAsset
 		GetClothCollection()->SetNumElements(0, FClothCollection::LodsGroup);
 	}
 
+	void FCollectionClothFacade::PostSerialize(const ::Chaos::FChaosArchive& Ar)
+	{
+		if (Ar.CustomVer(FUE5MainStreamObjectVersion::GUID) < FUE5MainStreamObjectVersion::ClothCollectionTetherInitialization)
+		{
+			const int32 NumLods = GetNumLods();
+			for (int32 LodIndex = 0; LodIndex < NumLods; ++LodIndex)
+			{
+				FCollectionClothLodFacade LodFacade = GetLod(LodIndex);
+				const int32 NumTetherBatches = LodFacade.GetNumTetherBatches();
+				for (int32 TetherBatchIndex = 0; TetherBatchIndex < NumTetherBatches; ++TetherBatchIndex)
+				{
+					// The tethers were never set to be empty, so do this manually.
+					const int32 TetherBatchElementIndex = (*ClothCollection->GetTetherBatchStart())[LodIndex] + TetherBatchIndex;
+
+					(*GetClothCollection()->GetTetherStart())[TetherBatchElementIndex] = INDEX_NONE;
+					(*GetClothCollection()->GetTetherEnd())[TetherBatchElementIndex] = INDEX_NONE;
+				}
+
+				// Clear any (now empty) tether batches
+				LodFacade.SetNumTetherBatches(0);
+			}
+		}
+	}
+
 	int32 FCollectionClothFacade::AddLod()
 	{
 		check(IsValid());
diff --git a/Engine/Plugins/Experimental/ChaosClothAsset/Source/ChaosClothAsset/Private/ChaosClothAsset/CollectionClothLodFacade.cpp b/Engine/Plugins/Experimental/ChaosClothAsset/Source/ChaosClothAsset/Private/ChaosClothAsset/CollectionClothLodFacade.cpp
--- a/Engine/Plugins/Experimental/ChaosClothAsset/Source/ChaosClothAsset/Private/ChaosClothAsset/CollectionClothLodFacade.cpp
+++ b/Engine/Plugins/Experimental/ChaosClothAsset/Source/ChaosClothAsset/Private/ChaosClothAsset/CollectionClothLodFacade.cpp
@@ -532,6 +532,11 @@ namespace UE::Chaos::ClothAsset
 		return FCollectionClothPatternConstFacade(ClothCollection, LodIndex, PatternIndex);
 	}
 
+	FCollectionClothTetherBatchConstFacade FCollectionClothLodConstFacade::GetTetherBatch(int32 TetherBatchIndex) const
+	{
+		return FCollectionClothTetherBatchConstFacade(ClothCollection, LodIndex, TetherBatchIndex);
+	}
+
 	const FString& FCollectionClothLodConstFacade::GetPhysicsAssetPathName() const
 	{
 		static const FString EmptyString;
@@ -794,14 +799,14 @@ namespace UE::Chaos::ClothAsset
 			LodIndex);
 	}
 
-	void FCollectionClothLodConstFacade::BuildSimulationMesh(TArray<FVector3f>& Positions, TArray<FVector3f>& Normals, TArray<uint32>& Indices, TArray<int32>& WeldingMap, TArray<FVector2f>& PatternsPositions, TArray<uint32>& PatternsIndices, TArray<uint32>& PatternToWeldedIndices) const
+	void FCollectionClothLodConstFacade::BuildSimulationMesh(TArray<FVector3f>& Positions, TArray<FVector3f>& Normals, TArray<uint32>& Indices, TArray<FVector2f>& PatternsPositions, TArray<uint32>& PatternsIndices, TArray<uint32>& PatternToWeldedIndices, TArray<TArray<int32>>* OptionalWeldedToPatternIndices) const
 	{
 		const int32 NumSimVertices = GetNumSimVertices();
 
 		// Initialize welding map with same index
 		// The welding map redirects to an existing vertex index if these two are part of the same welding group.
 		// The redirected index must be the smallest index in the group.
-
+		TArray<int32> WeldingMap;
 		WeldingMap.SetNumUninitialized(NumSimVertices);
 		for (int32 SimVertexIndex = 0; SimVertexIndex < NumSimVertices; ++SimVertexIndex)
 		{
@@ -892,6 +897,11 @@ namespace UE::Chaos::ClothAsset
 		const TConstArrayView<FVector3f> SimRestNormal = GetSimRestNormal();
 
 		PatternToWeldedIndices.SetNumUninitialized(NumSimVertices);
+		if (OptionalWeldedToPatternIndices)
+		{
+			OptionalWeldedToPatternIndices->Reset(NumWeldedVertices);
+			OptionalWeldedToPatternIndices->SetNum(NumWeldedVertices);
+		}
 
 		uint32 WeldedIndex = 0;
 		for (int32 VertexIndex = 0; VertexIndex < NumSimVertices; ++VertexIndex)
@@ -900,11 +910,20 @@ namespace UE::Chaos::ClothAsset
 			{
 				Positions[WeldedIndex] = SimRestPosition[VertexIndex];
 				Normals[WeldedIndex] = SimRestNormal[VertexIndex];
-				PatternToWeldedIndices[VertexIndex] = WeldedIndex++;
+				PatternToWeldedIndices[VertexIndex] = WeldedIndex;
+				if (OptionalWeldedToPatternIndices)
+				{
+					(*OptionalWeldedToPatternIndices)[WeldedIndex].Add(VertexIndex);
+				}
+				++WeldedIndex;
 			}
 			else
 			{
 				PatternToWeldedIndices[VertexIndex] = PatternToWeldedIndices[WeldingMap[VertexIndex]];
+				if (OptionalWeldedToPatternIndices)
+				{
+					(*OptionalWeldedToPatternIndices)[PatternToWeldedIndices[VertexIndex]].Add(VertexIndex);
+				}
 			}
 		}
 
@@ -1039,11 +1058,11 @@ namespace UE::Chaos::ClothAsset
 		}
 
 		// Tether Batches Group
-		const int32 NumTetherBatches = Other.GetNumPatterns();
+		const int32 NumTetherBatches = Other.GetNumTetherBatches();
 		SetNumTetherBatches(NumTetherBatches);
 		for (int32 TetherBatchIndex = 0; TetherBatchIndex < NumTetherBatches; ++TetherBatchIndex)
 		{
-			// GetTetherBatch(TetherBatchIndex).Initialize(Other.GetTetherBatch(TetherBatchIndex));  // TODO: Tether Batches facade
+			GetTetherBatch(TetherBatchIndex).Initialize(Other.GetTetherBatch(TetherBatchIndex));
 		}
 
 		// Materials Group
@@ -1139,7 +1158,7 @@ namespace UE::Chaos::ClothAsset
 
 		for (int32 TetherBatchIndex = InNumTetherBatches; TetherBatchIndex < NumTetherBatches; ++TetherBatchIndex)
 		{
-			// GetTetherBatch(TetherBatchIndex).Reset();  // TODO: Tether reset
+			GetTetherBatch(TetherBatchIndex).Reset();
 		}
 
 		GetClothCollection()->SetNumElements(
@@ -1151,7 +1170,7 @@ namespace UE::Chaos::ClothAsset
 
 		for (int32 TetherBatchIndex = NumTetherBatches; TetherBatchIndex < InNumTetherBatches; ++TetherBatchIndex)
 		{
-			// GetTetherBatch(TetherBatchIndex).SetDefaults();  // TODO: Tether set default
+			GetTetherBatch(TetherBatchIndex).SetDefaults();
 		}
 	}
 
@@ -1215,6 +1234,18 @@ namespace UE::Chaos::ClothAsset
 		return FCollectionClothPatternFacade(GetClothCollection(), LodIndex, PatternIndex);
 	}
 
+	int32 FCollectionClothLodFacade::AddTetherBatch()
+	{
+		const int32 TetherBatchIndex = GetNumTetherBatches();
+		SetNumTetherBatches(TetherBatchIndex + 1);
+		return TetherBatchIndex;
+	}
+
+	FCollectionClothTetherBatchFacade FCollectionClothLodFacade::GetTetherBatch(int32 TetherBatchIndex)
+	{
+		return FCollectionClothTetherBatchFacade(GetClothCollection(), LodIndex, TetherBatchIndex);
+	}
+
 	void FCollectionClothLodFacade::SetPhysicsAssetPathName(const FString& PhysicsAssetPathName)
 	{
 		(*GetClothCollection()->GetPhysicsAssetPathName())[LodIndex] = PhysicsAssetPathName;
diff --git a/Engine/Plugins/Experimental/ChaosClothAsset/Source/ChaosClothAsset/Private/ChaosClothAsset/CollectionClothTetherBatchFacade.cpp b/Engine/Plugins/Experimental/ChaosClothAsset/Source/ChaosClothAsset/Private/ChaosClothAsset/CollectionClothTetherBatchFacade.cpp
new file mode 100644
--- /dev/null
+++ b/Engine/Plugins/Experimental/ChaosClothAsset/Source/ChaosClothAsset/Private/ChaosClothAsset/CollectionClothTetherBatchFacade.cpp
@@ -0,0 +1,183 @@
+// Copyright Epic Games, Inc. All Rights Reserved.
+
+#include "ChaosClothAsset/CollectionClothTetherBatchFacade.h"
+#include "ChaosClothAsset/ClothCollection.h"
+
+namespace UE::Chaos::ClothAsset
+{
+	int32 FCollectionClothTetherBatchConstFacade::GetNumTethers() const
+	{
+		return ClothCollection->GetNumElements(
+			ClothCollection->GetTetherStart(),
+			ClothCollection->GetTetherEnd(),
+			GetElementIndex());
+	}
+
+	int32 FCollectionClothTetherBatchConstFacade::GetTethersOffset() const
+	{
+		return ClothCollection->GetElementsOffset(
+			ClothCollection->GetTetherStart(),
+			GetBaseElementIndex(),
+			GetElementIndex());
+	}
+
+	TConstArrayView<int32> FCollectionClothTetherBatchConstFacade::GetTetherKinematicIndex() const
+	{
+		return ClothCollection->GetElements(
+			ClothCollection->GetTetherKinematicIndex(),
+			ClothCollection->GetTetherStart(),
+			ClothCollection->GetTetherEnd(),
+			GetElementIndex());
+	}
+
+	TConstArrayView<int32> FCollectionClothTetherBatchConstFacade::GetTetherDynamicIndex() const
+	{
+		return ClothCollection->GetElements(
+			ClothCollection->GetTetherDynamicIndex(),
+			ClothCollection->GetTetherStart(),
+			ClothCollection->GetTetherEnd(),
+			GetElementIndex());
+	}
+
+	TConstArrayView<float> FCollectionClothTetherBatchConstFacade::GetTetherReferenceLength() const
+	{
+		return ClothCollection->GetElements(
+			ClothCollection->GetTetherReferenceLength(),
+			ClothCollection->GetTetherStart(),
+			ClothCollection->GetTetherEnd(),
+			GetElementIndex());
+	}
+
+	TArray<TTuple<int32, int32, float>> FCollectionClothTetherBatchConstFacade::GetZippedTetherData() const
+	{
+		TArray<TTuple<int32, int32, float>> ZippedTethers;
+
+		const int32 NumTethers = GetNumTethers();
+		ZippedTethers.Reserve(NumTethers);
+
+		const TConstArrayView<int32> TetherKinematicIndex = GetTetherKinematicIndex();
+		const TConstArrayView<int32> TetherDynamicIndex = GetTetherDynamicIndex();
+		const TConstArrayView<float> TetherReferenceLength = GetTetherReferenceLength();
+		for (int32 TetherIndex = 0; TetherIndex < NumTethers; ++TetherIndex)
+		{
+			ZippedTethers.Emplace(TetherKinematicIndex[TetherIndex], TetherDynamicIndex[TetherIndex], TetherReferenceLength[TetherIndex]);
+		}
+
+		return ZippedTethers;
+	}
+
+	FCollectionClothTetherBatchConstFacade::FCollectionClothTetherBatchConstFacade(const TSharedPtr<const FClothCollection>& ClothCollection, int32 InLodIndex, int32 TetherBatchIndex)
+		: ClothCollection(ClothCollection)
+		, LodIndex(InLodIndex)
+		, TetherBatchIndex(TetherBatchIndex)
+	{
+		check(ClothCollection.IsValid());
+		check(ClothCollection->IsValid());
+		check(LodIndex >= 0 && LodIndex < ClothCollection->GetNumElements(FClothCollection::LodsGroup));
+		check(TetherBatchIndex >= 0 && TetherBatchIndex < ClothCollection->GetNumElements(ClothCollection->GetTetherBatchStart(), ClothCollection->GetTetherBatchEnd(), LodIndex));
+	}
+
+	int32 FCollectionClothTetherBatchConstFacade::GetBaseElementIndex() const
+	{
+		return (*ClothCollection->GetTetherBatchStart())[LodIndex];
+	}
+
+	void FCollectionClothTetherBatchFacade::Reset()
+	{
+		SetNumTethers(0);
+		SetDefaults();
+	}
+
+	void FCollectionClothTetherBatchFacade::Initialize(const TArray<TTuple<int32, int32, float>>& Tethers)
+	{
+		Reset();
+
+		const int32 NumTethers = Tethers.Num();
+
+		SetNumTethers(NumTethers);
+
+		const TArrayView<int32> TetherKinematicIndex = GetTetherKinematicIndex();
+		const TArrayView<int32> TetherDynamicIndex = GetTetherDynamicIndex();
+		const TArrayView<float> TetherReferenceLength = GetTetherReferenceLength();
+
+		for (int32 TetherIndex = 0; TetherIndex < NumTethers; ++TetherIndex)
+		{
+			TetherKinematicIndex[TetherIndex] = Tethers[TetherIndex].Get<0>();
+			TetherDynamicIndex[TetherIndex] = Tethers[TetherIndex].Get<1>();
+			TetherReferenceLength[TetherIndex] = Tethers[TetherIndex].Get<2>();
+		}
+	}
+
+	void FCollectionClothTetherBatchFacade::Initialize(const FCollectionClothTetherBatchConstFacade& Other)
+	{
+		Reset();
+
+		//~ Tethers Group. 
+		const int32 NumTethers = Other.GetNumTethers();
+		SetNumTethers(NumTethers);
+		
+		const TConstArrayView<int32> OtherTetherKinematicIndex = Other.GetTetherKinematicIndex();
+		const TConstArrayView<int32> OtherTetherDynamicIndex = Other.GetTetherDynamicIndex();
+		const TConstArrayView<float> OtherTetherReferenceLength = Other.GetTetherReferenceLength();
+		const TArrayView<int32> TetherKinematicIndex = GetTetherKinematicIndex();
+		const TArrayView<int32> TetherDynamicIndex = GetTetherDynamicIndex();
+		const TArrayView<float> TetherReferenceLength = GetTetherReferenceLength();
+
+		for (int32 TetherIndex = 0; TetherIndex < NumTethers; ++TetherIndex)
+		{
+			TetherKinematicIndex[TetherIndex] = OtherTetherKinematicIndex[TetherIndex];
+			TetherDynamicIndex[TetherIndex] = OtherTetherDynamicIndex[TetherIndex];
+			TetherReferenceLength[TetherIndex] = OtherTetherReferenceLength[TetherIndex];
+		}
+	}
+
+	void FCollectionClothTetherBatchFacade::SetNumTethers(int32 NumTethers)
+	{
+		GetClothCollection()->SetNumElements(
+			NumTethers,
+			FClothCollection::TethersGroup,
+			GetClothCollection()->GetTetherStart(),
+			GetClothCollection()->GetTetherEnd(),
+			GetElementIndex());
+	}
+
+	TArrayView<int32> FCollectionClothTetherBatchFacade::GetTetherKinematicIndex()
+	{
+		return GetClothCollection()->GetElements(
+			GetClothCollection()->GetTetherKinematicIndex(),
+			GetClothCollection()->GetTetherStart(),
+			GetClothCollection()->GetTetherEnd(),
+			GetElementIndex());
+	}
+
+	TArrayView<int32> FCollectionClothTetherBatchFacade::GetTetherDynamicIndex()
+	{
+		return GetClothCollection()->GetElements(
+			GetClothCollection()->GetTetherDynamicIndex(),
+			GetClothCollection()->GetTetherStart(),
+			GetClothCollection()->GetTetherEnd(),
+			GetElementIndex());
+	}
+
+	TArrayView<float> FCollectionClothTetherBatchFacade::GetTetherReferenceLength()
+	{
+		return GetClothCollection()->GetElements(
+			GetClothCollection()->GetTetherReferenceLength(),
+			GetClothCollection()->GetTetherStart(),
+			GetClothCollection()->GetTetherEnd(),
+			GetElementIndex());
+	}
+
+	FCollectionClothTetherBatchFacade::FCollectionClothTetherBatchFacade(const TSharedPtr<FClothCollection>& ClothCollection, int32 InLodIndex, int32 InTetherIndex)
+		: FCollectionClothTetherBatchConstFacade(ClothCollection, InLodIndex, InTetherIndex)
+	{
+	}
+
+	void FCollectionClothTetherBatchFacade::SetDefaults()
+	{
+		const int32 ElementIndex = GetElementIndex();
+
+		(*GetClothCollection()->GetTetherStart())[ElementIndex] = INDEX_NONE;
+		(*GetClothCollection()->GetTetherEnd())[ElementIndex] = INDEX_NONE;
+	}
+}  // End namespace UE::Chaos::ClothAsset
diff --git a/Engine/Plugins/Experimental/ChaosClothAsset/Source/ChaosClothAsset/Public/ChaosClothAsset/ClothGeometryTools.h b/Engine/Plugins/Experimental/ChaosClothAsset/Source/ChaosClothAsset/Public/ChaosClothAsset/ClothGeometryTools.h
--- a/Engine/Plugins/Experimental/ChaosClothAsset/Source/ChaosClothAsset/Public/ChaosClothAsset/ClothGeometryTools.h
+++ b/Engine/Plugins/Experimental/ChaosClothAsset/Source/ChaosClothAsset/Public/ChaosClothAsset/ClothGeometryTools.h
@@ -50,5 +50,14 @@ namespace UE::Chaos::ClothAsset
 			bool bBindSimMesh,
 			bool bBindRenderMesh,
 			const TArray<int32> Lods = {});
+
+		/**
+		* Generate a welded weight map.
+		*/
+		static TArray<float> BuildWeldedWeightMapForLod(
+			const TSharedPtr<const FManagedArrayCollection>& ClothCollection,
+			int32 LodIndex,
+			const FName& WeightMapName,
+			const TArray<TArray<int32>>& WeldedToUnweldedMap);
 	};
 }  // End namespace UE::Chaos::ClothAsset
\ No newline at end of file
diff --git a/Engine/Plugins/Experimental/ChaosClothAsset/Source/ChaosClothAsset/Public/ChaosClothAsset/CollectionClothFacade.h b/Engine/Plugins/Experimental/ChaosClothAsset/Source/ChaosClothAsset/Public/ChaosClothAsset/CollectionClothFacade.h
--- a/Engine/Plugins/Experimental/ChaosClothAsset/Source/ChaosClothAsset/Public/ChaosClothAsset/CollectionClothFacade.h
+++ b/Engine/Plugins/Experimental/ChaosClothAsset/Source/ChaosClothAsset/Public/ChaosClothAsset/CollectionClothFacade.h
@@ -5,6 +5,10 @@
 #include "ChaosClothAsset/CollectionClothLodFacade.h"
 
 class FClothCollection;
+namespace Chaos
+{
+class FChaosArchive;
+}
 
 namespace UE::Chaos::ClothAsset
 {
@@ -70,6 +74,9 @@ namespace UE::Chaos::ClothAsset
 		/** Remove all LODs from this cloth. */
 		void Reset();
 
+		/** Post serialize function. Use to upgrade between versions.*/
+		void PostSerialize(const ::Chaos::FChaosArchive& Ar);
+
 		/** Add a new LOD to this cloth. */
 		int32 AddLod();
 
diff --git a/Engine/Plugins/Experimental/ChaosClothAsset/Source/ChaosClothAsset/Public/ChaosClothAsset/CollectionClothLodFacade.h b/Engine/Plugins/Experimental/ChaosClothAsset/Source/ChaosClothAsset/Public/ChaosClothAsset/CollectionClothLodFacade.h
--- a/Engine/Plugins/Experimental/ChaosClothAsset/Source/ChaosClothAsset/Public/ChaosClothAsset/CollectionClothLodFacade.h
+++ b/Engine/Plugins/Experimental/ChaosClothAsset/Source/ChaosClothAsset/Public/ChaosClothAsset/CollectionClothLodFacade.h
@@ -3,6 +3,7 @@
 #pragma once
 
 #include "ChaosClothAsset/CollectionClothPatternFacade.h"
+#include "ChaosClothAsset/CollectionClothTetherBatchFacade.h"
 
 namespace UE::Geometry
 {
@@ -51,6 +52,9 @@ namespace UE::Chaos::ClothAsset
 		/** Return a pattern facade for the specified pattern index. */
 		FCollectionClothPatternConstFacade GetPattern(int32 PatternIndex) const;
 
+		/** Return a tether batch facade for the specified tether batch index. */
+		FCollectionClothTetherBatchConstFacade GetTetherBatch(int32 TetherBatchIndex) const;
+
 		/** Return the physics asset path names used for this LOD. */
 		const FString& GetPhysicsAssetPathName() const;
 		/** Return the skeleton asset path names used for this LOD. */
@@ -101,8 +105,9 @@ namespace UE::Chaos::ClothAsset
 		int32 GetLodIndex() const { return LodIndex; }
 
 		/** Return the welded simulation mesh for this LOD. */
-		void BuildSimulationMesh(TArray<FVector3f>& Positions, TArray<FVector3f>& Normals, TArray<uint32>& Indices, TArray<int32>& WeldingMap,
-			TArray<FVector2f>& PatternsPositions, TArray<uint32>& PatternsIndices, TArray<uint32>& PatternToWeldedIndices) const;
+		void BuildSimulationMesh(TArray<FVector3f>& Positions, TArray<FVector3f>& Normals, TArray<uint32>& Indices,
+			TArray<FVector2f>& PatternsPositions, TArray<uint32>& PatternsIndices, TArray<uint32>& PatternToWeldedIndices, 
+			TArray<TArray<int32>>* OptionalWeldedToPatternIndices = nullptr) const;
 
 	protected:
 		friend class FCollectionClothConstFacade;
@@ -170,6 +175,15 @@ namespace UE::Chaos::ClothAsset
 		/** Add a new pattern to this cloth LOD, and return the cloth pattern facade set to its index. */
 		FCollectionClothPatternFacade AddGetPattern() { return GetPattern(AddPattern()); }
 
+		/** Add a new tether batch to this cloth LOD and return its index in the LOD tether batch list. */
+		int32 AddTetherBatch();
+
+		/** Return a tether batch facade for the specified tether batch index. */
+		FCollectionClothTetherBatchFacade GetTetherBatch(int32 TetherBatchIndex);
+
+		/** Add a new tether batch to this cloth LOD, and return the tether batch facade set to its index. */
+		FCollectionClothTetherBatchFacade AddGetTetherBatch() { return GetTetherBatch(AddTetherBatch()); }
+
 		/** Set the physics asset path names used for this LOD. */
 		void SetPhysicsAssetPathName(const FString& PathName);
 		/** Set the skeleton asset path names used for this LOD. */
diff --git a/Engine/Plugins/Experimental/ChaosClothAsset/Source/ChaosClothAsset/Public/ChaosClothAsset/CollectionClothTetherBatchFacade.h b/Engine/Plugins/Experimental/ChaosClothAsset/Source/ChaosClothAsset/Public/ChaosClothAsset/CollectionClothTetherBatchFacade.h
new file mode 100644
--- /dev/null
+++ b/Engine/Plugins/Experimental/ChaosClothAsset/Source/ChaosClothAsset/Public/ChaosClothAsset/CollectionClothTetherBatchFacade.h
@@ -0,0 +1,109 @@
+// Copyright Epic Games, Inc. All Rights Reserved.
+
+#pragma once
+
+#include "Templates/SharedPointer.h"
+
+struct FManagedArrayCollection;
+
+namespace UE::Chaos::ClothAsset
+{
+	class FClothCollection;
+
+	/**
+	 * Cloth Asset collection tether batch facade class to access cloth tether batch data.
+	 * Constructed from FCollectionClothLodConstFacade.
+	 * Const access (read only) version.
+	 */
+	class CHAOSCLOTHASSET_API FCollectionClothTetherBatchConstFacade
+	{
+	public:
+		FCollectionClothTetherBatchConstFacade() = delete;
+
+		FCollectionClothTetherBatchConstFacade(const FCollectionClothTetherBatchConstFacade&) = delete;
+		FCollectionClothTetherBatchConstFacade& operator=(const FCollectionClothTetherBatchConstFacade&) = delete;
+
+		FCollectionClothTetherBatchConstFacade(FCollectionClothTetherBatchConstFacade&&) = default;
+		FCollectionClothTetherBatchConstFacade& operator=(FCollectionClothTetherBatchConstFacade&&) = default;
+
+		virtual ~FCollectionClothTetherBatchConstFacade() = default;
+
+		/** Return the total number of tethers for this batch. */
+		int32 GetNumTethers() const;
+
+		/** Return the tethers offset for this batch in the tethers for the LOD. */
+		int32 GetTethersOffset() const;
+
+		//~ Tethers Group. 
+		// Indices correspond with the FCollectionClothLodConstFacade indices
+		TConstArrayView<int32> GetTetherKinematicIndex() const;
+		TConstArrayView<int32> GetTetherDynamicIndex() const;
+		TConstArrayView<float> GetTetherReferenceLength() const;
+
+		/** Return the LOD index this facade has been created with. */
+		int32 GetLodIndex() const { return LodIndex; }
+
+		/** Return the tether batch index this facade has been created with. */
+		int32 GetTetherBatchIndex() const { return TetherBatchIndex; }
+
+		/** Cloth data expects tethers to be zipped into tuples. This is a convenience function to generate those.*/
+		TArray<TTuple<int32, int32, float>> GetZippedTetherData() const;
+
+	protected:
+		friend class FCollectionClothTetherBatchFacade;  // For other instances access
+		friend class FCollectionClothLodConstFacade;
+		FCollectionClothTetherBatchConstFacade(const TSharedPtr<const FClothCollection>& InClothCollection, int32 InLodIndex, int32 TetherBatchIndex);
+
+		int32 GetBaseElementIndex() const;
+		int32 GetElementIndex() const { return GetBaseElementIndex() + TetherBatchIndex; }
+
+		TSharedPtr<const FClothCollection> ClothCollection;
+		int32 LodIndex;
+		int32 TetherBatchIndex;
+	};
+
+	/**
+	 * Cloth Asset collection tether batch facade class to access cloth tether batch data.
+	 * Constructed from FCollectionClothLodFacade.
+	 * Non-const access (read/write) version.
+	 */
+	class CHAOSCLOTHASSET_API FCollectionClothTetherBatchFacade final : public FCollectionClothTetherBatchConstFacade
+	{
+	public:
+		FCollectionClothTetherBatchFacade() = delete;
+
+		FCollectionClothTetherBatchFacade(const FCollectionClothTetherBatchFacade&) = delete;
+		FCollectionClothTetherBatchFacade& operator=(const FCollectionClothTetherBatchFacade&) = delete;
+
+		FCollectionClothTetherBatchFacade(FCollectionClothTetherBatchFacade&&) = default;
+		FCollectionClothTetherBatchFacade& operator=(FCollectionClothTetherBatchFacade&&) = default;
+
+		virtual ~FCollectionClothTetherBatchFacade() override = default;
+
+		/** Remove all tethers from this tether batch. */
+		void Reset();
+
+		/** Initialize the tether batch based on zipped tether tuple data. */
+		void Initialize(const TArray<TTuple<int32, int32, float>>& Tethers);
+
+		/** Initialize the tether batch using another tether batch. */
+		void Initialize(const FCollectionClothTetherBatchConstFacade& Other);
+
+		/** Grow or shrink the space reserved for tethers for this tether batch within the cloth collection. */
+		void SetNumTethers(int32 NumTethers);
+
+		//~ Tethers Group. 
+		// Indices correspond with the FCollectionClothLodConstFacade indices
+		TArrayView<int32> GetTetherKinematicIndex();
+		TArrayView<int32> GetTetherDynamicIndex();
+		TArrayView<float> GetTetherReferenceLength();
+
+	protected:
+		friend class FCollectionClothLodFacade;
+		FCollectionClothTetherBatchFacade(const TSharedPtr<FClothCollection>& InClothCollection, int32 InLodIndex, int32 InTetherIndex);
+
+		void SetDefaults();
+
+		TSharedPtr<FClothCollection> GetClothCollection() { return ConstCastSharedPtr<FClothCollection>(ClothCollection); }
+	};
+}  // End namespace UE::Chaos::ClothAsset
diff --git a/Engine/Plugins/Experimental/ChaosClothAsset/Source/ChaosClothAssetEngine/Private/ChaosClothAsset/ClothAsset.cpp b/Engine/Plugins/Experimental/ChaosClothAsset/Source/ChaosClothAssetEngine/Private/ChaosClothAsset/ClothAsset.cpp
--- a/Engine/Plugins/Experimental/ChaosClothAsset/Source/ChaosClothAssetEngine/Private/ChaosClothAsset/ClothAsset.cpp
+++ b/Engine/Plugins/Experimental/ChaosClothAsset/Source/ChaosClothAssetEngine/Private/ChaosClothAsset/ClothAsset.cpp
@@ -89,6 +89,8 @@ void UChaosClothAsset::Serialize(FArchive& Ar)
 
 	Chaos::FChaosArchive ChaosArchive(Ar);
 	ClothCollection->Serialize(ChaosArchive);
+	UE::Chaos::ClothAsset::FCollectionClothFacade ClothFacade(ClothCollection);
+	ClothFacade.PostSerialize(ChaosArchive);
 
 	Ar << GetRefSkeleton();
 	if (Ar.IsLoading())
diff --git a/Engine/Plugins/Experimental/ChaosClothAsset/Source/ChaosClothAssetEngine/Private/ChaosClothAsset/ClothEngineTools.cpp b/Engine/Plugins/Experimental/ChaosClothAsset/Source/ChaosClothAssetEngine/Private/ChaosClothAsset/ClothEngineTools.cpp
new file mode 100644
--- /dev/null
+++ b/Engine/Plugins/Experimental/ChaosClothAsset/Source/ChaosClothAssetEngine/Private/ChaosClothAsset/ClothEngineTools.cpp
@@ -0,0 +1,56 @@
+// Copyright Epic Games, Inc. All Rights Reserved.
+
+#include "ChaosClothAsset/ClothEngineTools.h"
+#include "ChaosClothAsset/CollectionClothFacade.h"
+#include "ChaosClothAsset/ClothCollection.h"
+#include "ChaosClothAsset/ClothGeometryTools.h"
+#include "ClothTetherData.h"
+
+namespace UE::Chaos::ClothAsset
+{
+
+void FClothEngineTools::GenerateTethers(const TSharedPtr<FManagedArrayCollection>& ClothCollection, const FName& WeightMap, const bool bGenerateGeodesicTethers)
+{
+	FCollectionClothFacade ClothFacade(ClothCollection);
+	if (ClothFacade.HasWeightMap(WeightMap))
+	{
+		for (int32 LodIndex = 0; LodIndex < ClothFacade.GetNumLods(); ++LodIndex)
+		{
+			FCollectionClothLodFacade ClothLodFacade = ClothFacade.GetLod(LodIndex);
+
+			// We want to create the tethers based on the welded mesh, so create that now
+			TArray<FVector3f> Positions;
+			TArray<FVector3f> Normals;
+			TArray<uint32> Indices;
+			TArray<FVector2f> PatternsPositions;
+			TArray<uint32> PatternsIndices;
+			TArray<uint32> PatternToWeldedIndices;
+			TArray<TArray<int32>> WeldedToPatternIndices;
+			ClothLodFacade.BuildSimulationMesh(Positions, Normals, Indices, PatternsPositions, PatternsIndices, PatternToWeldedIndices, &WeldedToPatternIndices);
+
+			const TArray<float> WeldedWeightMap = FClothGeometryTools::BuildWeldedWeightMapForLod(ClothCollection, LodIndex, WeightMap, WeldedToPatternIndices);
+
+			FClothTetherData TetherData;
+			TetherData.GenerateTethers(TConstArrayView<FVector3f>(Positions), TConstArrayView<uint32>(Indices), TConstArrayView<float>(WeldedWeightMap), bGenerateGeodesicTethers);
+
+			// Set new tethers
+			const int32 NumTetherBatches = TetherData.Tethers.Num();
+			ClothLodFacade.SetNumTetherBatches(NumTetherBatches);
+			for (int32 TetherBatchIndex = 0; TetherBatchIndex < NumTetherBatches; ++TetherBatchIndex)
+			{
+				TArray<TTuple<int32, int32, float>>& TetherBatch = TetherData.Tethers[TetherBatchIndex];
+				// Remap tether data indices back to unwelded indices. Just set them to the first of the unwelded index if that welded index is formed from multiple unwelded indices.
+				for (TTuple<int32, int32, float>& Tether : TetherBatch)
+				{
+					Tether.Get<0>() = WeldedToPatternIndices[Tether.Get<0>()][0];
+					Tether.Get<1>() = WeldedToPatternIndices[Tether.Get<1>()][0];
+				}
+
+				FCollectionClothTetherBatchFacade TetherBatchFacade = ClothLodFacade.GetTetherBatch(TetherBatchIndex);
+				TetherBatchFacade.Initialize(TetherBatch);
+			}
+		}
+	}
+}
+
+}  // End namespace UE::Chaos::ClothAsset
diff --git a/Engine/Plugins/Experimental/ChaosClothAsset/Source/ChaosClothAssetEngine/Private/ChaosClothAsset/ClothSimulationMesh.cpp b/Engine/Plugins/Experimental/ChaosClothAsset/Source/ChaosClothAssetEngine/Private/ChaosClothAsset/ClothSimulationMesh.cpp
--- a/Engine/Plugins/Experimental/ChaosClothAsset/Source/ChaosClothAssetEngine/Private/ChaosClothAsset/ClothSimulationMesh.cpp
+++ b/Engine/Plugins/Experimental/ChaosClothAsset/Source/ChaosClothAssetEngine/Private/ChaosClothAsset/ClothSimulationMesh.cpp
@@ -123,9 +123,9 @@ PRAGMA_ENABLE_DEPRECATION_WARNINGS
 		return WeightMaps;
 	}
 
-	TArray<TConstArrayView<TTuple<int32, int32, float>>> FClothSimulationMesh::GetTethers(int32 LODIndex, bool bUseGeodesicTethers) const
+	TArray<TConstArrayView<TTuple<int32, int32, float>>> FClothSimulationMesh::GetTethers(int32 LODIndex, bool /*bUseGeodesicTethers*/) const
 	{
-		return TArray<TConstArrayView<TTuple<int32, int32, float>>>();  // TODO: Tethers
+		return ClothSimulationModel.GetTethers(LODIndex);
 	}
 
 	int32 FClothSimulationMesh::GetReferenceBoneIndex() const
diff --git a/Engine/Plugins/Experimental/ChaosClothAsset/Source/ChaosClothAssetEngine/Private/ChaosClothAsset/ClothSimulationMesh.h b/Engine/Plugins/Experimental/ChaosClothAsset/Source/ChaosClothAssetEngine/Private/ChaosClothAsset/ClothSimulationMesh.h
--- a/Engine/Plugins/Experimental/ChaosClothAsset/Source/ChaosClothAssetEngine/Private/ChaosClothAsset/ClothSimulationMesh.h
+++ b/Engine/Plugins/Experimental/ChaosClothAsset/Source/ChaosClothAssetEngine/Private/ChaosClothAsset/ClothSimulationMesh.h
@@ -41,7 +41,8 @@ PRAGMA_ENABLE_DEPRECATION_WARNINGS
 		virtual TArray<FName> GetWeightMapNames() const override;
 		virtual TMap<FString, int32> GetWeightMapIndices() const override;
 		virtual TArray<TConstArrayView<::Chaos::FRealSingle>> GetWeightMaps(int32 LODIndex) const override;
-		virtual TArray<TConstArrayView<TTuple<int32, int32, float>>> GetTethers(int32 LODIndex, bool bUseGeodesicTethers) const override;
+		// Note: there is only one set of tethers stored on ClothSimulationMesh assets
+		virtual TArray<TConstArrayView<TTuple<int32, int32, float>>> GetTethers(int32 LODIndex, bool /*bUseGeodesicTethers*/) const override;
 		virtual int32 GetReferenceBoneIndex() const override;
 		virtual FTransform GetReferenceBoneTransform() const override;
 		virtual const TArray<FTransform>& GetBoneTransforms() const override;
diff --git a/Engine/Plugins/Experimental/ChaosClothAsset/Source/ChaosClothAssetEngine/Private/ChaosClothAsset/ClothSimulationModel.cpp b/Engine/Plugins/Experimental/ChaosClothAsset/Source/ChaosClothAssetEngine/Private/ChaosClothAsset/ClothSimulationModel.cpp
--- a/Engine/Plugins/Experimental/ChaosClothAsset/Source/ChaosClothAssetEngine/Private/ChaosClothAsset/ClothSimulationModel.cpp
+++ b/Engine/Plugins/Experimental/ChaosClothAsset/Source/ChaosClothAssetEngine/Private/ChaosClothAsset/ClothSimulationModel.cpp
@@ -2,6 +2,7 @@
 
 #include "ChaosClothAsset/ClothSimulationModel.h"
 #include "ChaosClothAsset/CollectionClothFacade.h"
+#include "ChaosClothAsset/ClothGeometryTools.h"
 #include "ReferenceSkeleton.h"
 
 namespace UE::Chaos::ClothAsset
@@ -103,35 +104,6 @@ namespace UE::Chaos::ClothAsset
 		}
 		return ReferenceBoneIndex;
 	}
-
-	// TODO: Check which welding implementation is best and if we should use this one instead
-	// Welded weights are averaged. TODO: add other welding options such as Max and Min
-	static TArray<float> WeldWeightMap(const TConstArrayView<float>& UnweldedMap, const TArray<uint32>& PatternToWeldedIndices, const int32 NumWelded)
-	{
-		TArray<float> WeldedMap;
-		WeldedMap.SetNumZeroed(NumWelded);
-
-		TArray<float> NumInfluences;
-		NumInfluences.SetNumZeroed(NumWelded);
-		
-		check(UnweldedMap.Num() == PatternToWeldedIndices.Num());
-		for (int32 OrigIdx = 0; OrigIdx < UnweldedMap.Num(); ++OrigIdx)
-		{
-			const int32 WeldedIdx = PatternToWeldedIndices[OrigIdx];
-			NumInfluences[WeldedIdx] += 1.f;
-			WeldedMap[WeldedIdx] += UnweldedMap[OrigIdx];
-		}
-
-		for (int32 WeldedIdx = 0; WeldedIdx < NumWelded; ++WeldedIdx)
-		{
-			if (NumInfluences[WeldedIdx] > 0.f)
-			{
-				WeldedMap[WeldedIdx] /= NumInfluences[WeldedIdx];
-			}
-		}
-		return WeldedMap;
-	}
-
 }  // End namespace UE::Chaos::ClothAsset
 
 bool FChaosClothSimulationLodModel::Serialize(FArchive& Ar)
@@ -167,32 +139,13 @@ FChaosClothSimulationModel::FChaosClothSimulationModel(const TSharedPtr<const FM
 	{
 		FChaosClothSimulationLodModel& LodModel = ClothSimulationLodModels[LodIndex];
 		const FCollectionClothLodConstFacade ClothLod = Cloth.GetLod(LodIndex);
-		TArray<int32> WeldingMap;
-		ClothLod.BuildSimulationMesh(LodModel.Positions, LodModel.Normals, LodModel.Indices, WeldingMap, LodModel.PatternPositions, LodModel.PatternIndices, LodModel.PatternToWeldedIndices);
 
-		// Build welding index lookup  // TODO: Could probably do this better in the BuildSimulationMesh function
+		TArray<TArray<int32>> WeldedToPatternIndices;
+		ClothLod.BuildSimulationMesh(LodModel.Positions, LodModel.Normals, LodModel.Indices, LodModel.PatternPositions, LodModel.PatternIndices, LodModel.PatternToWeldedIndices, &WeldedToPatternIndices);
+
 		const int32 LodModelNumSimVertices = LodModel.Positions.Num();
 		const int32 ClothLodNumSimVertices = ClothLod.GetNumSimVertices();
 
-		TArray<int32> ReverseWeldingMap;
-		ReverseWeldingMap.Reserve(LodModelNumSimVertices);
-		TMap<int32, TArray<int32>> SourceVertexIndicesMap;
-		SourceVertexIndicesMap.Reserve(ClothLodNumSimVertices);
-			
-		for (int32 VertexIndex = 0; VertexIndex < ClothLodNumSimVertices; ++VertexIndex)
-		{
-			if (WeldingMap[VertexIndex] == VertexIndex)
-			{
-				SourceVertexIndicesMap.Emplace(VertexIndex, { VertexIndex } );
-				ReverseWeldingMap.Emplace(VertexIndex);
-			}
-			else
-			{
-				TArray<int32>& SourceVertexIndices = SourceVertexIndicesMap.FindChecked(WeldingMap[VertexIndex]);
-				SourceVertexIndices.Emplace(VertexIndex);
-			}
-		}
-
 		// Copy and weld (average) weight maps
 		LodModel.WeightMaps.Reserve(WeightMapNames.Num());
 
@@ -201,21 +154,7 @@ FChaosClothSimulationModel::FChaosClothSimulationModel(const TSharedPtr<const FM
 			const TConstArrayView<float> ClothLodWeightMap = ClothLod.GetWeightMap(WeightMapName);
 			TArray<float>& LodModelWeightMap = LodModel.WeightMaps.Add(WeightMapName);
 
-			LodModelWeightMap.SetNumZeroed(LodModelNumSimVertices);
-
-			for (int32 WeldedIndex = 0; WeldedIndex < LodModelNumSimVertices; ++WeldedIndex)
-			{
-				const int32 BaseSourceVertexIndex = ReverseWeldingMap[WeldedIndex];
-				const TArray<int32>& SourceVertexIndices = SourceVertexIndicesMap.FindChecked(BaseSourceVertexIndex);
-				check(SourceVertexIndices.Num());
-
-				for (const int32 VertexIndex : SourceVertexIndices)
-				{
-					LodModelWeightMap[WeldedIndex] += ClothLodWeightMap[VertexIndex];
-				}
-
-				LodModelWeightMap[WeldedIndex] /= (float)SourceVertexIndices.Num();
-			}
+			LodModelWeightMap = FClothGeometryTools::BuildWeldedWeightMapForLod(ClothCollection, LodIndex, WeightMapName, WeldedToPatternIndices);
 		}
 
 		// Weld bone influences
@@ -226,7 +165,10 @@ FChaosClothSimulationModel::FChaosClothSimulationModel(const TSharedPtr<const FM
 
 		for (int32 WeldedIndex = 0; WeldedIndex < LodModelNumSimVertices; ++WeldedIndex)
 		{
-			const int32 SourceVertexIndex = ReverseWeldingMap[WeldedIndex];
+			const TArray<int32>& SourceVertices = WeldedToPatternIndices[WeldedIndex];
+			check(SourceVertices.Num());
+			const int32 SourceVertexIndex = SourceVertices[0]; // Just using first bone's data.
+
 			LodModel.BoneData[WeldedIndex].NumInfluences = NumBoneInfluences[SourceVertexIndex];
 			for (int32 BoneIndex = 0; BoneIndex < LodModel.BoneData[WeldedIndex].NumInfluences; ++BoneIndex)
 			{
@@ -234,6 +176,22 @@ FChaosClothSimulationModel::FChaosClothSimulationModel(const TSharedPtr<const FM
 				LodModel.BoneData[WeldedIndex].BoneWeights[BoneIndex] = SimBoneWeights[SourceVertexIndex][BoneIndex];
 			}
 		}
+
+		// Weld tethers. This will not deduplicate any tethers that may be duplicated by welding. 
+		// This would currently only occur if seams were added after tethers were created in the dataflow graph.
+		const int32 NumTetherBatches = ClothLod.GetNumTetherBatches();
+		LodModel.TetherData.Tethers.Reserve(NumTetherBatches);
+		for (int32 TetherBatchIndex = 0; TetherBatchIndex < NumTetherBatches; ++TetherBatchIndex)
+		{
+			FCollectionClothTetherBatchConstFacade TetherBatchFacade = ClothLod.GetTetherBatch(TetherBatchIndex);
+			TArray<TTuple<int32, int32, float>>& Tethers = LodModel.TetherData.Tethers.Emplace_GetRef(TetherBatchFacade.GetZippedTetherData());
+
+			for (TTuple<int32, int32, float>& Tether : Tethers)
+			{
+				Tether.Get<0>() = LodModel.PatternToWeldedIndices[Tether.Get<0>()];
+				Tether.Get<1>() = LodModel.PatternToWeldedIndices[Tether.Get<1>()];
+			}
+		}
 	}
 
 	// Populate used bone names and indices
@@ -246,3 +204,20 @@ FChaosClothSimulationModel::FChaosClothSimulationModel(const TSharedPtr<const FM
 	// Initialize Reference bone index
 	ReferenceBoneIndex = CalculateReferenceBoneIndex(ClothSimulationLodModels, ReferenceSkeleton, UsedBoneIndices);
 }
+
+TArray<TConstArrayView<TTuple<int32, int32, float>>> FChaosClothSimulationModel::GetTethers(int32 LODIndex) const
+{
+	TArray<TConstArrayView<TTuple<int32, int32, float>>> Tethers;
+	if (IsValidLodIndex(LODIndex))
+	{
+		const FClothTetherData& ClothTetherData = ClothSimulationLodModels[LODIndex].TetherData;
+
+		const int32 NumTetherBatches = ClothTetherData.Tethers.Num();
+		Tethers.Reserve(NumTetherBatches);
+		for (int32 Index = 0; Index < NumTetherBatches; ++Index)
+		{
+			Tethers.Emplace(TConstArrayView<TTuple<int32, int32, float>>(ClothTetherData.Tethers[Index]));
+		}
+	}
+	return Tethers;
+}
\ No newline at end of file
diff --git a/Engine/Plugins/Experimental/ChaosClothAsset/Source/ChaosClothAssetEngine/Private/ChaosClothAsset/ClothSimulationModel.h b/Engine/Plugins/Experimental/ChaosClothAsset/Source/ChaosClothAssetEngine/Private/ChaosClothAsset/ClothSimulationModel.h
--- a/Engine/Plugins/Experimental/ChaosClothAsset/Source/ChaosClothAssetEngine/Private/ChaosClothAsset/ClothSimulationModel.h
+++ b/Engine/Plugins/Experimental/ChaosClothAsset/Source/ChaosClothAssetEngine/Private/ChaosClothAsset/ClothSimulationModel.h
@@ -6,6 +6,7 @@
 #include "Math/MathFwd.h"
 #include "Templates/SharedPointer.h"
 #include "ClothVertBoneData.h"
+#include "ClothTetherData.h"
 
 #include "ClothSimulationModel.generated.h"
 
@@ -52,6 +53,10 @@ struct FChaosClothSimulationLodModel
 	/** Weight maps for storing painted attributes modifiers on constraint properties. */
 	TMap<FName, TArray<float>> WeightMaps;
 
+	/** Tether data*/
+	UPROPERTY()
+	FClothTetherData TetherData;
+
 	// Custom serialize for weight maps
 	bool Serialize(FArchive& Ar);
 };
@@ -106,4 +111,5 @@ struct FChaosClothSimulationModel
 	TConstArrayView<uint32> GetPatternIndices(int32 LodIndex) const { return IsValidLodIndex(LodIndex) ? ClothSimulationLodModels[LodIndex].PatternIndices : TConstArrayView<uint32>(); }
 	TConstArrayView<uint32> GetPatternToWeldedIndices(int32 LodIndex) const { return IsValidLodIndex(LodIndex) ? ClothSimulationLodModels[LodIndex].PatternToWeldedIndices : TConstArrayView<uint32>(); }
 	TConstArrayView<FClothVertBoneData> GetBoneData(int32 LodIndex) const { return IsValidLodIndex(LodIndex) ? ClothSimulationLodModels[LodIndex].BoneData : TConstArrayView<FClothVertBoneData>(); }
+	TArray<TConstArrayView<TTuple<int32, int32, float>>> GetTethers(int32 LodIndex) const;
 };
diff --git a/Engine/Plugins/Experimental/ChaosClothAsset/Source/ChaosClothAssetEngine/Public/ChaosClothAsset/ClothEngineTools.h b/Engine/Plugins/Experimental/ChaosClothAsset/Source/ChaosClothAssetEngine/Public/ChaosClothAsset/ClothEngineTools.h
new file mode 100644
--- /dev/null
+++ b/Engine/Plugins/Experimental/ChaosClothAsset/Source/ChaosClothAssetEngine/Public/ChaosClothAsset/ClothEngineTools.h
@@ -0,0 +1,20 @@
+// Copyright Epic Games, Inc. All Rights Reserved.
+
+#pragma once
+
+#include "Templates/SharedPointer.h"
+
+struct FManagedArrayCollection;
+class FName;
+
+namespace UE::Chaos::ClothAsset
+{
+	/**
+	 *  Tools operating on cloth collections with Engine dependency
+	 */
+	struct CHAOSCLOTHASSETENGINE_API FClothEngineTools
+	{
+		/** Generate tether data. */
+		static void GenerateTethers(const TSharedPtr<FManagedArrayCollection>& ClothCollection, const FName& WeightMap, const bool bGeodesicTethers);
+	};
+}  // End namespace UE::Chaos::ClothAsset
\ No newline at end of file
diff --git a/Engine/Plugins/Experimental/ChaosClothAssetEditor/Source/ChaosClothAssetDataflowNodes/Private/ChaosClothAsset/DataflowNodes.cpp b/Engine/Plugins/Experimental/ChaosClothAssetEditor/Source/ChaosClothAssetDataflowNodes/Private/ChaosClothAsset/DataflowNodes.cpp
--- a/Engine/Plugins/Experimental/ChaosClothAssetEditor/Source/ChaosClothAssetDataflowNodes/Private/ChaosClothAsset/DataflowNodes.cpp
+++ b/Engine/Plugins/Experimental/ChaosClothAssetEditor/Source/ChaosClothAssetDataflowNodes/Private/ChaosClothAsset/DataflowNodes.cpp
@@ -18,6 +18,7 @@
 #include "ChaosClothAsset/SimulationDampingConfigNode.h"
 #include "ChaosClothAsset/SimulationDefaultConfigNode.h"
 #include "ChaosClothAsset/SimulationGravityConfigNode.h"
+#include "ChaosClothAsset/SimulationLongRangeAttachmentConfigNode.h"
 #include "ChaosClothAsset/SimulationMassConfigNode.h"
 #include "ChaosClothAsset/SimulationMaxDistanceConfigNode.h"
 #include "ChaosClothAsset/SimulationPBDAreaSpringConfigNode.h"
@@ -65,6 +66,7 @@ namespace UE::Chaos::ClothAsset::DataflowNodes
 		DATAFLOW_NODE_REGISTER_CREATION_FACTORY(FChaosClothAssetSimulationDampingConfigNode);
 		DATAFLOW_NODE_REGISTER_CREATION_FACTORY(FChaosClothAssetSimulationDefaultConfigNode);
 		DATAFLOW_NODE_REGISTER_CREATION_FACTORY(FChaosClothAssetSimulationGravityConfigNode);
+		DATAFLOW_NODE_REGISTER_CREATION_FACTORY(FChaosClothAssetSimulationLongRangeAttachmentConfigNode);
 		DATAFLOW_NODE_REGISTER_CREATION_FACTORY(FChaosClothAssetSimulationMassConfigNode);
 		DATAFLOW_NODE_REGISTER_CREATION_FACTORY(FChaosClothAssetSimulationMaxDistanceConfigNode);
 		DATAFLOW_NODE_REGISTER_CREATION_FACTORY(FChaosClothAssetSimulationPBDAreaSpringConfigNode);
diff --git a/Engine/Plugins/Experimental/ChaosClothAssetEditor/Source/ChaosClothAssetDataflowNodes/Private/ChaosClothAsset/SimulationBaseConfigNode.cpp b/Engine/Plugins/Experimental/ChaosClothAssetEditor/Source/ChaosClothAssetDataflowNodes/Private/ChaosClothAsset/SimulationBaseConfigNode.cpp
--- a/Engine/Plugins/Experimental/ChaosClothAssetEditor/Source/ChaosClothAssetDataflowNodes/Private/ChaosClothAsset/SimulationBaseConfigNode.cpp
+++ b/Engine/Plugins/Experimental/ChaosClothAssetEditor/Source/ChaosClothAssetDataflowNodes/Private/ChaosClothAsset/SimulationBaseConfigNode.cpp
@@ -64,6 +64,8 @@ void FChaosClothAssetSimulationBaseConfigNode::Evaluate(Dataflow::FContext& Cont
 
 		AddProperties(Properties);
 
+		EvaluateClothCollection(Context, ClothCollection);
+
 		SetValue<FManagedArrayCollection>(Context, *ClothCollection, &Collection);
 	}
 }
diff --git a/Engine/Plugins/Experimental/ChaosClothAssetEditor/Source/ChaosClothAssetDataflowNodes/Private/ChaosClothAsset/SimulationDefaultConfigNode.cpp b/Engine/Plugins/Experimental/ChaosClothAssetEditor/Source/ChaosClothAssetDataflowNodes/Private/ChaosClothAsset/SimulationDefaultConfigNode.cpp
--- a/Engine/Plugins/Experimental/ChaosClothAssetEditor/Source/ChaosClothAssetDataflowNodes/Private/ChaosClothAsset/SimulationDefaultConfigNode.cpp
+++ b/Engine/Plugins/Experimental/ChaosClothAssetEditor/Source/ChaosClothAssetDataflowNodes/Private/ChaosClothAsset/SimulationDefaultConfigNode.cpp
@@ -2,8 +2,10 @@
 
 #include "ChaosClothAsset/SimulationDefaultConfigNode.h"
 #include "ChaosClothAsset/DataflowNodes.h"
+#include "ChaosClothAsset/ClothEngineTools.h"
 #include "ChaosCloth/ChaosClothConfig.h"
 #include "ChaosCloth/ChaosClothingSimulationConfig.h"
+#include "Chaos/PBDLongRangeConstraints.h"  // For Tether modes
 #include "Chaos/CollectionPropertyFacade.h"
 
 #include UE_INLINE_GENERATED_CPP_BY_NAME(SimulationDefaultConfigNode)
@@ -42,13 +44,23 @@ void FChaosClothAssetSimulationDefaultConfigNode::Evaluate(Dataflow::FContext& C
 	if (Out->IsA<FManagedArrayCollection>(&Collection))
 	{
 		FManagedArrayCollection InCollection = GetValue<FManagedArrayCollection>(Context, &Collection);
+		const TSharedRef<FManagedArrayCollection> ClothCollection = MakeShared<FManagedArrayCollection>(MoveTemp(InCollection));
 
 		FClothingSimulationConfig ClothingSimulationConfig;
 		ClothingSimulationConfig.Initialize(SimulationConfig.Get(), SharedSimulationConfig.Get());
 
-		ClothingSimulationConfig.GetPropertyCollection()->CopyTo(&InCollection);
+		ClothingSimulationConfig.GetPropertyCollection()->CopyTo(&ClothCollection.Get());
 
-		SetValue<FManagedArrayCollection>(Context, InCollection, &Collection);
+		// Generate tethers
+		const FCollectionPropertyConstFacade& Properties = ClothingSimulationConfig.GetProperties();		
+		constexpr bool bUseGeodesicTethersDefault = true;
+		const bool bUseGeodesicTethers = Properties.GetValue<bool>(TEXT("UseGeodesicTethers"), bUseGeodesicTethersDefault);
+		// Use the "MaxDistance" weight map to generate tethers. This follows legacy behavior.
+		static const FName MaxDistanceName(TEXT("MaxDistance"));
+
+		UE::Chaos::ClothAsset::FClothEngineTools::GenerateTethers(ClothCollection, MaxDistanceName, bUseGeodesicTethers);
+
+		SetValue<FManagedArrayCollection>(Context, *ClothCollection, &Collection);
 	}
 }
 
diff --git a/Engine/Plugins/Experimental/ChaosClothAssetEditor/Source/ChaosClothAssetDataflowNodes/Private/ChaosClothAsset/SimulationLongRangeAttachmentConfigNode.cpp b/Engine/Plugins/Experimental/ChaosClothAssetEditor/Source/ChaosClothAssetDataflowNodes/Private/ChaosClothAsset/SimulationLongRangeAttachmentConfigNode.cpp
new file mode 100644
--- /dev/null
+++ b/Engine/Plugins/Experimental/ChaosClothAssetEditor/Source/ChaosClothAssetDataflowNodes/Private/ChaosClothAsset/SimulationLongRangeAttachmentConfigNode.cpp
@@ -0,0 +1,30 @@
+// Copyright Epic Games, Inc. All Rights Reserved.
+
+#include "ChaosClothAsset/SimulationLongRangeAttachmentConfigNode.h"
+#include "ChaosClothAsset/DataflowNodes.h"
+#include "ChaosClothAsset/SimulationBaseConfigNodePrivate.h"
+#include "ChaosClothAsset/ClothEngineTools.h"
+#include "Chaos/CollectionPropertyFacade.h"
+
+#include UE_INLINE_GENERATED_CPP_BY_NAME(SimulationLongRangeAttachmentConfigNode)
+
+FChaosClothAssetSimulationLongRangeAttachmentConfigNode::FChaosClothAssetSimulationLongRangeAttachmentConfigNode(const Dataflow::FNodeParameters& InParam, FGuid InGuid)
+	: FChaosClothAssetSimulationBaseConfigNode(InParam, InGuid)
+{
+	RegisterCollectionConnections();
+	RegisterInputConnection(&FixedEndWeightMap);
+}
+
+void FChaosClothAssetSimulationLongRangeAttachmentConfigNode::AddProperties(::Chaos::Softs::FCollectionPropertyMutableFacade& Properties) const
+{
+	UE_CHAOS_CLOTHASSET_SIMULATIONCONFIG_SETPROPERTYWEIGHTED(TetherStiffness);
+	UE_CHAOS_CLOTHASSET_SIMULATIONCONFIG_SETPROPERTYWEIGHTED(TetherScale);
+	UE_CHAOS_CLOTHASSET_SIMULATIONCONFIG_SETPROPERTYBOOL(UseGeodesicTethers);
+}
+
+void FChaosClothAssetSimulationLongRangeAttachmentConfigNode::EvaluateClothCollection(Dataflow::FContext& Context, const TSharedRef<FManagedArrayCollection>& ClothCollection) const
+{
+	const FString InFixedEndWeightMapString = GetValue<FString>(Context, &FixedEndWeightMap);
+	const FName InFixedEndWeightMap(InFixedEndWeightMapString);
+	UE::Chaos::ClothAsset::FClothEngineTools::GenerateTethers(ClothCollection, InFixedEndWeightMap, bUseGeodesicTethers);
+}
diff --git a/Engine/Plugins/Experimental/ChaosClothAssetEditor/Source/ChaosClothAssetDataflowNodes/Public/ChaosClothAsset/SimulationBaseConfigNode.h b/Engine/Plugins/Experimental/ChaosClothAssetEditor/Source/ChaosClothAssetDataflowNodes/Public/ChaosClothAsset/SimulationBaseConfigNode.h
--- a/Engine/Plugins/Experimental/ChaosClothAssetEditor/Source/ChaosClothAssetDataflowNodes/Public/ChaosClothAsset/SimulationBaseConfigNode.h
+++ b/Engine/Plugins/Experimental/ChaosClothAssetEditor/Source/ChaosClothAssetDataflowNodes/Public/ChaosClothAsset/SimulationBaseConfigNode.h
@@ -34,6 +34,9 @@ protected:
 	virtual void AddProperties(::Chaos::Softs::FCollectionPropertyMutableFacade& Properties) const
 	PURE_VIRTUAL(FChaosClothAssetSimulationBaseConfigNode::AddProperties, );
 
+	/* Override this to do additional node-specific evaluate on the cloth collection output. AddProperties has already been called when this is called. */
+	virtual void EvaluateClothCollection(Dataflow::FContext& Context, const TSharedRef<FManagedArrayCollection>& ClothCollection) const {}
+
 	void RegisterCollectionConnections();
 
 	int32 AddPropertyHelper(
diff --git a/Engine/Plugins/Experimental/ChaosClothAssetEditor/Source/ChaosClothAssetDataflowNodes/Public/ChaosClothAsset/SimulationLongRangeAttachmentConfigNode.h b/Engine/Plugins/Experimental/ChaosClothAssetEditor/Source/ChaosClothAssetDataflowNodes/Public/ChaosClothAsset/SimulationLongRangeAttachmentConfigNode.h
new file mode 100644
--- /dev/null
+++ b/Engine/Plugins/Experimental/ChaosClothAssetEditor/Source/ChaosClothAssetDataflowNodes/Public/ChaosClothAsset/SimulationLongRangeAttachmentConfigNode.h
@@ -0,0 +1,57 @@
+// Copyright Epic Games, Inc. All Rights Reserved.
+
+#pragma once 
+
+#include "ChaosClothAsset/SimulationBaseConfigNode.h"
+#include "ChaosClothAsset/WeightedValue.h"
+#include "SimulationLongRangeAttachmentConfigNode.generated.h"
+
+/** Long range attachment constraint property configuration node. */
+USTRUCT(Meta = (DataflowCloth))
+struct FChaosClothAssetSimulationLongRangeAttachmentConfigNode : public FChaosClothAssetSimulationBaseConfigNode
+{
+	GENERATED_USTRUCT_BODY()
+	DATAFLOW_NODE_DEFINE_INTERNAL(FChaosClothAssetSimulationLongRangeAttachmentConfigNode, "SimulationLongRangeAttachmentConfig", "Cloth", "Cloth Simulation Long Range Attachment Config")
+
+public:
+
+	/**
+	 * The tethers' stiffness of the long range attachment constraints.
+	 * The long range attachment connects each of the cloth particles to its closest fixed point with a spring constraint.
+	 * This can be used to compensate for a lack of stretch resistance when the iterations count is kept low for performance reasons.
+	 * Can lead to an unnatural pull string puppet like behavior.
+	 * If a valid weight map is found with the given Weight Map name, then both Low and High values
+	 * are interpolated with the per particle weight to make the final value used for the simulation.
+	 * Otherwise all particles are considered to have a zero weight, and only the Low value is meaningful.
+	 */
+	UPROPERTY(EditAnywhere, Category = "Long Range Attachment Properties", Meta = (UIMin = "0", UIMax = "1", ClampMin = "0", ClampMax = "1"))
+	FChaosClothAssetWeightedValue TetherStiffness = { true, 1.f, 1.f, TEXT("TetherStiffness") };
+
+	/**
+	 * The limit scale of the long range attachment constraints (aka tether limit).
+	 * If a valid weight map is found with the given Weight Map name, then both Low and High values
+	 * are interpolated with the per particle weight to make the final value used for the simulation.
+	 * Otherwise all particles are considered to have a zero weight, and only the Low value is meaningful.
+	*/
+	UPROPERTY(EditAnywhere, Category = "Long Range Attachment Properties", Meta = (UIMin = "1.", UIMax = "1.1", ClampMin = "0.01", ClampMax = "10"))
+	FChaosClothAssetWeightedValue TetherScale = { true, 1.f, 1.f, TEXT("TetherScale") };
+
+	/**
+	 * Use geodesic instead of euclidean distance calculations for the Long Range Attachment constraint,
+	 * which is slower at setup but more accurate at establishing the correct position and length of the tethers,
+	 * and therefore is less prone to artifacts during the simulation.
+	 */
+	UPROPERTY(EditAnywhere, Category = "Long Range Attachment Properties")
+	bool bUseGeodesicTethers = true;
+
+	/** The name of the weight map used to calculate fixed tether ends. All vertices with weight = 0 will be considered fixed. */
+	UPROPERTY(EditAnywhere, Category = "Long Range Attachment Properties", Meta = (DataflowInput))
+	FString FixedEndWeightMap = TEXT("MaxDistance");
+
+	FChaosClothAssetSimulationLongRangeAttachmentConfigNode(const Dataflow::FNodeParameters& InParam, FGuid InGuid = FGuid::NewGuid());
+
+private:
+	virtual void AddProperties(::Chaos::Softs::FCollectionPropertyMutableFacade& Properties) const override;
+
+	virtual void EvaluateClothCollection(Dataflow::FContext& Context, const TSharedRef<FManagedArrayCollection>& ClothCollection) const override;
+};
diff --git a/Engine/Source/Runtime/ClothingSystemRuntimeCommon/Public/ClothTetherData.h b/Engine/Source/Runtime/ClothingSystemRuntimeCommon/Public/ClothTetherData.h
--- a/Engine/Source/Runtime/ClothingSystemRuntimeCommon/Public/ClothTetherData.h
+++ b/Engine/Source/Runtime/ClothingSystemRuntimeCommon/Public/ClothTetherData.h
@@ -29,7 +29,7 @@ struct FClothTetherData
 	 * Generate the tethers by following the triangle mesh network from the closest kinematic to each dynamic point.
 	 * Inside array items can be processed concurrently, but the outside array must be iterated on sequentially.
 	 */
-	void GenerateTethers(
+	CLOTHINGSYSTEMRUNTIMECOMMON_API void GenerateTethers(
 		const TConstArrayView<FVector3f>& Points,  // Reference pose
 		const TConstArrayView<uint32>& Indices,  // Triangle mesh
 		const TConstArrayView<float>& MaxDistances,  // Mask for sorting the kinematic from the dynamic points
diff --git a/Engine/Source/Runtime/Core/Public/UObject/UE5MainStreamObjectVersion.h b/Engine/Source/Runtime/Core/Public/UObject/UE5MainStreamObjectVersion.h
--- a/Engine/Source/Runtime/Core/Public/UObject/UE5MainStreamObjectVersion.h
+++ b/Engine/Source/Runtime/Core/Public/UObject/UE5MainStreamObjectVersion.h
@@ -360,6 +360,9 @@ struct CORE_API FUE5MainStreamObjectVersion
 
 		// serialization inputs and outputs as two different sections
 		DataflowSeparateInputOutputSerialization,
+
+		// Cloth collection tether initialization
+		ClothCollectionTetherInitialization,
 		
 		// -----<new versions can be added above this line>-------------------------------------------------
 		VersionPlusOne,
diff --git a/Engine/Source/Runtime/Experimental/Chaos/Public/Chaos/PBDLongRangeConstraintsBase.h b/Engine/Source/Runtime/Experimental/Chaos/Public/Chaos/PBDLongRangeConstraintsBase.h
--- a/Engine/Source/Runtime/Experimental/Chaos/Public/Chaos/PBDLongRangeConstraintsBase.h
+++ b/Engine/Source/Runtime/Experimental/Chaos/Public/Chaos/PBDLongRangeConstraintsBase.h
@@ -20,7 +20,7 @@ public:
 	static constexpr FSolverReal MinTetherScale = (FSolverReal)0.01;
 	static constexpr FSolverReal MaxTetherScale = (FSolverReal)10.;
 
-	enum class EMode : uint8
+	enum class UE_DEPRECATED(5.3, "Tether EMode has been replaced with bUseGeodesicTethers.") EMode : uint8
 	{
 		Euclidean,
 		Geodesic,
