Chaos Cloth - Fixed missing UV channel update when building the cloth asset mesh model, to prevent ensure(UVChannelData.bInitialized);
Note: This was made possible by a change in the Skeletal Mesh/Skinned Asset in CL 25504268.

#rnx

diff --git a/Engine/Plugins/Experimental/ChaosClothAsset/Source/ChaosClothAssetEngine/Private/ChaosClothAsset/ClothAsset.cpp b/Engine/Plugins/Experimental/ChaosClothAsset/Source/ChaosClothAssetEngine/Private/ChaosClothAsset/ClothAsset.cpp
--- a/Engine/Plugins/Experimental/ChaosClothAsset/Source/ChaosClothAssetEngine/Private/ChaosClothAsset/ClothAsset.cpp
+++ b/Engine/Plugins/Experimental/ChaosClothAsset/Source/ChaosClothAssetEngine/Private/ChaosClothAsset/ClothAsset.cpp
@@ -450,6 +450,7 @@ void UChaosClothAsset::BuildMeshModel()
 			const FCollectionClothConstFacade Cloth(GetClothCollection());
 			const int32 NumLods = Cloth.GetNumLods();
 
+			// Build each LOD
 			check(MeshModel);  // MeshModel should always be created in the Cloth Asset constructor WITH_EDITORONLY_DATA
 			MeshModel->LODModels.Empty();
 
@@ -459,6 +460,9 @@ void UChaosClothAsset::BuildMeshModel()
 				ClothAssetBuilder->BuildLod(*LODModel, *this, LodIndex);
 				MeshModel->LODModels.Add(LODModel);
 			}
+
+			// Build the material channel data used by the texture streamer
+			UpdateUVChannelData(true);
 		}
 	}
 }
