Chaos Cloth:
- Fixed spam log messages from empty Cloth Assets by preventing the Cloth Component from rendering, by updating its visibility flags in the OnRegister and SetSkinnedAsset functions.
- Allow changes of Cloth Asset in the Cloth Component without re-registering the component, and made sure to stop the simulation first when this happens.

#rnx

diff --git a/Engine/Plugins/Experimental/ChaosClothAsset/Source/ChaosClothAssetEngine/Private/ChaosClothAsset/ClothComponent.cpp b/Engine/Plugins/Experimental/ChaosClothAsset/Source/ChaosClothAssetEngine/Private/ChaosClothAsset/ClothComponent.cpp
--- a/Engine/Plugins/Experimental/ChaosClothAsset/Source/ChaosClothAssetEngine/Private/ChaosClothAsset/ClothComponent.cpp
+++ b/Engine/Plugins/Experimental/ChaosClothAsset/Source/ChaosClothAssetEngine/Private/ChaosClothAsset/ClothComponent.cpp
@@ -114,6 +114,9 @@ void UChaosClothComponent::OnRegister()
 			ClothSimulationProxy = CreateClothSimulationProxy();
 		}
 	}
+
+	// Update render visibility, so that an empty LODs doesn't unnecessarily go to render
+	UpdateVisibility();
 }
 
 TSharedPtr<UE::Chaos::ClothAsset::FClothSimulationProxy> UChaosClothComponent::CreateClothSimulationProxy()
@@ -235,6 +238,20 @@ void UChaosClothComponent::GetUpdateClothSimulationData_AnyThread(TMap<int32, FC
 	}
 }
 
+void UChaosClothComponent::SetSkinnedAssetAndUpdate(USkinnedAsset* InSkinnedAsset, bool bReinitPose)
+{
+	if (InSkinnedAsset != GetSkinnedAsset())
+	{
+		Super::SetSkinnedAssetAndUpdate(InSkinnedAsset, bReinitPose);
+
+		// Stop current simulation
+		HandleExistingParallelSimulation();
+
+		// Update the component visibility in case the new render mesh has no valid LOD
+		UpdateVisibility();
+	}
+}
+
 void UChaosClothComponent::PostLoad()
 {
 	Super::PostLoad();
@@ -248,14 +265,12 @@ void UChaosClothComponent::PostLoad()
 
 void UChaosClothComponent::SetClothAsset(UChaosClothAsset* InClothAsset)
 {
-	// De-register component, and re-register it when it reaches end of scope
-	const FComponentReregisterContext ComponentReregisterContext(this);
+	SetSkinnedAssetAndUpdate(InClothAsset);
 
 	PRAGMA_DISABLE_DEPRECATION_WARNINGS
 #if WITH_EDITORONLY_DATA
 	ClothAsset = InClothAsset;
 #endif
-	SetSkinnedAsset(InClothAsset);
 	PRAGMA_ENABLE_DEPRECATION_WARNINGS
 }
 
@@ -387,3 +402,17 @@ bool UChaosClothComponent::ShouldWaitForParallelSimulationInTickComponent() cons
 
 	return bWaitForParallelTask || (CVarClothPhysicsWaitForParallelClothTask && CVarClothPhysicsWaitForParallelClothTask->GetBool());
 }
+
+void UChaosClothComponent::UpdateVisibility()
+{
+	if (GetClothAsset())
+	{
+		const FSkeletalMeshRenderData* const SkeletalMeshRenderData = GetClothAsset()->GetResourceForRendering();
+		const int32 FirstValidLODIdx = SkeletalMeshRenderData ? SkeletalMeshRenderData->GetFirstValidLODIdx(0) : INDEX_NONE;
+		SetVisibility(FirstValidLODIdx != INDEX_NONE);
+	}
+	else
+	{
+		SetVisibility(false);
+	}
+}
diff --git a/Engine/Plugins/Experimental/ChaosClothAsset/Source/ChaosClothAssetEngine/Public/ChaosClothAsset/ClothComponent.h b/Engine/Plugins/Experimental/ChaosClothAsset/Source/ChaosClothAssetEngine/Public/ChaosClothAsset/ClothComponent.h
--- a/Engine/Plugins/Experimental/ChaosClothAsset/Source/ChaosClothAssetEngine/Public/ChaosClothAsset/ClothComponent.h
+++ b/Engine/Plugins/Experimental/ChaosClothAsset/Source/ChaosClothAssetEngine/Public/ChaosClothAsset/ClothComponent.h
@@ -122,6 +122,7 @@ protected:
 	//~ Begin USkinnedMeshComponent Interface
 	virtual void RefreshBoneTransforms(FActorComponentTickFunction* TickFunction = nullptr) override;
 	virtual void GetUpdateClothSimulationData_AnyThread(TMap<int32, FClothSimulData>& OutClothSimulData, FMatrix& OutLocalToWorld, float& OutBlendWeight) override;
+	virtual void SetSkinnedAssetAndUpdate(USkinnedAsset* InSkinnedAsset, bool bReinitPose = true) override;
 	//~ End USkinnedMeshComponent Interface
 
 	/** Override this function for setting up custom simulation proxies when the component is registered. */
@@ -131,6 +132,7 @@ private:
 	void StartNewParallelSimulation(float DeltaTime);
 	void HandleExistingParallelSimulation();
 	bool ShouldWaitForParallelSimulationInTickComponent() const;
+	void UpdateVisibility();
 
 #if WITH_EDITORONLY_DATA
 	/** Cloth asset used by this component. */
