Chaos Cloth - Reordered a few functions from the Cloth Component implementation to match the header functions organization before the file gets too large or widely used. No code change otherwise.

#rnx

diff --git a/Engine/Plugins/Experimental/ChaosClothAsset/Source/ChaosClothAssetEngine/Private/ChaosClothAsset/ClothComponent.cpp b/Engine/Plugins/Experimental/ChaosClothAsset/Source/ChaosClothAssetEngine/Private/ChaosClothAsset/ClothComponent.cpp
--- a/Engine/Plugins/Experimental/ChaosClothAsset/Source/ChaosClothAssetEngine/Private/ChaosClothAsset/ClothComponent.cpp
+++ b/Engine/Plugins/Experimental/ChaosClothAsset/Source/ChaosClothAssetEngine/Private/ChaosClothAsset/ClothComponent.cpp
@@ -83,46 +83,20 @@
 
 UChaosClothComponent::~UChaosClothComponent() = default;
 
-void UChaosClothComponent::OnRegister()
+void UChaosClothComponent::SetClothAsset(UChaosClothAsset* InClothAsset)
 {
-	using namespace UE::Chaos::ClothAsset;
+	SetSkinnedAssetAndUpdate(InClothAsset);
 
-	LLM_SCOPE(ELLMTag::Chaos);
-
-	Super::OnRegister();
-
-	if (GetClothAsset())
-	{
-		const TSharedPtr<const FChaosClothSimulationModel> ClothSimulationModel = GetClothAsset()->GetClothSimulationModel();
-		if (ensure(ClothSimulationModel) && ClothSimulationModel->GetNumLods())
-		{
-			if (!LeaderPoseComponent.IsValid())
-			{
-				FSkeletalMeshLODRenderData& LODData = GetClothAsset()->GetResourceForRendering()->LODRenderData[GetPredictedLODLevel()];
-				GetClothAsset()->FillComponentSpaceTransforms(GetClothAsset()->GetRefSkeleton().GetRefBonePose(), LODData.RequiredBones, GetEditableComponentSpaceTransforms());
-
-				bNeedToFlipSpaceBaseBuffers = true; // Have updated space bases so need to flip
-				FlipEditableSpaceBases();
-				bHasValidBoneTransform = true;
-			}
-
-			// Create cloth property collection, and fill up the collection with the original cloth asset properties
-			PropertyCollection = MakeShared<FManagedArrayCollection>();
-			ResetConfigProperties();
-
-			// Create simulation proxy
-			ClothSimulationProxy = CreateClothSimulationProxy();
-		}
-	}
-
-	// Update render visibility, so that an empty LODs doesn't unnecessarily go to render
-	UpdateVisibility();
+	PRAGMA_DISABLE_DEPRECATION_WARNINGS
+#if WITH_EDITORONLY_DATA
+		ClothAsset = InClothAsset;
+#endif
+	PRAGMA_ENABLE_DEPRECATION_WARNINGS
 }
 
-TSharedPtr<UE::Chaos::ClothAsset::FClothSimulationProxy> UChaosClothComponent::CreateClothSimulationProxy()
+UChaosClothAsset* UChaosClothComponent::GetClothAsset() const
 {
-	using namespace UE::Chaos::ClothAsset;
-	return MakeShared<FClothSimulationProxy>(*this);
+	return Cast<UChaosClothAsset>(GetSkinnedAsset());
 }
 
 bool UChaosClothComponent::IsSimulationSuspended() const
@@ -161,6 +135,71 @@
 	CollectionPropertyFacade = MakeUnique<FCollectionPropertyFacade>(PropertyCollection);
 }
 
+void UChaosClothComponent::PostLoad()
+{
+	Super::PostLoad();
+
+#if WITH_EDITORONLY_DATA
+	PRAGMA_DISABLE_DEPRECATION_WARNINGS
+		ClothAsset = GetClothAsset();
+	PRAGMA_ENABLE_DEPRECATION_WARNINGS
+#endif
+}
+
+#if WITH_EDITOR
+void UChaosClothComponent::PostEditChangeProperty(FPropertyChangedEvent& PropertyChangedEvent)
+{
+	// Set the skinned asset pointer with the alias pointer (must happen before the call to Super::PostEditChangeProperty)
+	if (const FProperty* const Property = PropertyChangedEvent.Property)
+	{
+		PRAGMA_DISABLE_DEPRECATION_WARNINGS
+			if (Property->GetFName() == GET_MEMBER_NAME_CHECKED(UChaosClothComponent, ClothAsset))
+			{
+				SetClothAsset(ClothAsset);
+			}
+		PRAGMA_ENABLE_DEPRECATION_WARNINGS
+	}
+
+	Super::PostEditChangeProperty(PropertyChangedEvent);
+}
+#endif // WITH_EDITOR
+
+void UChaosClothComponent::OnRegister()
+{
+	using namespace UE::Chaos::ClothAsset;
+
+	LLM_SCOPE(ELLMTag::Chaos);
+
+	Super::OnRegister();
+
+	if (GetClothAsset())
+	{
+		const TSharedPtr<const FChaosClothSimulationModel> ClothSimulationModel = GetClothAsset()->GetClothSimulationModel();
+		if (ensure(ClothSimulationModel) && ClothSimulationModel->GetNumLods())
+		{
+			if (!LeaderPoseComponent.IsValid())
+			{
+				FSkeletalMeshLODRenderData& LODData = GetClothAsset()->GetResourceForRendering()->LODRenderData[GetPredictedLODLevel()];
+				GetClothAsset()->FillComponentSpaceTransforms(GetClothAsset()->GetRefSkeleton().GetRefBonePose(), LODData.RequiredBones, GetEditableComponentSpaceTransforms());
+
+				bNeedToFlipSpaceBaseBuffers = true; // Have updated space bases so need to flip
+				FlipEditableSpaceBases();
+				bHasValidBoneTransform = true;
+			}
+
+			// Create cloth property collection, and fill up the collection with the original cloth asset properties
+			PropertyCollection = MakeShared<FManagedArrayCollection>();
+			ResetConfigProperties();
+
+			// Create simulation proxy
+			ClothSimulationProxy = CreateClothSimulationProxy();
+		}
+	}
+
+	// Update render visibility, so that an empty LODs doesn't unnecessarily go to render
+	UpdateVisibility();
+}
+
 void UChaosClothComponent::OnUnregister()
 {
 	Super::OnUnregister();
@@ -201,102 +240,6 @@
 	Super::TickComponent(DeltaTime, TickType, ThisTickFunction);
 }
 
-void UChaosClothComponent::RefreshBoneTransforms(FActorComponentTickFunction* TickFunction)
-{
-	MarkRenderDynamicDataDirty();
-
-	bNeedToFlipSpaceBaseBuffers = true;
-	bHasValidBoneTransform = false;
-	FlipEditableSpaceBases();
-	bHasValidBoneTransform = true;
-}
-
-void UChaosClothComponent::GetUpdateClothSimulationData_AnyThread(TMap<int32, FClothSimulData>& OutClothSimulData, FMatrix& OutLocalToWorld, float& OutBlendWeight)
-{
-	OutLocalToWorld = GetComponentToWorld().ToMatrixWithScale();
-
-	const UChaosClothComponent* const LeaderPoseClothComponent = Cast<UChaosClothComponent>(LeaderPoseComponent.Get());
-	if (LeaderPoseClothComponent && LeaderPoseClothComponent->ClothSimulationProxy && bBindToLeaderComponent)
-	{
-		OutBlendWeight = BlendWeight;
-		OutClothSimulData = LeaderPoseClothComponent->ClothSimulationProxy->GetCurrentSimulationData_AnyThread();
-	}
-	else if (bEnableSimulation && !bBindToLeaderComponent && ClothSimulationProxy)
-	{
-		OutBlendWeight = BlendWeight;
-		OutClothSimulData = ClothSimulationProxy->GetCurrentSimulationData_AnyThread();
-	}
-	else
-	{
-		OutClothSimulData.Reset();
-	}
-
-	// Blend cloth out whenever the simulation data is invalid
-	if (!OutClothSimulData.Num())
-	{
-		OutBlendWeight = 0.0f;
-	}
-}
-
-void UChaosClothComponent::SetSkinnedAssetAndUpdate(USkinnedAsset* InSkinnedAsset, bool bReinitPose)
-{
-	if (InSkinnedAsset != GetSkinnedAsset())
-	{
-		Super::SetSkinnedAssetAndUpdate(InSkinnedAsset, bReinitPose);
-
-		// Stop current simulation
-		HandleExistingParallelSimulation();
-
-		// Update the component visibility in case the new render mesh has no valid LOD
-		UpdateVisibility();
-	}
-}
-
-void UChaosClothComponent::PostLoad()
-{
-	Super::PostLoad();
-
-#if WITH_EDITORONLY_DATA
-	PRAGMA_DISABLE_DEPRECATION_WARNINGS
-	ClothAsset = GetClothAsset();
-	PRAGMA_ENABLE_DEPRECATION_WARNINGS
-#endif
-}
-
-void UChaosClothComponent::SetClothAsset(UChaosClothAsset* InClothAsset)
-{
-	SetSkinnedAssetAndUpdate(InClothAsset);
-
-	PRAGMA_DISABLE_DEPRECATION_WARNINGS
-#if WITH_EDITORONLY_DATA
-	ClothAsset = InClothAsset;
-#endif
-	PRAGMA_ENABLE_DEPRECATION_WARNINGS
-}
-
-UChaosClothAsset* UChaosClothComponent::GetClothAsset() const
-{
-	return Cast<UChaosClothAsset>(GetSkinnedAsset());
-}
-
-#if WITH_EDITOR
-void UChaosClothComponent::PostEditChangeProperty(FPropertyChangedEvent& PropertyChangedEvent)
-{
-	// Set the skinned asset pointer with the alias pointer (must happen before the call to Super::PostEditChangeProperty)
-	if (const FProperty* const Property = PropertyChangedEvent.Property)
-	{
-PRAGMA_DISABLE_DEPRECATION_WARNINGS
-		if (Property->GetFName() == GET_MEMBER_NAME_CHECKED(UChaosClothComponent, ClothAsset))
-		{
-			SetClothAsset(ClothAsset);
-		}
-PRAGMA_ENABLE_DEPRECATION_WARNINGS
-	}
-
-	Super::PostEditChangeProperty(PropertyChangedEvent);
-}
-#endif // WITH_EDITOR
-
 bool UChaosClothComponent::RequiresPreEndOfFrameSync() const
 {
 	if (!IsSimulationSuspended() && !ShouldWaitForParallelSimulationInTickComponent())
@@ -366,6 +309,63 @@
 	Super::OnAttachmentChanged();
 }
 
+void UChaosClothComponent::RefreshBoneTransforms(FActorComponentTickFunction* TickFunction)
+{
+	MarkRenderDynamicDataDirty();
+
+	bNeedToFlipSpaceBaseBuffers = true;
+	bHasValidBoneTransform = false;
+	FlipEditableSpaceBases();
+	bHasValidBoneTransform = true;
+}
+
+void UChaosClothComponent::GetUpdateClothSimulationData_AnyThread(TMap<int32, FClothSimulData>& OutClothSimulData, FMatrix& OutLocalToWorld, float& OutBlendWeight)
+{
+	OutLocalToWorld = GetComponentToWorld().ToMatrixWithScale();
+
+	const UChaosClothComponent* const LeaderPoseClothComponent = Cast<UChaosClothComponent>(LeaderPoseComponent.Get());
+	if (LeaderPoseClothComponent && LeaderPoseClothComponent->ClothSimulationProxy && bBindToLeaderComponent)
+	{
+		OutBlendWeight = BlendWeight;
+		OutClothSimulData = LeaderPoseClothComponent->ClothSimulationProxy->GetCurrentSimulationData_AnyThread();
+	}
+	else if (bEnableSimulation && !bBindToLeaderComponent && ClothSimulationProxy)
+	{
+		OutBlendWeight = BlendWeight;
+		OutClothSimulData = ClothSimulationProxy->GetCurrentSimulationData_AnyThread();
+	}
+	else
+	{
+		OutClothSimulData.Reset();
+	}
+
+	// Blend cloth out whenever the simulation data is invalid
+	if (!OutClothSimulData.Num())
+	{
+		OutBlendWeight = 0.0f;
+	}
+}
+
+void UChaosClothComponent::SetSkinnedAssetAndUpdate(USkinnedAsset* InSkinnedAsset, bool bReinitPose)
+{
+	if (InSkinnedAsset != GetSkinnedAsset())
+	{
+		Super::SetSkinnedAssetAndUpdate(InSkinnedAsset, bReinitPose);
+
+		// Stop current simulation
+		HandleExistingParallelSimulation();
+
+		// Update the component visibility in case the new render mesh has no valid LOD
+		UpdateVisibility();
+	}
+}
+
+TSharedPtr<UE::Chaos::ClothAsset::FClothSimulationProxy> UChaosClothComponent::CreateClothSimulationProxy()
+{
+	using namespace UE::Chaos::ClothAsset;
+	return MakeShared<FClothSimulationProxy>(*this);
+}
+
 void UChaosClothComponent::StartNewParallelSimulation(float DeltaTime)
 {
 	if (ClothSimulationProxy.IsValid())
