Chaos Cloth - Remove MaxDistance hard coded to 1000 hack.

#rnx

diff --git a/Engine/Plugins/Experimental/ChaosClothAsset/Source/ChaosClothAssetEngine/Private/ChaosClothAsset/ClothSimulationProxy.cpp b/Engine/Plugins/Experimental/ChaosClothAsset/Source/ChaosClothAssetEngine/Private/ChaosClothAsset/ClothSimulationProxy.cpp
--- a/Engine/Plugins/Experimental/ChaosClothAsset/Source/ChaosClothAssetEngine/Private/ChaosClothAsset/ClothSimulationProxy.cpp
+++ b/Engine/Plugins/Experimental/ChaosClothAsset/Source/ChaosClothAssetEngine/Private/ChaosClothAsset/ClothSimulationProxy.cpp
@@ -222,8 +222,6 @@ PRAGMA_ENABLE_DEPRECATION_WARNINGS
 		for (const TUniquePtr<FClothingSimulationCloth>& Cloth : Cloths)
 		{
 			// Update Cloth animatable parameters while in the cloth loop
-			Cloth->SetMaxDistancesMultiplier(1000.0);
-
 			if (bNeedsReset)
 			{
 				Cloth->Reset();
