Chaos Cloth:
- Added correct getter and setter for the cloth asset.
- Added simulation proxy reset when changing the cloth asset and the component is already registered.
- Improved management around component registration.
- Removed a bunch of duplicated constructor initializations that dated from before the Cloth Component inherited from USkinnedMeshComponent.
- Fixed warning on too many Finds caused by looking up for the cvar p.Cloth.MaxDeltaTimeTeleportMultiplierin when updating the simulation context.

#rnx

diff --git a/Engine/Plugins/Experimental/ChaosClothAsset/Source/ChaosClothAssetEngine/Private/ChaosClothAsset/ClothComponent.cpp b/Engine/Plugins/Experimental/ChaosClothAsset/Source/ChaosClothAssetEngine/Private/ChaosClothAsset/ClothComponent.cpp
--- a/Engine/Plugins/Experimental/ChaosClothAsset/Source/ChaosClothAssetEngine/Private/ChaosClothAsset/ClothComponent.cpp
+++ b/Engine/Plugins/Experimental/ChaosClothAsset/Source/ChaosClothAssetEngine/Private/ChaosClothAsset/ClothComponent.cpp
@@ -22,58 +22,7 @@ UChaosClothComponent::UChaosClothComponent(const FObjectInitializer& ObjectIniti
 	, bSuspendSimulation(0)
 	, bBindToLeaderComponent(0)
 {
-	bAutoActivate = true;
-	PrimaryComponentTick.bCanEverTick = true;
-	PrimaryComponentTick.TickGroup = TG_PrePhysics;
 	PrimaryComponentTick.EndTickGroup = TG_PostPhysics;
-
-	//VisibilityBasedAnimTickOption = EVisibilityBasedAnimTickOption::AlwaysTickPoseAndRefreshBones;
-
-	SetCollisionProfileName(UCollisionProfile::NoCollision_ProfileName);
-
-	StreamingDistanceMultiplier = 1.0f;
-	//bCanHighlightSelectedSections = false;
-	CanCharacterStepUpOn = ECB_Owner;
-#if WITH_EDITORONLY_DATA
-	//SectionIndexPreview = -1;
-	//MaterialIndexPreview = -1;
-
-	//SelectedEditorSection = INDEX_NONE;
-	//SelectedEditorMaterial = INDEX_NONE;
-#endif // WITH_EDITORONLY_DATA
-	bCastCapsuleDirectShadow = false;
-	bCastCapsuleIndirectShadow = false;
-	CapsuleIndirectShadowMinVisibility = .1f;
-
-	bDoubleBufferedComponentSpaceTransforms = true;
-	//LastStreamerUpdateBoundsRadius = -1.0;
-	CurrentEditableComponentTransforms = 0;
-	CurrentReadComponentTransforms = 1;
-	bNeedToFlipSpaceBaseBuffers = false;
-	bBoneVisibilityDirty = false;
-
-	//bUpdateDeformerAtNextTick = false;
-
-	bCanEverAffectNavigation = false;
-	//MasterBoneMapCacheCount = 0;
-	bSyncAttachParentLOD = true;
-	//bIgnoreMasterPoseComponentLOD = false;
-
-	//ExternalInterpolationAlpha = 0.0f;
-	//ExternalDeltaTime = 0.0f;
-	//ExternalTickRate = 1;
-	//bExternalInterpolate = false;
-	//bExternalUpdate = false;
-	//bExternalEvaluationRateLimited = false;
-	//bExternalTickRateControlled = false;
-
-	bMipLevelCallbackRegistered = false;
-
-#if !(UE_BUILD_SHIPPING || UE_BUILD_TEST)
-	//bDrawDebugSkeleton = false;
-#endif
-
-	//CurrentSkinWeightProfileName = NAME_None;
 }
 
 UChaosClothComponent::UChaosClothComponent(FVTableHelper& Helper)
@@ -115,24 +64,49 @@ bool UChaosClothComponent::IsSimulationEnabled() const
 
 void UChaosClothComponent::ResetConfigProperties()
 {
-	using namespace UE::Chaos::ClothAsset;
-	using namespace ::Chaos::Softs;
+	if (IsRegistered())
+	{
+		check(PropertyCollection);
 
-	// Create a mutable facade for our component's property collection, and use it to copy the properties from the cloth collection
-	FCollectionPropertyMutableFacade CollectionPropertyMutableFacade(PropertyCollection);
+		// Create a mutable facade for our component's property collection, and use it to copy the properties from the cloth collection
+		::Chaos::Softs::FCollectionPropertyMutableFacade CollectionPropertyMutableFacade(PropertyCollection);
 
-	if (GetClothAsset() && GetClothAsset()->GetClothCollection())
-	{
-		CollectionPropertyMutableFacade.Copy(*GetClothAsset()->GetClothCollection());
+		if (GetClothAsset() && GetClothAsset()->GetClothCollection())
+		{
+			CollectionPropertyMutableFacade.Copy(*GetClothAsset()->GetClothCollection());
+		}
+		else
+		{
+			PropertyCollection->Reset();
+			CollectionPropertyMutableFacade.DefineSchema();
+		}
 	}
 	else
 	{
-		PropertyCollection->Reset();
-		CollectionPropertyMutableFacade.DefineSchema();
+		UE_LOG(LogChaosClothAsset, Warning, TEXT("Chaos Cloth Component [%s]: Trying to reset runtime config properties without being registered."), *GetName());
 	}
+}
 
-	// The facade used to access the properties shouldn't be mutable, since adding/removing properties would cause issues
-	CollectionPropertyFacade = MakeUnique<FCollectionPropertyFacade>(PropertyCollection);
+void UChaosClothComponent::RecreateClothSimulationProxy()
+{
+	if (IsRegistered())
+	{
+		ClothSimulationProxy.Reset();
+
+		if (GetClothAsset())
+		{
+			const TSharedPtr<const FChaosClothSimulationModel> ClothSimulationModel = GetClothAsset()->GetClothSimulationModel();
+			if (ensure(ClothSimulationModel) && ClothSimulationModel->GetNumLods())
+			{
+				// Create the simulation proxy (note CreateClothSimulationProxy() can be overloaded)
+				ClothSimulationProxy = CreateClothSimulationProxy();
+			}
+		}
+	}
+	else
+	{
+		UE_LOG(LogChaosClothAsset, Warning, TEXT("Chaos Cloth Component [%s]: Trying to recreate the simulation proxy without being registered."), *GetName());
+	}
 }
 
 void UChaosClothComponent::PostLoad()
@@ -166,35 +140,23 @@ void UChaosClothComponent::PostEditChangeProperty(FPropertyChangedEvent& Propert
 
 void UChaosClothComponent::OnRegister()
 {
-	using namespace UE::Chaos::ClothAsset;
-
 	LLM_SCOPE(ELLMTag::Chaos);
 
+	// Register the component first, otherwise calls to ResetConfigProperties and RecreateClothSimulationProxy wouldn't work
 	Super::OnRegister();
 
-	if (GetClothAsset())
-	{
-		const TSharedPtr<const FChaosClothSimulationModel> ClothSimulationModel = GetClothAsset()->GetClothSimulationModel();
-		if (ensure(ClothSimulationModel) && ClothSimulationModel->GetNumLods())
-		{
-			if (!LeaderPoseComponent.IsValid())
-			{
-				FSkeletalMeshLODRenderData& LODData = GetClothAsset()->GetResourceForRendering()->LODRenderData[GetPredictedLODLevel()];
-				GetClothAsset()->FillComponentSpaceTransforms(GetClothAsset()->GetRefSkeleton().GetRefBonePose(), LODData.RequiredBones, GetEditableComponentSpaceTransforms());
+	// Update the component bone transforms (for colliders) from the cloth asset until these are animated from a leader component
+	UpdateComponentSpaceTransforms();
 
-				bNeedToFlipSpaceBaseBuffers = true; // Have updated space bases so need to flip
-				FlipEditableSpaceBases();
-				bHasValidBoneTransform = true;
-			}
+	// Create the runtime simulation property 
+	PropertyCollection = MakeShared<FManagedArrayCollection>();
+	CollectionPropertyFacade = MakeUnique<::Chaos::Softs::FCollectionPropertyFacade>(PropertyCollection);
 
-			// Create cloth property collection, and fill up the collection with the original cloth asset properties
-			PropertyCollection = MakeShared<FManagedArrayCollection>();
-			ResetConfigProperties();
+	// Fill up the property collection with the original cloth asset properties
+	ResetConfigProperties();
 
-			// Create simulation proxy
-			ClothSimulationProxy = CreateClothSimulationProxy();
-		}
-	}
+	// Create the proxy to start the simulation
+	RecreateClothSimulationProxy();
 
 	// Update render visibility, so that an empty LODs doesn't unnecessarily go to render
 	UpdateVisibility();
@@ -206,6 +168,10 @@ void UChaosClothComponent::OnUnregister()
 
 	// Release cloth simulation
 	ClothSimulationProxy.Reset();
+
+	// Release the runtime simulation collection and facade
+	PropertyCollection.Reset();
+	CollectionPropertyFacade.Reset();
 }
 
 bool UChaosClothComponent::IsComponentTickEnabled() const
@@ -309,7 +275,7 @@ void UChaosClothComponent::OnAttachmentChanged()
 	Super::OnAttachmentChanged();
 }
 
-void UChaosClothComponent::RefreshBoneTransforms(FActorComponentTickFunction* TickFunction)
+void UChaosClothComponent::RefreshBoneTransforms(FActorComponentTickFunction* /*TickFunction*/)
 {
 	MarkRenderDynamicDataDirty();
 
@@ -350,10 +316,22 @@ void UChaosClothComponent::SetSkinnedAssetAndUpdate(USkinnedAsset* InSkinnedAsse
 {
 	if (InSkinnedAsset != GetSkinnedAsset())
 	{
+		// Note: It is not necessary to stop the current simulation here, since it will die off once the proxy is recreated
+
+		// Change the skinned asset, dirty render states, ...etc.
 		Super::SetSkinnedAssetAndUpdate(InSkinnedAsset, bReinitPose);
 
-		// Stop current simulation
-		HandleExistingParallelSimulation();
+		if (IsRegistered())
+		{
+			// Update the component bone transforms (for colliders) from the new cloth asset
+			UpdateComponentSpaceTransforms();
+
+			// Fill up the property collection with the new cloth asset properties
+			ResetConfigProperties();
+
+			// Hard reset the simulation
+			RecreateClothSimulationProxy();
+		}
 
 		// Update the component visibility in case the new render mesh has no valid LOD
 		UpdateVisibility();
@@ -373,7 +351,7 @@ void UChaosClothComponent::StartNewParallelSimulation(float DeltaTime)
 		CSV_SCOPED_TIMING_STAT(Animation, Cloth);
 		const bool bIsSimulating = ClothSimulationProxy->Tick_GameThread(DeltaTime);
 
-		if (bIsSimulating)
+		if (bIsSimulating && CollectionPropertyFacade.IsValid())
 		{
 			CollectionPropertyFacade->ClearDirtyFlags();
 		}
@@ -403,6 +381,21 @@ bool UChaosClothComponent::ShouldWaitForParallelSimulationInTickComponent() cons
 	return bWaitForParallelTask || (CVarClothPhysicsWaitForParallelClothTask && CVarClothPhysicsWaitForParallelClothTask->GetBool());
 }
 
+void UChaosClothComponent::UpdateComponentSpaceTransforms()
+{
+	check(IsRegistered());
+
+	if (!LeaderPoseComponent.IsValid() && GetClothAsset())
+	{
+		FSkeletalMeshLODRenderData& LODData = GetClothAsset()->GetResourceForRendering()->LODRenderData[GetPredictedLODLevel()];
+		GetClothAsset()->FillComponentSpaceTransforms(GetClothAsset()->GetRefSkeleton().GetRefBonePose(), LODData.RequiredBones, GetEditableComponentSpaceTransforms());
+
+		bNeedToFlipSpaceBaseBuffers = true; // Have updated space bases so need to flip
+		FlipEditableSpaceBases();
+		bHasValidBoneTransform = true;
+	}
+}
+
 void UChaosClothComponent::UpdateVisibility()
 {
 	if (GetClothAsset())
diff --git a/Engine/Plugins/Experimental/ChaosClothAsset/Source/ChaosClothAssetEngine/Private/ChaosClothAsset/ClothSimulationContext.cpp b/Engine/Plugins/Experimental/ChaosClothAsset/Source/ChaosClothAssetEngine/Private/ChaosClothAsset/ClothSimulationContext.cpp
--- a/Engine/Plugins/Experimental/ChaosClothAsset/Source/ChaosClothAssetEngine/Private/ChaosClothAsset/ClothSimulationContext.cpp
+++ b/Engine/Plugins/Experimental/ChaosClothAsset/Source/ChaosClothAssetEngine/Private/ChaosClothAsset/ClothSimulationContext.cpp
@@ -18,7 +18,7 @@ namespace UE::Chaos::ClothAsset
 		LodIndex = ClothComponent.GetPredictedLODLevel();
 
 		// Set the teleport mode
-		const IConsoleVariable* const CVarMaxDeltaTimeTeleportMultiplier = IConsoleManager::Get().FindConsoleVariable(TEXT("p.Cloth.MaxDeltaTimeTeleportMultiplier"));
+		static const IConsoleVariable* const CVarMaxDeltaTimeTeleportMultiplier = IConsoleManager::Get().FindConsoleVariable(TEXT("p.Cloth.MaxDeltaTimeTeleportMultiplier"));
 		constexpr float MaxDeltaTimeTeleportMultiplierDefault = 1.5f;
 		const float MaxDeltaTimeTeleportMultiplier = CVarMaxDeltaTimeTeleportMultiplier ? CVarMaxDeltaTimeTeleportMultiplier->GetFloat() : MaxDeltaTimeTeleportMultiplierDefault;
 
diff --git a/Engine/Plugins/Experimental/ChaosClothAsset/Source/ChaosClothAssetEngine/Public/ChaosClothAsset/ClothComponent.h b/Engine/Plugins/Experimental/ChaosClothAsset/Source/ChaosClothAssetEngine/Public/ChaosClothAsset/ClothComponent.h
--- a/Engine/Plugins/Experimental/ChaosClothAsset/Source/ChaosClothAssetEngine/Public/ChaosClothAsset/ClothComponent.h
+++ b/Engine/Plugins/Experimental/ChaosClothAsset/Source/ChaosClothAssetEngine/Public/ChaosClothAsset/ClothComponent.h
@@ -86,6 +86,10 @@ public:
 	UFUNCTION(BlueprintCallable, Category = "Components|Simulation", Meta = (Keywords = "Chaos Cloth Config Property"))
 	void ResetConfigProperties();
 
+	/** Hard reset the cloth simulation by recreating the proxy. */
+	UFUNCTION(BlueprintCallable, Category = "Components|Simulation", Meta = (Keywords = "Chaos Cloth Recreate Simulation Proxy"))
+	void RecreateClothSimulationProxy();
+
 	/**
 	 * Return the property collection holding the runtime properties for this cloth component.
 	 * This might be different from the cloth asset's since the component's properties can be modified in code or in blueprints.
@@ -130,6 +134,7 @@ private:
 	void StartNewParallelSimulation(float DeltaTime);
 	void HandleExistingParallelSimulation();
 	bool ShouldWaitForParallelSimulationInTickComponent() const;
+	void UpdateComponentSpaceTransforms();
 	void UpdateVisibility();
 
 	friend UE::Chaos::ClothAsset::FClothComponentCacheAdapter;
@@ -137,7 +142,7 @@ private:
 #if WITH_EDITORONLY_DATA
 	/** Cloth asset used by this component. */
 	UE_DEPRECATED(5.1, "This property isn't deprecated, but getter and setter must be used at all times to preserve correct operations.")
-	UPROPERTY(EditAnywhere, Transient, BlueprintSetter = SetClothAsset, BlueprintGetter = GetClothAsset, Category = ClothAsset)
+	UPROPERTY(EditAnywhere, Transient, Setter = SetClothAsset, BlueprintSetter = SetClothAsset, Getter = GetClothAsset, BlueprintGetter = GetClothAsset, Category = ClothAsset)
 	TObjectPtr<UChaosClothAsset> ClothAsset;
 #endif
 
