Chaos Cloth Asset: Fix sim vertex indexing bug when merging collections.

diff --git a/Engine/Plugins/Experimental/ChaosClothAsset/Source/ChaosClothAsset/Private/ChaosClothAsset/CollectionClothPatternFacade.cpp b/Engine/Plugins/Experimental/ChaosClothAsset/Source/ChaosClothAsset/Private/ChaosClothAsset/CollectionClothPatternFacade.cpp
--- a/Engine/Plugins/Experimental/ChaosClothAsset/Source/ChaosClothAsset/Private/ChaosClothAsset/CollectionClothPatternFacade.cpp
+++ b/Engine/Plugins/Experimental/ChaosClothAsset/Source/ChaosClothAsset/Private/ChaosClothAsset/CollectionClothPatternFacade.cpp
@@ -382,9 +382,9 @@ namespace UE::Chaos::ClothAsset
 			GetBaseElementIndex(),
 			GetElementIndex());
 		const int32 OtherLodSimVerticesOffset = Other.ClothCollection->GetElementsOffset(
-			ClothCollection->GetSimVerticesStart(),
-			GetBaseElementIndex(),
-			GetElementIndex());
+			Other.ClothCollection->GetSimVerticesStart(),
+			Other.GetBaseElementIndex(),
+			Other.GetElementIndex());
 		const FIntVector3 OtherToLodSimVerticesOffset(LodSimVerticesOffset - OtherLodSimVerticesOffset);
 
 		const TConstArrayView<FIntVector3> OtherSimIndices = Other.GetSimIndices();
