Adding a new scope cycle counter to trace Cloth Component tick

diff --git a/Engine/Plugins/Experimental/ChaosClothAsset/Source/ChaosClothAssetEngine/Private/ChaosClothAsset/ClothComponent.cpp b/Engine/Plugins/Experimental/ChaosClothAsset/Source/ChaosClothAssetEngine/Private/ChaosClothAsset/ClothComponent.cpp
--- a/Engine/Plugins/Experimental/ChaosClothAsset/Source/ChaosClothAssetEngine/Private/ChaosClothAsset/ClothComponent.cpp
+++ b/Engine/Plugins/Experimental/ChaosClothAsset/Source/ChaosClothAssetEngine/Private/ChaosClothAsset/ClothComponent.cpp
@@ -11,6 +11,7 @@
 #include "Rendering/SkeletalMeshRenderData.h"
 #include "SkeletalRenderPublic.h"
 #include "PhysicsEngine/PhysicsAsset.h"
+#include "Stats/Stats.h"
 
 CSV_DECLARE_CATEGORY_MODULE_EXTERN(ENGINE_API, Animation);
 
@@ -182,6 +183,7 @@ bool UChaosClothComponent::IsComponentTickEnabled() const
 void UChaosClothComponent::TickComponent(float DeltaTime, enum ELevelTick TickType, FActorComponentTickFunction* ThisTickFunction)
 {
 	CSV_SCOPED_TIMING_STAT_EXCLUSIVE(Physics);
+	QUICK_SCOPE_CYCLE_COUNTER(STAT_ClothComponentTick);
 
 	// TODO: Fields
 	//if (ClothingSimulation)
