ClothCollection Schema rework. Existing cloth assets will be emptied, but their associated Dataflow graphs and nodes are still valid (with the exception of AddWeightMapNodes which need to be repainted because indexing has changed).

The topline changes are:

* No more Lods in the collection. Instead, have one collection per Lod in the ClothAsset.
* Separate SimPatterns and RenderPatterns. SimPattern = contiguous 2D island. RenderPattern = SkeletalMeshSection (single material)
* Separate 2D and 3D sim vertices. 2D vertices are unwelded patterns. 3D vertices are welded. Welding and bookkeeping is handled by the new FCollectionClothSeamFacade
** Most sim vertex data live on 3D vertices since that is how they will be consumed by the solver.
** Weight maps live on 3D vertices. Cloth paint tool has been updated to show welded 2D values changing together when in 2D mode.
* Tethers are now stored on their dynamic node as SimVertex3D data. They are batched at the end when generating the final asset. (This is because seaming after tether creation could generate invalid batches)

Callouts for a few features NOT included in the CL that will need to be done in the future. They were not implemented in the previous version either, so no loss of functionality, but they need to be implemented for completeness:
* SkeletalMeshImport as sim mesh. We're now requiring that SimPatterns are 2D islands, so need to possibly generate multiple SimPatterns.
* Multiple LODs not fully supported. Missing Up and Down transition data as well as handling property collection per LOD at the solver level.

Features not implemented because we lack DataflowNodes to make these construction changes, but will need to be implemented
* Ability to remove seams and split the associated 3D vertices.
* Cleanup/compactify 2D and 3D sim vertices if only 2D or only 3D vertices are removed.

diff --git a/Engine/Plugins/Enterprise/DatasmithImporter/Source/DatasmithImporter/Private/DatasmithImporter.cpp b/Engine/Plugins/Enterprise/DatasmithImporter/Source/DatasmithImporter/Private/DatasmithImporter.cpp
--- a/Engine/Plugins/Enterprise/DatasmithImporter/Source/DatasmithImporter/Private/DatasmithImporter.cpp
+++ b/Engine/Plugins/Enterprise/DatasmithImporter/Source/DatasmithImporter/Private/DatasmithImporter.cpp
@@ -495,52 +495,46 @@ void FDatasmithImporter::ImportClothes(FDatasmithImportContext& ImportContext)
 
 		using namespace UE::Chaos::ClothAsset;
 
-		TSharedPtr<FManagedArrayCollection> Collection = ClothAsset->GetClothCollection();
-		FCollectionClothFacade Cloth(Collection);
+		TArray<TSharedPtr<FManagedArrayCollection>>& Collections = ClothAsset->GetClothCollections();
+		Collections.SetNum(1);
+		FCollectionClothFacade Cloth(Collections[0]);
 		Cloth.Reset(); // on reimport, asset is non-empty
-		FCollectionClothLodFacade ClothLod = Cloth.GetNumLods() ? Cloth.GetLod(0) : Cloth.AddGetLod();
 
 		for (FDatasmithClothPattern& Pattern : DsCloth.Patterns)
 		{
 			if (Pattern.IsValid())
 			{
-				FCollectionClothPatternFacade ClothPattern = ClothLod.AddGetPattern();
+				FCollectionClothSimPatternFacade ClothPattern = Cloth.AddGetSimPattern();
 				ClothPattern.Initialize(Pattern.SimPosition, Pattern.SimRestPosition, Pattern.SimTriangleIndices);
 			}
 		}
 
-		ClothLod.SetNumSeams(DsCloth.Sewing.Num());
-		TArrayView<FIntVector2> SeamPatterns = ClothLod.GetSeamPatterns();
-		TArrayView<TArray<FIntVector2>> SeamStitches = ClothLod.GetSeamStitches();
-
-		int32 SeamIndex = 0;
 		for (const FDatasmithClothSewingInfo& SeamInfo : DsCloth.Sewing)
 		{
 			const int32 SeamPattern0 = (int32)SeamInfo.Seam0PanelIndex;
 			const int32 SeamPattern1 = (int32)SeamInfo.Seam1PanelIndex;
 
-			if (SeamPattern0 >= 0 && SeamPattern0 < ClothLod.GetNumPatterns() &&
-				SeamPattern1 >= 0 && SeamPattern1 < ClothLod.GetNumPatterns())
+			if (SeamPattern0 >= 0 && SeamPattern0 < Cloth.GetNumSimPatterns() &&
+				SeamPattern1 >= 0 && SeamPattern1 < Cloth.GetNumSimPatterns())
 			{
-				SeamPatterns[SeamIndex] = FIntVector2(SeamPattern0, SeamPattern1);
-
-				const FCollectionClothPatternConstFacade ClothPattern0 = ClothLod.GetPattern(SeamPattern0);
-				const FCollectionClothPatternConstFacade ClothPattern1 = ClothLod.GetPattern(SeamPattern1);
+				const FCollectionClothSimPatternConstFacade ClothPattern0 = Cloth.GetSimPattern(SeamPattern0);
+				const FCollectionClothSimPatternConstFacade ClothPattern1 = Cloth.GetSimPattern(SeamPattern1);
 
-				const int32 ClothPattern0VerticesOffset = ClothPattern0.GetSimVerticesOffset();
-				const int32 ClothPattern1VerticesOffset = ClothPattern1.GetSimVerticesOffset();
+				const int32 ClothPattern0VerticesOffset = ClothPattern0.GetSimVertices2DOffset();
+				const int32 ClothPattern1VerticesOffset = ClothPattern1.GetSimVertices2DOffset();
 
-				TArray<FIntVector2>& Stitches = SeamStitches[SeamIndex];
-				const int32 StitchesCount = FMath::Min(SeamInfo.Seam0MeshIndices.Num(), SeamInfo.Seam1MeshIndices.Num());
+				TArray<FIntVector2> Stitches;
+				const uint32 StitchesCount = FMath::Min(SeamInfo.Seam0MeshIndices.Num(), SeamInfo.Seam1MeshIndices.Num());
 				Stitches.Reserve(StitchesCount);
-				for (int32 StitchIndex = 0; StitchIndex < StitchesCount; ++StitchIndex)
+				for (uint32 StitchIndex = 0; StitchIndex < StitchesCount; ++StitchIndex)
 				{
 					Stitches.Emplace(
-						(int32)SeamInfo.Seam0MeshIndices[StitchIndex] + ClothPattern0VerticesOffset,
+						(int32)SeamInfo.Seam0MeshIndices[StitchIndex] + ClothPattern0VerticesOffset, 
 						(int32)SeamInfo.Seam1MeshIndices[StitchIndex] + ClothPattern1VerticesOffset);
 				}
 
-				++SeamIndex;
+				FCollectionClothSeamFacade SeamFacade = Cloth.AddGetSeam();
+				SeamFacade.Initialize(Stitches);
 			}
 		}
 
diff --git a/Engine/Plugins/Experimental/ChaosClothAsset/Source/ChaosClothAsset/Private/ChaosClothAsset/ClothCollection.cpp b/Engine/Plugins/Experimental/ChaosClothAsset/Source/ChaosClothAsset/Private/ChaosClothAsset/ClothCollection.cpp
--- a/Engine/Plugins/Experimental/ChaosClothAsset/Source/ChaosClothAsset/Private/ChaosClothAsset/ClothCollection.cpp
+++ b/Engine/Plugins/Experimental/ChaosClothAsset/Source/ChaosClothAsset/Private/ChaosClothAsset/ClothCollection.cpp
@@ -6,132 +6,117 @@
 namespace UE::Chaos::ClothAsset::Private
 {
 	// Groups
-	static const FName LodsGroup(TEXT("Lods"));
-	static const FName MaterialsGroup(TEXT("Materials"));
-	static const FName TetherBatchesGroup(TEXT("TetherBatches"));
-	static const FName TethersGroup(TEXT("Tethers"));
+	static const FName LodsGroup(TEXT("Lods")); // One lod per collection, but keeping the name for now to represent whatever is a single scalar for the whole collection
 	static const FName SeamsGroup(TEXT("Seams"));
-	static const FName PatternsGroup(TEXT("Patterns"));
+	static const FName SeamStitchesGroup(TEXT("SeamStitches"));
+	static const FName SimPatternsGroup(TEXT("SimPatterns"));
+	static const FName RenderPatternsGroup(TEXT("RenderPatterns"));
 	static const FName SimFacesGroup(TEXT("SimFaces"));
-	static const FName SimVerticesGroup(TEXT("SimVertices"));
+	static const FName SimVertices2DGroup(TEXT("SimVertices2D"));
+	static const FName SimVertices3DGroup(TEXT("SimVertices3D"));
 	static const FName RenderFacesGroup(TEXT("RenderFaces"));
 	static const FName RenderVerticesGroup(TEXT("RenderVertices"));
 
-	// LODs Group
-	static const FName MaterialStartAttribute(TEXT("MaterialStart"));
-	static const FName MaterialEndAttribute(TEXT("MaterialEnd"));
-	static const FName TetherBatchStartAttribute(TEXT("TetherBatchStart"));
-	static const FName TetherBatchEndAttribute(TEXT("TetherBatchEnd"));
-	static const FName SeamStartAttribute(TEXT("SeamStart"));
-	static const FName SeamEndAttribute(TEXT("SeamEnd"));
-	static const FName PatternStartAttribute(TEXT("PatternStart"));
-	static const FName PatternEndAttribute(TEXT("PatternEnd"));
+	// Lods Group
 	static const FName PhysicsAssetPathNameAttribute(TEXT("PhysicsAssetPathName"));
 	static const FName SkeletonAssetPathNameAttribute(TEXT("SkeletonAssetPathName"));
 	static const TArray<FName> LodsGroupAttributes =
 	{
-		MaterialStartAttribute,
-		MaterialEndAttribute,
-		TetherBatchStartAttribute,
-		TetherBatchEndAttribute,
-		SeamStartAttribute,
-		SeamEndAttribute,
-		PatternStartAttribute,
-		PatternEndAttribute,
 		PhysicsAssetPathNameAttribute,
 		SkeletonAssetPathNameAttribute
 	};
 
-	// Materials Group
-	static const FName RenderMaterialPathNameAttribute(TEXT("RenderMaterialPathName"));
-	static const TArray<FName> MaterialsGroupAttributes =
-	{
-		RenderMaterialPathNameAttribute
-	};
-
-	// Tether Batches Group
-	static const FName TetherStartAttribute(TEXT("TetherStart"));
-	static const FName TetherEndAttribute(TEXT("TetherEnd"));
-	static const TArray<FName> TetherBatchesGroupAttributes =
-	{
-		TetherStartAttribute,
-		TetherEndAttribute
-	};
-
-	// Tethers Group
-	static const FName TetherKinematicIndexAttribute(TEXT("TetherKinematicIndex"));
-	static const FName TetherDynamicIndexAttribute(TEXT("TetherDynamicIndex"));
-	static const FName TetherReferenceLengthAttribute(TEXT("TetherReferenceLength"));
-	static const TArray<FName> TethersGroupAttributes =
+	// Seam Group
+	static const FName SeamStitchStartAttribute(TEXT("SeamStitchStart"));
+	static const FName SeamStitchEndAttribute(TEXT("SeamStitchEnd"));
+	static const TArray<FName> SeamsGroupAttributes =
 	{
-		TetherKinematicIndexAttribute,
-		TetherDynamicIndexAttribute,
-		TetherReferenceLengthAttribute
+		SeamStitchStartAttribute,
+		SeamStitchEndAttribute
 	};
 
-	// Seam Group
-	static const FName SeamPatternsAttribute(TEXT("SeamPatterns"));
-	static const FName SeamStitchesAttribute(TEXT("SeamStitches"));
-	static const TArray<FName> SeamsGroupAttributes =
+	// Seam Stitches Group
+	static const FName SeamStitch2DEndIndicesAttribute(TEXT("SeamStitch2DEndIndices"));
+	static const FName SeamStitch3DIndexAttribute(TEXT("SeamStitch3DIndex"));
+	static const TArray<FName> SeamStitchesGroupAttributes =
 	{
-		SeamPatternsAttribute,
-		SeamStitchesAttribute
+		SeamStitch2DEndIndicesAttribute,
+		SeamStitch3DIndexAttribute
 	};
 
-	// Patterns Group
-	static const FName SimVerticesStartAttribute(TEXT("SimVerticesStart"));
-	static const FName SimVerticesEndAttribute(TEXT("SimVerticesEnd"));
+	// Sim Patterns Group
+	static const FName SimVertices2DStartAttribute(TEXT("SimVertices2DStart"));
+	static const FName SimVertices2DEndAttribute(TEXT("SimVertices2DEnd"));
 	static const FName SimFacesStartAttribute(TEXT("SimFacesStart"));
 	static const FName SimFacesEndAttribute(TEXT("SimFacesEnd"));
-	static const FName RenderVerticesStartAttribute(TEXT("RenderVerticesStar"));
+	static const TArray<FName> SimPatternsGroupAttributes =
+	{ 
+		SimVertices2DStartAttribute,
+		SimVertices2DEndAttribute,
+		SimFacesStartAttribute,
+		SimFacesEndAttribute,
+	};
+
+	// RenderPatterns Group
+	static const FName RenderVerticesStartAttribute(TEXT("RenderVerticesStart"));
 	static const FName RenderVerticesEndAttribute(TEXT("RenderVerticesEnd"));
 	static const FName RenderFacesStartAttribute(TEXT("RenderFacesStart"));
 	static const FName RenderFacesEndAttribute(TEXT("RenderFacesEnd"));
-	static const FName StatusFlagsAttribute(TEXT("StatusFlags"));
-	static const TArray<FName> PatternsGroupAttributes =
+	static const FName RenderMaterialPathNameAttribute(TEXT("RenderMaterialPathName"));
+	static const TArray<FName> RenderPatternsGroupAttributes =
 	{
-		SimVerticesStartAttribute,
-		SimVerticesEndAttribute,
-		SimFacesStartAttribute,
-		SimFacesEndAttribute,
 		RenderVerticesStartAttribute,
 		RenderVerticesEndAttribute,
 		RenderFacesStartAttribute,
 		RenderFacesEndAttribute,
-		StatusFlagsAttribute
+		RenderMaterialPathNameAttribute
 	};
 
 	// Sim Faces Group
-	static const FName SimIndicesAttribute(TEXT("SimIndices"));
+	static const FName SimIndices2DAttribute(TEXT("SimIndices2D"));
+	static const FName SimIndices3DAttribute(TEXT("SimIndices3D"));
 	static const TArray<FName> SimFacesGroupAttributes =
 	{
-		SimIndicesAttribute
+		SimIndices2DAttribute,
+		SimIndices3DAttribute
+	};
+
+	// Sim Vertices 2D Group
+	static const FName SimPosition2DAttribute(TEXT("SimPosition2D"));
+	static const FName SimVertex3DLookupAttribute(TEXT("SimVertex3DLookup"));
+	static const TArray<FName> SimVertices2DGroupAttributes =
+	{
+		SimPosition2DAttribute,
+		SimVertex3DLookupAttribute,
 	};
 
-	// Sim Vertices Group
-	static const FName SimPositionAttribute(TEXT("SimPosition"));
-	static const FName SimRestPositionAttribute(TEXT("SimRestPosition"));
-	static const FName SimRestNormalAttribute(TEXT("SimRestNormal"));
-	static const FName SimNumBoneInfluencesAttribute(TEXT("SimNumBoneInfluences"));
+	// Sim Vertices 3D Group
+	// NOTE: if you add anything here, you need to implement how to merge it 
+	// in CollectionClothSeamsFacade. Otherwise, the data in the lowest vertex index 
+	// will survive and the other data will be lost whenever seams are added.
+	static const FName SimPosition3DAttribute(TEXT("SimPosition3D"));
+	static const FName SimNormalAttribute(TEXT("SimNormal"));
 	static const FName SimBoneIndicesAttribute(TEXT("SimBoneIndices"));
 	static const FName SimBoneWeightsAttribute(TEXT("SimBoneWeights"));
-	static const TArray<FName> SimVerticesGroupAttributes =
+	static const FName TetherKinematicIndexAttribute(TEXT("TetherKinematicIndex"));
+	static const FName TetherReferenceLengthAttribute(TEXT("TetherReferenceLength"));
+	static const FName SimVertex2DLookupAttribute(TEXT("SimVertex2DLookup"));
+	static const FName SeamStitchLookupAttribute(TEXT("SeamStitchLookup"));
+	static const TArray<FName> SimVertices3DGroupAttributes =
 	{
-		SimPositionAttribute,
-		SimRestPositionAttribute,
-		SimRestNormalAttribute,
-		SimNumBoneInfluencesAttribute,
+		SimPosition3DAttribute,
+		SimNormalAttribute,
 		SimBoneIndicesAttribute,
-		SimBoneWeightsAttribute
+		SimBoneWeightsAttribute,
+		SimVertex2DLookupAttribute,
+		SeamStitchLookupAttribute
 	};
 
 	// Render Faces Group
 	static const FName RenderIndicesAttribute(TEXT("RenderIndices"));
-	static const FName RenderMaterialIndexAttribute(TEXT("RenderMaterialIndex"));
 	static const TArray<FName> RenderFacesGroupAttributes =
 	{
 		RenderIndicesAttribute,
-		RenderMaterialIndexAttribute
 	};
 
 	// Render Vertices Group
@@ -141,7 +126,6 @@ namespace UE::Chaos::ClothAsset::Private
 	static const FName RenderTangentVAttribute(TEXT("RenderTangentV"));
 	static const FName RenderUVsAttribute(TEXT("RenderUVs"));
 	static const FName RenderColorAttribute(TEXT("RenderColor"));
-	static const FName RenderNumBoneInfluencesAttribute(TEXT("RenderNumBoneInfluences"));
 	static const FName RenderBoneIndicesAttribute(TEXT("RenderBoneIndices"));
 	static const FName RenderBoneWeightsAttribute(TEXT("RenderBoneWeights"));
 	static const TArray<FName> RenderVerticesGroupAttributes =
@@ -152,7 +136,6 @@ namespace UE::Chaos::ClothAsset::Private
 		RenderTangentVAttribute,
 		RenderUVsAttribute,
 		RenderColorAttribute,
-		RenderNumBoneInfluencesAttribute,
 		RenderBoneIndicesAttribute,
 		RenderBoneWeightsAttribute
 	};
@@ -160,13 +143,13 @@ namespace UE::Chaos::ClothAsset::Private
 	static const TMap<FName, TArray<FName>> FixedAttributeNamesMap =
 	{
 		{ LodsGroup, LodsGroupAttributes },
-		{ MaterialsGroup, MaterialsGroupAttributes },
-		{ TetherBatchesGroup, TetherBatchesGroupAttributes },
-		{ TethersGroup, TethersGroupAttributes },
 		{ SeamsGroup, SeamsGroupAttributes },
-		{ PatternsGroup, PatternsGroupAttributes },
+		{ SeamStitchesGroup, SeamStitchesGroupAttributes },
+		{ SimPatternsGroup, SimPatternsGroupAttributes },
+		{ RenderPatternsGroup, RenderPatternsGroupAttributes },
 		{ SimFacesGroup, SimFacesGroupAttributes },
-		{ SimVerticesGroup, SimVerticesGroupAttributes },
+		{ SimVertices2DGroup, SimVertices2DGroupAttributes },
+		{ SimVertices3DGroup, SimVertices3DGroupAttributes },
 		{ RenderFacesGroup, RenderFacesGroupAttributes },
 		{ RenderVerticesGroup, RenderVerticesGroupAttributes }
 	};
@@ -176,13 +159,13 @@ namespace UE::Chaos::ClothAsset
 {
 	// Groups
 	const FName FClothCollection::LodsGroup = Private::LodsGroup;
-	const FName FClothCollection::MaterialsGroup = Private::MaterialsGroup;
-	const FName FClothCollection::TetherBatchesGroup = Private::TetherBatchesGroup;
-	const FName FClothCollection::TethersGroup = Private::TethersGroup;
 	const FName FClothCollection::SeamsGroup = Private::SeamsGroup;
-	const FName FClothCollection::PatternsGroup = Private::PatternsGroup;
+	const FName FClothCollection::SeamStitchesGroup = Private::SeamStitchesGroup;
+	const FName FClothCollection::SimPatternsGroup = Private::SimPatternsGroup;
+	const FName FClothCollection::RenderPatternsGroup = Private::RenderPatternsGroup;
 	const FName FClothCollection::SimFacesGroup = Private::SimFacesGroup;
-	const FName FClothCollection::SimVerticesGroup = Private::SimVerticesGroup;
+	const FName FClothCollection::SimVertices2DGroup = Private::SimVertices2DGroup;
+	const FName FClothCollection::SimVertices3DGroup = Private::SimVertices3DGroup;
 	const FName FClothCollection::RenderFacesGroup = Private::RenderFacesGroup;
 	const FName FClothCollection::RenderVerticesGroup = Private::RenderVerticesGroup;
 
@@ -194,58 +177,50 @@ namespace UE::Chaos::ClothAsset
 		check(ManagedArrayCollection.IsValid());
 
 		// LODs Group
-		MaterialStart = ManagedArrayCollection->FindAttribute<int32>(MaterialStartAttribute, LodsGroup);
-		MaterialEnd = ManagedArrayCollection->FindAttribute<int32>(MaterialEndAttribute, LodsGroup);
-		TetherBatchStart = ManagedArrayCollection->FindAttribute<int32>(TetherBatchStartAttribute, LodsGroup);
-		TetherBatchEnd = ManagedArrayCollection->FindAttribute<int32>(TetherBatchEndAttribute, LodsGroup);
-		SeamStart = ManagedArrayCollection->FindAttribute<int32>(SeamStartAttribute, LodsGroup);
-		SeamEnd = ManagedArrayCollection->FindAttribute<int32>(SeamEndAttribute, LodsGroup);
-		PatternStart = ManagedArrayCollection->FindAttribute<int32>(PatternStartAttribute, LodsGroup);
-		PatternEnd = ManagedArrayCollection->FindAttribute<int32>(PatternEndAttribute, LodsGroup);
 		PhysicsAssetPathName = ManagedArrayCollection->FindAttribute<FString>(PhysicsAssetPathNameAttribute, LodsGroup);
 		SkeletonAssetPathName = ManagedArrayCollection->FindAttribute<FString>(SkeletonAssetPathNameAttribute, LodsGroup);
 
-		// Materials Group
-		RenderMaterialPathName = ManagedArrayCollection->FindAttribute<FString>(RenderMaterialPathNameAttribute, MaterialsGroup);
-
-		// Tether Batches Group
-		TetherStart = ManagedArrayCollection->FindAttribute<int32>(TetherStartAttribute, TetherBatchesGroup);
-		TetherEnd = ManagedArrayCollection->FindAttribute<int32>(TetherEndAttribute, TetherBatchesGroup);
-
-		// Tethers Group
-		TetherKinematicIndex = ManagedArrayCollection->FindAttribute<int32>(TetherKinematicIndexAttribute, TethersGroup);
-		TetherDynamicIndex = ManagedArrayCollection->FindAttribute<int32>(TetherDynamicIndexAttribute, TethersGroup);
-		TetherReferenceLength = ManagedArrayCollection->FindAttribute<float>(TetherReferenceLengthAttribute, TethersGroup);
-
 		// Seam Group
-		SeamPatterns = ManagedArrayCollection->FindAttribute<FIntVector2>(SeamPatternsAttribute, SeamsGroup);
-		SeamStitches = ManagedArrayCollection->FindAttribute<TArray<FIntVector2>>(SeamStitchesAttribute, SeamsGroup);
-
-		// Patterns Group
-		SimVerticesStart = ManagedArrayCollection->FindAttribute<int32>(SimVerticesStartAttribute, PatternsGroup);
-		SimVerticesEnd = ManagedArrayCollection->FindAttribute<int32>(SimVerticesEndAttribute, PatternsGroup);
-		SimFacesStart = ManagedArrayCollection->FindAttribute<int32>(SimFacesStartAttribute, PatternsGroup);
-		SimFacesEnd = ManagedArrayCollection->FindAttribute<int32>(SimFacesEndAttribute, PatternsGroup);
-		RenderVerticesStart = ManagedArrayCollection->FindAttribute<int32>(RenderVerticesStartAttribute, PatternsGroup);
-		RenderVerticesEnd = ManagedArrayCollection->FindAttribute<int32>(RenderVerticesEndAttribute, PatternsGroup);
-		RenderFacesStart = ManagedArrayCollection->FindAttribute<int32>(RenderFacesStartAttribute, PatternsGroup);
-		RenderFacesEnd = ManagedArrayCollection->FindAttribute<int32>(RenderFacesEndAttribute, PatternsGroup);
-		StatusFlags = ManagedArrayCollection->FindAttribute<int32>(StatusFlagsAttribute, PatternsGroup);
+		SeamStitchStart = ManagedArrayCollection->FindAttribute<int32>(SeamStitchStartAttribute, SeamsGroup);
+		SeamStitchEnd = ManagedArrayCollection->FindAttribute<int32>(SeamStitchEndAttribute, SeamsGroup);
+
+		// Seam Stitches Group
+		SeamStitch2DEndIndices = ManagedArrayCollection->FindAttribute<FIntVector2>(SeamStitch2DEndIndicesAttribute, SeamStitchesGroup);
+		SeamStitch3DIndex = ManagedArrayCollection->FindAttribute<int32>(SeamStitch3DIndexAttribute, SeamStitchesGroup);
+
+		// Sim Patterns Group
+		SimVertices2DStart = ManagedArrayCollection->FindAttribute<int32>(SimVertices2DStartAttribute, SimPatternsGroup);
+		SimVertices2DEnd = ManagedArrayCollection->FindAttribute<int32>(SimVertices2DEndAttribute, SimPatternsGroup);
+		SimFacesStart = ManagedArrayCollection->FindAttribute<int32>(SimFacesStartAttribute, SimPatternsGroup);
+		SimFacesEnd = ManagedArrayCollection->FindAttribute<int32>(SimFacesEndAttribute, SimPatternsGroup);
+
+		// Render Patterns Group
+		RenderVerticesStart = ManagedArrayCollection->FindAttribute<int32>(RenderVerticesStartAttribute, RenderPatternsGroup);
+		RenderVerticesEnd = ManagedArrayCollection->FindAttribute<int32>(RenderVerticesEndAttribute, RenderPatternsGroup);
+		RenderFacesStart = ManagedArrayCollection->FindAttribute<int32>(RenderFacesStartAttribute, RenderPatternsGroup);
+		RenderFacesEnd = ManagedArrayCollection->FindAttribute<int32>(RenderFacesEndAttribute, RenderPatternsGroup);
+		RenderMaterialPathName = ManagedArrayCollection->FindAttribute<FString>(RenderMaterialPathNameAttribute, RenderPatternsGroup);
 
 		// Sim Faces Group
-		SimIndices = ManagedArrayCollection->FindAttribute<FIntVector3>(SimIndicesAttribute, SimFacesGroup);
-
-		// Sim Vertices Group
-		SimPosition = ManagedArrayCollection->FindAttribute<FVector2f>(SimPositionAttribute, SimVerticesGroup);
-		SimRestPosition = ManagedArrayCollection->FindAttribute<FVector3f>(SimRestPositionAttribute, SimVerticesGroup);
-		SimRestNormal = ManagedArrayCollection->FindAttribute<FVector3f>(SimRestNormalAttribute, SimVerticesGroup);
-		SimNumBoneInfluences = ManagedArrayCollection->FindAttribute<int32>(SimNumBoneInfluencesAttribute, SimVerticesGroup);
-		SimBoneIndices = ManagedArrayCollection->FindAttribute<TArray<int32>>(SimBoneIndicesAttribute, SimVerticesGroup);
-		SimBoneWeights = ManagedArrayCollection->FindAttribute<TArray<float>>(SimBoneWeightsAttribute, SimVerticesGroup);
+		SimIndices2D = ManagedArrayCollection->FindAttribute<FIntVector3>(SimIndices2DAttribute, SimFacesGroup);
+		SimIndices3D = ManagedArrayCollection->FindAttribute<FIntVector3>(SimIndices3DAttribute, SimFacesGroup);
+
+		// Sim Vertices 2D Group
+		SimPosition2D = ManagedArrayCollection->FindAttribute<FVector2f>(SimPosition2DAttribute, SimVertices2DGroup);
+		SimVertex3DLookup = ManagedArrayCollection->FindAttribute<int32>(SimVertex3DLookupAttribute, SimVertices2DGroup);
+
+		// Sim Vertices 3D Group
+		SimPosition3D = ManagedArrayCollection->FindAttribute<FVector3f>(SimPosition3DAttribute, SimVertices3DGroup);
+		SimNormal = ManagedArrayCollection->FindAttribute<FVector3f>(SimNormalAttribute, SimVertices3DGroup);
+		SimBoneIndices = ManagedArrayCollection->FindAttribute<TArray<int32>>(SimBoneIndicesAttribute, SimVertices3DGroup);
+		SimBoneWeights = ManagedArrayCollection->FindAttribute<TArray<float>>(SimBoneWeightsAttribute, SimVertices3DGroup);
+		TetherKinematicIndex = ManagedArrayCollection->FindAttribute<TArray<int32>>(TetherKinematicIndexAttribute, SimVertices3DGroup);
+		TetherReferenceLength = ManagedArrayCollection->FindAttribute<TArray<float>>(TetherReferenceLengthAttribute, SimVertices3DGroup);
+		SimVertex2DLookup = ManagedArrayCollection->FindAttribute<TArray<int32>>(SimVertex2DLookupAttribute, SimVertices3DGroup);
+		SeamStitchLookup = ManagedArrayCollection->FindAttribute<TArray<int32>>(SeamStitchLookupAttribute, SimVertices3DGroup);
 
 		// Render Faces Group
 		RenderIndices = ManagedArrayCollection->FindAttribute<FIntVector3>(RenderIndicesAttribute, RenderFacesGroup);
-		RenderMaterialIndex = ManagedArrayCollection->FindAttribute<int32>(RenderMaterialIndexAttribute, RenderFacesGroup);
 
 		// Render Vertices Group
 		RenderPosition = ManagedArrayCollection->FindAttribute<FVector3f>(RenderPositionAttribute, RenderVerticesGroup);
@@ -254,7 +229,6 @@ namespace UE::Chaos::ClothAsset
 		RenderTangentV = ManagedArrayCollection->FindAttribute<FVector3f>(RenderTangentVAttribute, RenderVerticesGroup);
 		RenderUVs = ManagedArrayCollection->FindAttribute<TArray<FVector2f>>(RenderUVsAttribute, RenderVerticesGroup);
 		RenderColor = ManagedArrayCollection->FindAttribute<FLinearColor>(RenderColorAttribute, RenderVerticesGroup);
-		RenderNumBoneInfluences = ManagedArrayCollection->FindAttribute<int32>(RenderNumBoneInfluencesAttribute, RenderVerticesGroup);
 		RenderBoneIndices = ManagedArrayCollection->FindAttribute<TArray<int32>>(RenderBoneIndicesAttribute, RenderVerticesGroup);
 		RenderBoneWeights = ManagedArrayCollection->FindAttribute<TArray<float>>(RenderBoneWeightsAttribute, RenderVerticesGroup);
 	}
@@ -263,58 +237,50 @@ namespace UE::Chaos::ClothAsset
 	{
 		return 
 			// LODs Group
-			MaterialStart &&
-			MaterialEnd &&
-			TetherBatchStart &&
-			TetherBatchEnd &&
-			SeamStart &&
-			SeamEnd &&
-			PatternStart &&
-			PatternEnd &&
 			PhysicsAssetPathName &&
 			SkeletonAssetPathName &&
 
-			// Materials Group
-			RenderMaterialPathName &&
-
-			// Tether Batches Group
-			TetherStart &&
-			TetherEnd &&
-
-			// Tethers Group
-			TetherKinematicIndex &&
-			TetherDynamicIndex &&
-			TetherReferenceLength &&
-
 			// Seam Group
-			SeamPatterns &&
-			SeamStitches &&
+			SeamStitchStart &&
+			SeamStitchEnd &&
+
+			// Seam Stitches Group
+			SeamStitch2DEndIndices &&
+			SeamStitch3DIndex &&
 
-			// Patterns Group
-			SimVerticesStart &&
-			SimVerticesEnd &&
+			// Sim Patterns Group
+			SimVertices2DStart &&
+			SimVertices2DEnd &&
 			SimFacesStart &&
 			SimFacesEnd &&
+
+			// Render Patterns Group
 			RenderVerticesStart &&
 			RenderVerticesEnd &&
 			RenderFacesStart &&
 			RenderFacesEnd &&
-			StatusFlags &&
+			RenderMaterialPathName &&
 
 			// Sim Faces Group
-			SimIndices &&
+			SimIndices2D &&
+			SimIndices3D &&
 
-			// Sim Vertices Group
-			SimPosition &&
-			SimRestPosition &&
-			SimRestNormal &&
-			SimNumBoneInfluences &&
+			// Sim Vertices 2D Group
+			SimPosition2D  &&
+			SimVertex3DLookup &&
+
+			// Sim Vertices 3D Group
+			SimPosition3D &&
+			SimNormal &&
 			SimBoneIndices &&
 			SimBoneWeights &&
+			TetherKinematicIndex &&
+			TetherReferenceLength &&
+			SimVertex2DLookup &&
+			SeamStitchLookup &&
 
 			// Render Faces Group
 			RenderIndices &&
-			RenderMaterialIndex &&
 
 			// Render Vertices Group
 			RenderPosition &&
@@ -323,7 +289,6 @@ namespace UE::Chaos::ClothAsset
 			RenderTangentV &&
 			RenderUVs &&
 			RenderColor &&
-			RenderNumBoneInfluences &&
 			RenderBoneIndices &&
 			RenderBoneWeights;
 	}
@@ -333,69 +298,60 @@ namespace UE::Chaos::ClothAsset
 		using namespace UE::Chaos::ClothAsset::Private;
 
 		// Dependencies
-		FManagedArrayCollection::FConstructionParameters MaterialsDependency(MaterialsGroup);
-		FManagedArrayCollection::FConstructionParameters TetherBatchesDependency(TetherBatchesGroup);
-		FManagedArrayCollection::FConstructionParameters TethersDependency(TethersGroup);
-		FManagedArrayCollection::FConstructionParameters SeamsDependency(SeamsGroup);
-		FManagedArrayCollection::FConstructionParameters PatternsDependency(PatternsGroup);
-		FManagedArrayCollection::FConstructionParameters RenderFacesDependency(RenderFacesGroup);
-		FManagedArrayCollection::FConstructionParameters RenderVerticesDependency(RenderVerticesGroup);
-		FManagedArrayCollection::FConstructionParameters SimFacesDependency(SimFacesGroup);
-		FManagedArrayCollection::FConstructionParameters SimVerticesDependency(SimVerticesGroup);
+		constexpr bool bSaved = true;
+		constexpr bool bAllowCircularDependency = true;
+		FManagedArrayCollection::FConstructionParameters SeamStitchesDependency(SeamStitchesGroup, bSaved, bAllowCircularDependency);
+		FManagedArrayCollection::FConstructionParameters RenderFacesDependency(RenderFacesGroup, bSaved, bAllowCircularDependency);
+		FManagedArrayCollection::FConstructionParameters RenderVerticesDependency(RenderVerticesGroup, bSaved, bAllowCircularDependency);
+		FManagedArrayCollection::FConstructionParameters SimFacesDependency(SimFacesGroup, bSaved, bAllowCircularDependency);
+		FManagedArrayCollection::FConstructionParameters SimVertices2DDependency(SimVertices2DGroup, bSaved, bAllowCircularDependency);
+		FManagedArrayCollection::FConstructionParameters SimVertices3DDependency(SimVertices3DGroup, bSaved, bAllowCircularDependency);  // Any attribute with this dependency must handle welding and splitting in FCollectionClothSeamFacade
 
 		// LODs Group
-		PatternStart = &ManagedArrayCollection->AddAttribute<int32>(PatternStartAttribute, LodsGroup, PatternsDependency);
-		PatternEnd = &ManagedArrayCollection->AddAttribute<int32>(PatternEndAttribute, LodsGroup, PatternsDependency);
-		SeamStart = &ManagedArrayCollection->AddAttribute<int32>(SeamStartAttribute, LodsGroup, SeamsDependency);
-		SeamEnd = &ManagedArrayCollection->AddAttribute<int32>(SeamEndAttribute, LodsGroup, SeamsDependency);
-		TetherBatchStart = &ManagedArrayCollection->AddAttribute<int32>(TetherBatchStartAttribute, LodsGroup, TetherBatchesDependency);
-		TetherBatchEnd = &ManagedArrayCollection->AddAttribute<int32>(TetherBatchEndAttribute, LodsGroup, TetherBatchesDependency);
-		MaterialStart = &ManagedArrayCollection->AddAttribute<int32>(MaterialStartAttribute, LodsGroup, MaterialsDependency);
-		MaterialEnd = &ManagedArrayCollection->AddAttribute<int32>(MaterialEndAttribute, LodsGroup, MaterialsDependency);
 		PhysicsAssetPathName = &ManagedArrayCollection->AddAttribute<FString>(PhysicsAssetPathNameAttribute, LodsGroup);
 		SkeletonAssetPathName = &ManagedArrayCollection->AddAttribute<FString>(SkeletonAssetPathNameAttribute, LodsGroup);
 
-		// Materials Group
-		RenderMaterialPathName = &ManagedArrayCollection->AddAttribute<FString>(RenderMaterialPathNameAttribute, MaterialsGroup);
-
-		// Tether Batches Group
-		TetherStart = &ManagedArrayCollection->AddAttribute<int32>(TetherStartAttribute, TetherBatchesGroup, TethersDependency);
-		TetherEnd = &ManagedArrayCollection->AddAttribute<int32>(TetherEndAttribute, TetherBatchesGroup, TethersDependency);
-
-		// Tethers Group
-		TetherKinematicIndex = &ManagedArrayCollection->AddAttribute<int32>(TetherKinematicIndexAttribute, TethersGroup, SimVerticesDependency);
-		TetherDynamicIndex = &ManagedArrayCollection->AddAttribute<int32>(TetherDynamicIndexAttribute, TethersGroup, SimVerticesDependency);
-		TetherReferenceLength = &ManagedArrayCollection->AddAttribute<float>(TetherReferenceLengthAttribute, TethersGroup);
-
 		// Seams Group
-		SeamPatterns = &ManagedArrayCollection->AddAttribute<FIntVector2>(SeamPatternsAttribute, SeamsGroup, SimVerticesDependency);
-		SeamStitches = &ManagedArrayCollection->AddAttribute<TArray<FIntVector2>>(SeamStitchesAttribute, SeamsGroup, SimVerticesDependency);
-
-		// Patterns Group
-		SimVerticesStart = &ManagedArrayCollection->AddAttribute<int32>(SimVerticesStartAttribute, PatternsGroup, SimVerticesDependency);
-		SimVerticesEnd = &ManagedArrayCollection->AddAttribute<int32>(SimVerticesEndAttribute, PatternsGroup, SimVerticesDependency);
-		SimFacesStart = &ManagedArrayCollection->AddAttribute<int32>(SimFacesStartAttribute, PatternsGroup, SimFacesDependency);
-		SimFacesEnd = &ManagedArrayCollection->AddAttribute<int32>(SimFacesEndAttribute, PatternsGroup, SimFacesDependency);
-		RenderVerticesStart = &ManagedArrayCollection->AddAttribute<int32>(RenderVerticesStartAttribute, PatternsGroup, RenderVerticesDependency);
-		RenderVerticesEnd = &ManagedArrayCollection->AddAttribute<int32>(RenderVerticesEndAttribute, PatternsGroup, RenderVerticesDependency);
-		RenderFacesStart = &ManagedArrayCollection->AddAttribute<int32>(RenderFacesStartAttribute, PatternsGroup, RenderFacesDependency);
-		RenderFacesEnd = &ManagedArrayCollection->AddAttribute<int32>(RenderFacesEndAttribute, PatternsGroup, RenderFacesDependency);
-		StatusFlags = &ManagedArrayCollection->AddAttribute<int32>(StatusFlagsAttribute, PatternsGroup);
+		SeamStitchStart = &ManagedArrayCollection->AddAttribute<int32>(SeamStitchStartAttribute, SeamsGroup, SeamStitchesDependency);
+		SeamStitchEnd = &ManagedArrayCollection->AddAttribute<int32>(SeamStitchEndAttribute, SeamsGroup, SeamStitchesDependency);
+
+		// Seam Stitches Group
+		SeamStitch2DEndIndices = &ManagedArrayCollection->AddAttribute<FIntVector2>(SeamStitch2DEndIndicesAttribute, SeamStitchesGroup, SimVertices2DDependency);
+		SeamStitch3DIndex = &ManagedArrayCollection->AddAttribute<int32>(SeamStitch3DIndexAttribute, SeamStitchesGroup, SimVertices3DDependency);
+
+		// Sim Patterns Group
+		SimVertices2DStart = &ManagedArrayCollection->AddAttribute<int32>(SimVertices2DStartAttribute, SimPatternsGroup, SimVertices2DDependency);
+		SimVertices2DEnd = &ManagedArrayCollection->AddAttribute<int32>(SimVertices2DEndAttribute, SimPatternsGroup, SimVertices2DDependency);
+		SimFacesStart = &ManagedArrayCollection->AddAttribute<int32>(SimFacesStartAttribute, SimPatternsGroup, SimFacesDependency);
+		SimFacesEnd = &ManagedArrayCollection->AddAttribute<int32>(SimFacesEndAttribute, SimPatternsGroup, SimFacesDependency);
+
+		// Render Patterns Group
+		RenderVerticesStart = &ManagedArrayCollection->AddAttribute<int32>(RenderVerticesStartAttribute, RenderPatternsGroup, RenderVerticesDependency);
+		RenderVerticesEnd = &ManagedArrayCollection->AddAttribute<int32>(RenderVerticesEndAttribute, RenderPatternsGroup, RenderVerticesDependency);
+		RenderFacesStart = &ManagedArrayCollection->AddAttribute<int32>(RenderFacesStartAttribute, RenderPatternsGroup, RenderFacesDependency);
+		RenderFacesEnd = &ManagedArrayCollection->AddAttribute<int32>(RenderFacesEndAttribute, RenderPatternsGroup, RenderFacesDependency);
+		RenderMaterialPathName = &ManagedArrayCollection->AddAttribute<FString>(RenderMaterialPathNameAttribute, RenderPatternsGroup);
 
 		// Sim Faces Group
-		SimIndices = &ManagedArrayCollection->AddAttribute<FIntVector3>(SimIndicesAttribute, SimFacesGroup, SimVerticesDependency);
-
-		// Sim Vertices Group
-		SimPosition = &ManagedArrayCollection->AddAttribute<FVector2f>(SimPositionAttribute, SimVerticesGroup);
-		SimRestPosition = &ManagedArrayCollection->AddAttribute<FVector3f>(SimRestPositionAttribute, SimVerticesGroup);
-		SimRestNormal = &ManagedArrayCollection->AddAttribute<FVector3f>(SimRestNormalAttribute, SimVerticesGroup);
-		SimNumBoneInfluences = &ManagedArrayCollection->AddAttribute<int32>(SimNumBoneInfluencesAttribute, SimVerticesGroup);
-		SimBoneIndices = &ManagedArrayCollection->AddAttribute<TArray<int32>>(SimBoneIndicesAttribute, SimVerticesGroup);
-		SimBoneWeights = &ManagedArrayCollection->AddAttribute<TArray<float>>(SimBoneWeightsAttribute, SimVerticesGroup);
+		SimIndices2D = &ManagedArrayCollection->AddAttribute<FIntVector3>(SimIndices2DAttribute, SimFacesGroup, SimVertices2DDependency);
+		SimIndices3D = &ManagedArrayCollection->AddAttribute<FIntVector3>(SimIndices3DAttribute, SimFacesGroup, SimVertices3DDependency);
+
+		// Sim Vertices 2D Group
+		SimPosition2D = &ManagedArrayCollection->AddAttribute<FVector2f>(SimPosition2DAttribute, SimVertices2DGroup);
+		SimVertex3DLookup = &ManagedArrayCollection->AddAttribute<int32>(SimVertex3DLookupAttribute, SimVertices2DGroup, SimVertices3DDependency);
+
+		// Sim Vertices 3D Group
+		SimPosition3D = &ManagedArrayCollection->AddAttribute<FVector3f>(SimPosition3DAttribute, SimVertices3DGroup);
+		SimNormal = &ManagedArrayCollection->AddAttribute<FVector3f>(SimNormalAttribute, SimVertices3DGroup);
+		SimBoneIndices = &ManagedArrayCollection->AddAttribute<TArray<int32>>(SimBoneIndicesAttribute, SimVertices3DGroup);
+		SimBoneWeights = &ManagedArrayCollection->AddAttribute<TArray<float>>(SimBoneWeightsAttribute, SimVertices3DGroup);
+		TetherKinematicIndex = &ManagedArrayCollection->AddAttribute<TArray<int32>>(TetherKinematicIndexAttribute, SimVertices3DGroup, SimVertices3DDependency);
+		TetherReferenceLength = &ManagedArrayCollection->AddAttribute<TArray<float>>(TetherReferenceLengthAttribute, SimVertices3DGroup);
+		SimVertex2DLookup = &ManagedArrayCollection->AddAttribute<TArray<int32>>(SimVertex2DLookupAttribute, SimVertices3DGroup, SimVertices2DDependency);
+		SeamStitchLookup = &ManagedArrayCollection->AddAttribute<TArray<int32>>(SeamStitchLookupAttribute, SimVertices3DGroup, SeamStitchesDependency);
 
 		// Render Faces Group
 		RenderIndices = &ManagedArrayCollection->AddAttribute<FIntVector3>(RenderIndicesAttribute, RenderFacesGroup, RenderVerticesDependency);
-		RenderMaterialIndex = &ManagedArrayCollection->AddAttribute<int32>(RenderMaterialIndexAttribute, RenderFacesGroup);
 
 		// Render Vertices Group
 		RenderPosition = &ManagedArrayCollection->AddAttribute<FVector3f>(RenderPositionAttribute, RenderVerticesGroup);
@@ -404,7 +360,6 @@ namespace UE::Chaos::ClothAsset
 		RenderTangentV = &ManagedArrayCollection->AddAttribute<FVector3f>(RenderTangentVAttribute, RenderVerticesGroup);
 		RenderUVs = &ManagedArrayCollection->AddAttribute<TArray<FVector2f>>(RenderUVsAttribute, RenderVerticesGroup);
 		RenderColor = &ManagedArrayCollection->AddAttribute<FLinearColor>(RenderColorAttribute, RenderVerticesGroup);
-		RenderNumBoneInfluences = &ManagedArrayCollection->AddAttribute<int32>(RenderNumBoneInfluencesAttribute, RenderVerticesGroup);
 		RenderBoneIndices = &ManagedArrayCollection->AddAttribute<TArray<int32>>(RenderBoneIndicesAttribute, RenderVerticesGroup);
 		RenderBoneWeights = &ManagedArrayCollection->AddAttribute<TArray<float>>(RenderBoneWeightsAttribute, RenderVerticesGroup);
 	}
@@ -512,7 +467,12 @@ namespace UE::Chaos::ClothAsset
 		return Start;
 	}
 
-	int32 FClothCollection::GetElementsOffset(const TManagedArray<int32>* StartArray, int32 BaseElementIndex, int32 ElementIndex) const
+	void FClothCollection::RemoveElements(const FName& Group, const TArray<int32>& SortedDeletionList)
+	{
+		ManagedArrayCollection->RemoveElements(Group, SortedDeletionList);
+	}
+
+	int32 FClothCollection::GetElementsOffset(const TManagedArray<int32>* StartArray, int32 BaseElementIndex, int32 ElementIndex)
 	{
 		while ((*StartArray)[BaseElementIndex] == INDEX_NONE && BaseElementIndex < ElementIndex)
 		{
@@ -521,12 +481,27 @@ namespace UE::Chaos::ClothAsset
 		return (*StartArray)[ElementIndex] - (*StartArray)[BaseElementIndex];
 	}
 
+	int32 FClothCollection::GetArrayIndexForContainedElement(
+		const TManagedArray<int32>* StartArray,
+		const TManagedArray<int32>* EndArray,
+		int32 ElementIndex)
+	{
+		for (int32 ArrayIndex = 0; ArrayIndex < StartArray->Num(); ++ArrayIndex)
+		{
+			if (ElementIndex >= (*StartArray)[ArrayIndex] && ElementIndex <= (*EndArray)[ArrayIndex])
+			{
+				return ArrayIndex;
+			}
+		}
+		return INDEX_NONE;
+	}
+
 	int32 FClothCollection::GetNumSubElements(
 		const TManagedArray<int32>* StartArray,
 		const TManagedArray<int32>* EndArray,
 		const TManagedArray<int32>* StartSubArray,
 		const TManagedArray<int32>* EndSubArray,
-		int32 ArrayIndex) const
+		int32 ArrayIndex)
 	{
 		const TTuple<int32, int32> StartEnd = GetSubElementsStartEnd(StartArray, EndArray, StartSubArray, EndSubArray, ArrayIndex);
 		const int32 Start = StartEnd.Get<0>();
@@ -545,7 +520,7 @@ namespace UE::Chaos::ClothAsset
 		const TManagedArray<int32>* EndArray,
 		const TManagedArray<int32>* StartSubArray,
 		const TManagedArray<int32>* EndSubArray,
-		int32 ArrayIndex) const
+		int32 ArrayIndex)
 	{
 		int32 Start = INDEX_NONE;  // Find Start and End indices for the entire LOD minding empty patterns on the way
 		int32 End = INDEX_NONE;
@@ -577,9 +552,9 @@ namespace UE::Chaos::ClothAsset
 		}
 		return TTuple<int32, int32>(Start, End);
 	}
-	template CHAOSCLOTHASSET_API TTuple<int32, int32> FClothCollection::GetSubElementsStartEnd<true, false>(const TManagedArray<int32>* StartArray, const TManagedArray<int32>* EndArray, const TManagedArray<int32>* StartSubArray, const TManagedArray<int32>* EndSubArray, int32 ArrayIndex) const;
-	template CHAOSCLOTHASSET_API TTuple<int32, int32> FClothCollection::GetSubElementsStartEnd<false, true>(const TManagedArray<int32>* StartArray, const TManagedArray<int32>* EndArray, const TManagedArray<int32>* StartSubArray, const TManagedArray<int32>* EndSubArray, int32 ArrayIndex) const;
-	template CHAOSCLOTHASSET_API TTuple<int32, int32> FClothCollection::GetSubElementsStartEnd<true, true>(const TManagedArray<int32>* StartArray, const TManagedArray<int32>* EndArray, const TManagedArray<int32>* StartSubArray, const TManagedArray<int32>* EndSubArray, int32 ArrayIndex) const;
+	template CHAOSCLOTHASSET_API TTuple<int32, int32> FClothCollection::GetSubElementsStartEnd<true, false>(const TManagedArray<int32>* StartArray, const TManagedArray<int32>* EndArray, const TManagedArray<int32>* StartSubArray, const TManagedArray<int32>* EndSubArray, int32 ArrayIndex);
+	template CHAOSCLOTHASSET_API TTuple<int32, int32> FClothCollection::GetSubElementsStartEnd<false, true>(const TManagedArray<int32>* StartArray, const TManagedArray<int32>* EndArray, const TManagedArray<int32>* StartSubArray, const TManagedArray<int32>* EndSubArray, int32 ArrayIndex);
+	template CHAOSCLOTHASSET_API TTuple<int32, int32> FClothCollection::GetSubElementsStartEnd<true, true>(const TManagedArray<int32>* StartArray, const TManagedArray<int32>* EndArray, const TManagedArray<int32>* StartSubArray, const TManagedArray<int32>* EndSubArray, int32 ArrayIndex);
 
 	template<typename T, typename TEnableIf<TIsUserAttributeType<T>::Value, int>::type>
 	TArray<FName> FClothCollection::GetUserDefinedAttributeNames(const FName& GroupName) const
@@ -655,4 +630,17 @@ namespace UE::Chaos::ClothAsset
 	template CHAOSCLOTHASSET_API TManagedArray<int32>* FClothCollection::GetUserDefinedAttribute<int32>(const FName& Name, const FName& GroupName);
 	template CHAOSCLOTHASSET_API TManagedArray<float>* FClothCollection::GetUserDefinedAttribute<float>(const FName& Name, const FName& GroupName);
 	template CHAOSCLOTHASSET_API TManagedArray<FVector3f>* FClothCollection::GetUserDefinedAttribute<FVector3f>(const FName& Name, const FName& GroupName);
+
+	void FClothCollection::CopyArrayViewDataAndApplyOffset(const TArrayView<TArray<int32>>& To, const TConstArrayView<TArray<int32>>& From, const int32 Offset)
+	{
+		check(To.Num() == From.Num());
+		for (int32 Index = 0; Index < To.Num(); ++Index)
+		{
+			To[Index] = From[Index];
+			for (int32& Value : To[Index])
+			{
+				Value += Offset;
+			}
+		}
+	}
 } // End namespace UE::Chaos::ClothAsset
diff --git a/Engine/Plugins/Experimental/ChaosClothAsset/Source/ChaosClothAsset/Private/ChaosClothAsset/ClothGeometryTools.cpp b/Engine/Plugins/Experimental/ChaosClothAsset/Source/ChaosClothAsset/Private/ChaosClothAsset/ClothGeometryTools.cpp
--- a/Engine/Plugins/Experimental/ChaosClothAsset/Source/ChaosClothAsset/Private/ChaosClothAsset/ClothGeometryTools.cpp
+++ b/Engine/Plugins/Experimental/ChaosClothAsset/Source/ChaosClothAsset/Private/ChaosClothAsset/ClothGeometryTools.cpp
@@ -2,215 +2,563 @@
 
 #include "ChaosClothAsset/ClothGeometryTools.h"
 #include "ChaosClothAsset/CollectionClothFacade.h"
+#include "Containers/Queue.h"
+#include "DynamicMesh/DynamicMesh3.h"
+#include "DynamicMesh/DynamicMeshAttributeSet.h"
+#include "DynamicMesh/NonManifoldMappingSupport.h"
 #include "Math/Vector.h"
+#include "Util/IndexUtil.h"
 
 namespace UE::Chaos::ClothAsset
 {
-	// TODO: Move these functions to the cloth facade?
-	bool FClothGeometryTools::HasSimMesh(const TSharedPtr<const FManagedArrayCollection>& ClothCollection)
+	namespace Private::SimMeshBuilder
 	{
-		bool bHasSimMesh = false;
+		// Triangle islands to become patterns, although in this case all the seams are internal (same pattern)
+		struct FIsland
+		{
+			TArray<int32> Indices;  // 3x number of triangles
+			TArray<FVector2f> Positions2D;
+			TArray<FVector3f> Positions3D;  // Same size as Positions
+			TArray<int32> PositionToSourceIndex; // Same size as Positions. Index in the original welded position array
+		};
 
-		FCollectionClothConstFacade ClothFacade(ClothCollection);
+		enum class EIntersectCirclesResult
+		{
+			SingleIntersect,
+			DoubleIntersect,
+			Coincident,
+			Separate,
+			Contained
+		};
+
+		static EIntersectCirclesResult IntersectCircles(const FVector2f& C0, float R0, const FVector2f& C1, float R1, FVector2f& OutI0, FVector2f& OutI1)
+		{
+			const FVector2f C0C1 = C0 - C1;
+			const float D = C0C1.Length();
+			if (D < SMALL_NUMBER)
+			{
+				return EIntersectCirclesResult::Coincident;
+			}
+			else if (D > R0 + R1)
+			{
+				return EIntersectCirclesResult::Separate;
+			}
+			else if (D < FMath::Abs(R0 - R1))
+			{
+				return EIntersectCirclesResult::Contained;
+			}
+			const float SquareR0 = R0 * R0;
+			const float SquareR1 = R1 * R1;
+			const float SquareD = D * D;
+			const float A = (SquareD - SquareR1 + SquareR0) / (2.f * D);
+
+			OutI0 = OutI1 = C0 + A * (C1 - C0) / D;
+
+			if (FMath::Abs(A - R0) < SMALL_NUMBER)
+			{
+				return EIntersectCirclesResult::SingleIntersect;
+			}
 
-		for (int32 LodIndex = 0; LodIndex < ClothFacade.GetNumLods(); ++LodIndex)
+			const float SquareA = A * A;
+			const float H = FMath::Sqrt(SquareR0 - SquareA);
+
+			const FVector2f N(C0C1.Y, -C0C1.X);
+
+			OutI0 += N * H / D;
+			OutI1 -= N * H / D;
+
+			return EIntersectCirclesResult::DoubleIntersect;
+		}
+		;
+		static FIntVector2 MakeSortedIntVector2(int32 Index0, int32 Index1)
 		{
-			FCollectionClothLodConstFacade ClothLodFacade = ClothFacade.GetLod(LodIndex);
+			return Index0 < Index1 ? FIntVector2(Index0, Index1) : FIntVector2(Index1, Index0);
+		}
+
+		static void UnwrapDynamicMesh(const UE::Geometry::FDynamicMesh3& DynamicMesh, TArray<FIsland>& OutIslands)
+		{
+			using namespace UE::Geometry;
+
+			OutIslands.Reset();
+			constexpr float SquaredWeldingDistance = FMath::Square(0.01f);  // 0.1 mm
 
-			for (int32 PatternIndex = ClothLodFacade.GetNumPatterns() - 1; PatternIndex >= 0 ; --PatternIndex)  // Use a reverse order to avoid having to move previous elements
+			// Build pattern islands. 
+			const int32 NumTriangles = DynamicMesh.TriangleCount();
+			TSet<int32> VisitedTriangles;
+			VisitedTriangles.Reserve(NumTriangles);
+
+
+			for (int32 SeedTriangle : DynamicMesh.TriangleIndicesItr())
 			{
-				FCollectionClothPatternConstFacade ClothPatternFacade = ClothLodFacade.GetPattern(PatternIndex);
-				if (ClothPatternFacade.GetNumSimVertices() && ClothPatternFacade.GetNumSimFaces())
+				if (VisitedTriangles.Contains(SeedTriangle))
+				{
+					continue;
+				}
+				const FIndex3i TriangleIndices = DynamicMesh.GetTriangle(SeedTriangle);
+
+				const int32 SeedIndex0 = TriangleIndices[0];
+				const int32 SeedIndex1 = TriangleIndices[1];
+
+				const FVector3f Position0(DynamicMesh.GetVertex(SeedIndex0));
+				const FVector3f Position1(DynamicMesh.GetVertex(SeedIndex1));
+				const float Position01DistSq = FVector3f::DistSquared(Position0, Position1);
+
+				if (Position01DistSq <= SquaredWeldingDistance)
 				{
-					bHasSimMesh = true;
-					break;
+					continue;  // A degenerated triangle edge is not a good start
 				}
+
+				// Setup first visitor from seed, and add the first two points
+				FIsland& Island = OutIslands.AddDefaulted_GetRef();
+
+				Island.Positions3D.Add(Position0);
+				Island.Positions3D.Add(Position1);
+				Island.PositionToSourceIndex.Add(SeedIndex0);
+				Island.PositionToSourceIndex.Add(SeedIndex1);
+
+				const int32 SeedIndex2D0 = Island.Positions2D.Add(FVector2f::ZeroVector);
+				const int32 SeedIndex2D1 = Island.Positions2D.Add(FVector2f(FMath::Sqrt(Position01DistSq), 0.f));
+
+				struct FVisitor
+				{
+					int32 Triangle;
+					FIndex2i OldEdge;
+					FIndex2i NewEdge;
+					int32 CrossEdgePoint;  // Keep the opposite point to orientate degenerate cases
+				} Visitor =
+				{
+					SeedTriangle,
+					FIndex2i(SeedIndex0, SeedIndex1),
+					FIndex2i(SeedIndex2D0, SeedIndex2D1),
+					INDEX_NONE
+				};
+
+				VisitedTriangles.Add(SeedTriangle);
+
+				TQueue<FVisitor> Visitors;
+				do
+				{
+					const int32 Triangle = Visitor.Triangle;
+					const int32 CrossEdgePoint = Visitor.CrossEdgePoint;
+					const int32 OldIndex0 = Visitor.OldEdge.A;
+					const int32 OldIndex1 = Visitor.OldEdge.B;
+					const int32 NewIndex0 = Visitor.NewEdge.A;
+					const int32 NewIndex1 = Visitor.NewEdge.B;
+
+					// Find opposite index from this triangle edge
+
+					const int32 OldIndex2 = IndexUtil::FindTriOtherVtxUnsafe(OldIndex0, OldIndex1, DynamicMesh.GetTriangle(Triangle));
+
+					// Find the 2D intersection of the two connecting adjacent edges using the 3D reference length
+					const FVector3f P0(DynamicMesh.GetVertexRef(OldIndex0));
+					const FVector3f P1(DynamicMesh.GetVertexRef(OldIndex1));
+					const FVector3f P2(DynamicMesh.GetVertexRef(OldIndex2));
+
+					const float R0 = FVector3f::Dist(P0, P2);
+					const float R1 = FVector3f::Dist(P1, P2);
+					const FVector2f& C0 = Island.Positions2D[NewIndex0];
+					const FVector2f& C1 = Island.Positions2D[NewIndex1];
+
+					FVector2f I0, I1;
+					const EIntersectCirclesResult IntersectCirclesResult = IntersectCircles(C0, R0, C1, R1, I0, I1);
+
+					FVector2f C2;
+					switch (IntersectCirclesResult)
+					{
+					case EIntersectCirclesResult::SingleIntersect:
+						C2 = I0;  // Degenerated C2 is on (C0C1)
+						break;
+					case EIntersectCirclesResult::DoubleIntersect:
+						C2 = (FVector2f::CrossProduct(C0 - C1, C0 - I0) > 0) ? I0 : I1;  // Keep correct winding order
+						break;
+					case EIntersectCirclesResult::Coincident:
+						check(CrossEdgePoint != INDEX_NONE); // We can't start on a degenerated triangle
+						C2 = C0 - (Island.Positions2D[CrossEdgePoint] - C0).GetSafeNormal() * R0;  // Degenerated C0 == C1, choose C2 on the opposite of the visitor opposite point
+						break;
+					case EIntersectCirclesResult::Separate: [[fallthrough]];
+					case EIntersectCirclesResult::Contained:
+						C2 = C0 - (C1 - C0).GetSafeNormal() * R0;  // Degenerated + some tolerance, C2 is on (C0C1)
+						break;
+					}
+
+					// Add the new position found for the opposite point
+					int32 NewIndex2 = INDEX_NONE;
+					for (int32 UsedIndex = 0; UsedIndex < Island.Positions2D.Num(); ++UsedIndex)
+					{
+						if (Island.PositionToSourceIndex[UsedIndex] == OldIndex2 &&
+							FVector2f::DistSquared(Island.Positions2D[UsedIndex], C2) <= SquaredWeldingDistance)
+						{
+							NewIndex2 = UsedIndex;  // Both Rest and 2D positions match, reuse this index
+							break;
+						}
+					}
+
+					if (NewIndex2 == INDEX_NONE)
+					{
+						NewIndex2 = Island.Positions2D.Add(C2);
+						Island.Positions3D.Add(P2);
+						Island.PositionToSourceIndex.Add(OldIndex2);
+					}
+
+					// Add triangle to list of indices, unless it is degenerated to a segment
+					if (NewIndex0 != NewIndex1 && NewIndex1 != NewIndex2 && NewIndex2 != NewIndex0)
+					{
+						Island.Indices.Add(NewIndex0);
+						Island.Indices.Add(NewIndex1);
+						Island.Indices.Add(NewIndex2);
+					}
+
+					// Add neighbor triangles to the queue
+					const FIndex2i OldEdgeList[3] =
+					{
+						FIndex2i(OldIndex1, OldIndex0),  // Reversed as to keep the correct winding order
+						FIndex2i(OldIndex2, OldIndex1),
+						FIndex2i(OldIndex0, OldIndex2)
+					};
+					const FIndex3i NewEdgeList[3] =
+					{
+						FIndex3i(NewIndex1, NewIndex0, NewIndex2),  // Adds opposite point index
+						FIndex3i(NewIndex2, NewIndex1, NewIndex0),
+						FIndex3i(NewIndex0, NewIndex2, NewIndex1)
+					};
+					for (int32 Edge = 0; Edge < 3; ++Edge)
+					{
+						const int32 EdgeIndex0 = OldEdgeList[Edge].A;
+						const int32 EdgeIndex1 = OldEdgeList[Edge].B;
+
+						const FIndex2i EdgeT = DynamicMesh.GetEdgeT(DynamicMesh.FindEdgeFromTri(EdgeIndex0, EdgeIndex1, Triangle));
+						const int32 NeighborTriangle = EdgeT.OtherElement(Triangle);
+						if (NeighborTriangle != IndexConstants::InvalidID)
+						{
+							if (!VisitedTriangles.Contains(NeighborTriangle))
+							{
+								// Mark neighboring triangle as visited
+								VisitedTriangles.Add(NeighborTriangle);
+
+								// Enqueue next triangle
+								Visitors.Enqueue(FVisitor
+									{
+										NeighborTriangle,
+										OldEdgeList[Edge],
+										FIndex2i(NewEdgeList[Edge].A, NewEdgeList[Edge].B),
+										NewEdgeList[Edge].C,  // Pass the cross edge 2D opposite point to help define orientation of any degenerated triangles
+									});
+							}
+
+						}
+					}
+				} while (Visitors.Dequeue(Visitor));
 			}
 		}
-		return bHasSimMesh;
-	}
 
-	bool FClothGeometryTools::HasRenderMesh(const TSharedPtr<const FManagedArrayCollection>& ClothCollection)
-	{
-		bool bHasRenderMesh = false;
+		static void BuildIslandsFromDynamicMeshUVs(const UE::Geometry::FDynamicMeshUVOverlay& UVOverlay, const FVector2f& UVScale, TArray<FIsland>& OutIslands)
+		{
+			using namespace UE::Geometry;
 
-		FCollectionClothConstFacade ClothFacade(ClothCollection);
+			const FDynamicMesh3* DynamicMesh = UVOverlay.GetParentMesh();
+			check(DynamicMesh);
 
-		for (int32 LodIndex = 0; LodIndex < ClothFacade.GetNumLods(); ++LodIndex)
-		{
-			FCollectionClothLodConstFacade ClothLodFacade = ClothFacade.GetLod(LodIndex);
+			OutIslands.Reset();
 
-			for (int32 PatternIndex = ClothLodFacade.GetNumPatterns() - 1; PatternIndex >= 0 ; --PatternIndex)  // Use a reverse order to avoid having to move previous elements
+			// Build pattern islands. 
+			const int32 NumTriangles = DynamicMesh->TriangleCount();
+			TSet<int32> VisitedTriangles;
+			VisitedTriangles.Reserve(NumTriangles);
+
+			// This is reused for each island, but only allocate once.
+			TArray<int32> SourceElementIndexToNewIndex;
+
+			for (int32 SeedTriangle : DynamicMesh->TriangleIndicesItr())
 			{
-				FCollectionClothPatternConstFacade ClothPatternFacade = ClothLodFacade.GetPattern(PatternIndex);
-				if (ClothPatternFacade.GetNumRenderVertices() && ClothPatternFacade.GetNumRenderFaces())
+				if (VisitedTriangles.Contains(SeedTriangle))
 				{
-					bHasRenderMesh = true;
-					break;
+					continue;
 				}
+
+				// Setup first visitor from seed
+				FIsland& Island = OutIslands.AddDefaulted_GetRef();
+				SourceElementIndexToNewIndex.Init(INDEX_NONE, UVOverlay.MaxElementID());
+
+				struct FVisitor
+				{
+					int32 Triangle;
+				} Visitor =
+				{
+					SeedTriangle
+				};
+
+				VisitedTriangles.Add(SeedTriangle);
+
+				TQueue<FVisitor> Visitors;
+				do
+				{
+					const int32 Triangle = Visitor.Triangle;
+					const FIndex3i TriangleIndices = DynamicMesh->GetTriangle(Triangle);
+					const FIndex3i TriangleUVElements = UVOverlay.GetTriangle(Triangle);
+
+					auto GetOrAddNewIndex = [&UVOverlay, &Island, &SourceElementIndexToNewIndex, &DynamicMesh, &UVScale](int32 ElementId, int32 VertexId)
+					{
+						int32& NewIndex = SourceElementIndexToNewIndex[ElementId];
+						if (NewIndex == INDEX_NONE)
+						{
+							NewIndex = Island.Positions3D.Add(FVector3f(DynamicMesh->GetVertexRef(VertexId)));
+							Island.Positions2D.Add((FVector2f(1.f) - UVOverlay.GetElement(ElementId)) * UVScale);  // The static mesh import uses 1 - UV for some reason
+							Island.PositionToSourceIndex.Add(VertexId);
+						}
+						return NewIndex;
+					};
+
+					const int32 NewIndex0 = GetOrAddNewIndex(TriangleUVElements[0], TriangleIndices[0]);
+					const int32 NewIndex1 = GetOrAddNewIndex(TriangleUVElements[1], TriangleIndices[1]);
+					const int32 NewIndex2 = GetOrAddNewIndex(TriangleUVElements[2], TriangleIndices[2]);
+					Island.Indices.Add(NewIndex0);
+					Island.Indices.Add(NewIndex1);
+					Island.Indices.Add(NewIndex2);
+
+					TArray<int32> NeighborTriangles;
+					for (int32 LocalVertexId = 0; LocalVertexId < 3; ++LocalVertexId)
+					{
+						NeighborTriangles.Reset();
+						UVOverlay.GetElementTriangles(TriangleUVElements[LocalVertexId], NeighborTriangles);
+						for (const int32 NeighborTriangle : NeighborTriangles)
+						{
+							if (!VisitedTriangles.Contains(NeighborTriangle))
+							{
+								// Mark neighboring triangle as visited
+								VisitedTriangles.Add(NeighborTriangle);
+
+								// Enqueue next triangle
+								Visitors.Enqueue(FVisitor({ NeighborTriangle }));
+							}
+						}
+					}
+				} while (Visitors.Dequeue(Visitor));
 			}
 		}
-		return bHasRenderMesh;
-	}
 
-	void FClothGeometryTools::SetSkeletonAssetPathName(const TSharedPtr<FManagedArrayCollection>& ClothCollection, const FString& SkeletonAssetPathName)
-	{
-		FCollectionClothFacade ClothFacade(ClothCollection);
+		struct FSeam
+		{
+			TSet<FIntVector2> Stitches;
+			FIntVector2 Patterns;
+		};
 
-		for (int32 LodIndex = 0; LodIndex < ClothFacade.GetNumLods(); ++LodIndex)
+		// Stitch together any vertices that were split, either via DynamicMesh NonManifoldMapping or UV Unwrap
+		static void BuildSeams(const TArray<FIsland>& Islands, const UE::Geometry::FDynamicMesh3& DynamicMesh, TArray<FSeam>& OutSeams)
 		{
-			FCollectionClothLodFacade ClothLodFacade = ClothFacade.GetLod(LodIndex);
-			ClothLodFacade.SetSkeletonAssetPathName(SkeletonAssetPathName);
-		}
-	}
+			OutSeams.Reset();
 
-	void FClothGeometryTools::SetPhysicsAssetPathName(const TSharedPtr<FManagedArrayCollection>& ClothCollection, const FString& PhysicsAssetPathName)
-	{
-		FCollectionClothFacade ClothFacade(ClothCollection);
+			const UE::Geometry::FNonManifoldMappingSupport NonManifoldMapping(DynamicMesh);
+
+			TArray<TMap<int32, TArray<int32>>> IslandSourceIndexToPositions;
+			IslandSourceIndexToPositions.SetNum(Islands.Num());
 
-		for (int32 LodIndex = 0; LodIndex < ClothFacade.GetNumLods(); ++LodIndex)
+			for (int32 IslandIndex = 0; IslandIndex < Islands.Num(); ++IslandIndex)
+			{
+				const FIsland& Island = Islands[IslandIndex];
+				TMap<int32, TArray<int32>>& SourceIndexToPositions = IslandSourceIndexToPositions[IslandIndex];
+
+				// Build reverse lookup to PositionToSourceIndex
+				SourceIndexToPositions.Reserve(Island.PositionToSourceIndex.Num());
+				for (int32 PositionIndex = 0; PositionIndex < Island.PositionToSourceIndex.Num(); ++PositionIndex)
+				{
+					const int32 SourceIndex = NonManifoldMapping.GetOriginalNonManifoldVertexID(Island.PositionToSourceIndex[PositionIndex]);
+					SourceIndexToPositions.FindOrAdd(SourceIndex).Add(PositionIndex);
+				}
+
+				// Find all internal seams
+				FSeam InternalSeam;
+				InternalSeam.Patterns = FIntVector2(IslandIndex);
+				for (const TPair<int32, TArray<int32>>& Source : SourceIndexToPositions)
+				{
+					for (int32 FirstSourceArrayIdx = 0; FirstSourceArrayIdx < Source.Value.Num() - 1; ++FirstSourceArrayIdx)
+					{
+						for (int32 SecondSourceArrayIdx = FirstSourceArrayIdx + 1; SecondSourceArrayIdx < Source.Value.Num(); ++SecondSourceArrayIdx)
+						{
+							InternalSeam.Stitches.Emplace(MakeSortedIntVector2(Source.Value[FirstSourceArrayIdx], Source.Value[SecondSourceArrayIdx]));
+						}
+					}
+				}
+				if (InternalSeam.Stitches.Num())
+				{
+					OutSeams.Emplace(MoveTemp(InternalSeam));
+				}
+
+				for (int32 OtherIslandIndex = 0; OtherIslandIndex < IslandIndex; ++OtherIslandIndex)
+				{
+					// Find all seams between the two islands
+					const TMap<int32, TArray<int32>>& OtherSourceIndexToPositions = IslandSourceIndexToPositions[OtherIslandIndex];
+
+					FSeam Seam;
+					Seam.Patterns = FIntVector2(OtherIslandIndex, IslandIndex);
+					for (const TPair<int32, TArray<int32>>& FirstSource : SourceIndexToPositions)
+					{
+						if (const TArray<int32>* OtherSource = OtherSourceIndexToPositions.Find(FirstSource.Key))
+						{
+							for (const int32 FirstSourceVert : FirstSource.Value)
+							{
+								for (const int32 OtherSourceVert : *OtherSource)
+								{
+									Seam.Stitches.Emplace(FIntVector2(OtherSourceVert, FirstSourceVert));
+								}
+							}
+						}
+					}
+					if (Seam.Stitches.Num())
+					{
+						OutSeams.Emplace(MoveTemp(Seam));
+					}
+				}
+			}
+		}
+
+		static void AddSeam(FCollectionClothFacade& Cloth, const FSeam& Seam)
 		{
-			FCollectionClothLodFacade ClothLodFacade = ClothFacade.GetLod(LodIndex);
-			ClothLodFacade.SetPhysicsAssetPathName(PhysicsAssetPathName);
+			const int32 Pattern0Start = Cloth.GetSimPattern(Seam.Patterns[0]).GetSimVertices2DOffset();
+			const int32 Pattern1Start = Cloth.GetSimPattern(Seam.Patterns[1]).GetSimVertices2DOffset();
+
+			FCollectionClothSeamFacade SeamFacade = Cloth.AddGetSeam();
+			TArray<FIntVector2> Stitches;
+			Stitches.Reset(Seam.Stitches.Num());
+			for (const FIntVector2& Stitch : Seam.Stitches)
+			{
+				Stitches.Emplace(Stitch[0] + Pattern0Start, Stitch[1] + Pattern1Start);
+			}
+			SeamFacade.Initialize(Stitches);
 		}
+	}// namespace Private::SimMeshBuilder
+
+	// TODO: Move these functions to the cloth facade?
+	bool FClothGeometryTools::HasSimMesh(const TSharedPtr<const FManagedArrayCollection>& ClothCollection)
+	{
+		FCollectionClothConstFacade ClothFacade(ClothCollection);
+		return ClothFacade.GetNumSimVertices2D() > 0 && ClothFacade.GetNumSimVertices3D() && ClothFacade.GetNumSimFaces() > 0;
+	}
+
+	bool FClothGeometryTools::HasRenderMesh(const TSharedPtr<const FManagedArrayCollection>& ClothCollection)
+	{
+		FCollectionClothConstFacade ClothFacade(ClothCollection);
+		return ClothFacade.GetNumRenderVertices() > 0 && ClothFacade.GetNumRenderFaces() > 0;
 	}
 
 	void FClothGeometryTools::DeleteRenderMesh(const TSharedPtr<FManagedArrayCollection>& ClothCollection)
 	{
 		FCollectionClothFacade ClothFacade(ClothCollection);
+		ClothFacade.SetNumRenderPatterns(0);
+	}
 
-		for (int32 LodIndex = 0; LodIndex < ClothFacade.GetNumLods(); ++LodIndex)
-		{
-			FCollectionClothLodFacade ClothLodFacade = ClothFacade.GetLod(LodIndex);
-
-			for (int32 PatternIndex = ClothLodFacade.GetNumPatterns() - 1; PatternIndex >= 0 ; --PatternIndex)  // Use a reverse order to avoid having to move previous elements
-			{
-				FCollectionClothPatternFacade ClothPatternFacade = ClothLodFacade.GetPattern(PatternIndex);
-				ClothPatternFacade.SetNumRenderVertices(0);
-				ClothPatternFacade.SetNumRenderFaces(0);
-			}
+	void FClothGeometryTools::DeleteSimMesh(const TSharedPtr<FManagedArrayCollection>& ClothCollection)
+	{
+		FCollectionClothFacade ClothFacade(ClothCollection);
+		ClothFacade.SetNumSimPatterns(0);
+		ClothFacade.RemoveAllSimVertices3D();
+		ClothFacade.SetNumSeams(0);
+	}
 
-			ClothLodFacade.SetNumMaterials(0);
+	void FClothGeometryTools::DeleteTethers(const TSharedPtr<FManagedArrayCollection>& ClothCollection)
+	{
+		FCollectionClothFacade ClothFacade(ClothCollection);
+		for (TArray<int32>& KinematicIndex : ClothFacade.GetTetherKinematicIndex())
+		{
+			KinematicIndex.Reset();
+		}
+		for (TArray<float>& ReferenceLength : ClothFacade.GetTetherReferenceLength())
+		{
+			ReferenceLength.Reset();
 		}
 	}
 
-	void FClothGeometryTools::CopySimMeshToRenderMesh(const TSharedPtr<FManagedArrayCollection>& ClothCollection, const FString& RenderMaterialPathName)
+	void FClothGeometryTools::CopySimMeshToRenderMesh(const TSharedPtr<FManagedArrayCollection>& ClothCollection, const FString& RenderMaterialPathName, bool bSingleRenderPattern)
 	{
 		FCollectionClothFacade ClothFacade(ClothCollection);
 
-		for (int32 LodIndex = 0; LodIndex < ClothFacade.GetNumLods(); ++LodIndex)
-		{
-			FCollectionClothLodFacade ClothLodFacade = ClothFacade.GetLod(LodIndex);
+		// Use 2D topology (unwelded mesh)
 
-			// Add the material if it doesn't already exists in this LOD
-			const TConstArrayView<FString> RenderMaterialPathNameArray = ClothLodFacade.GetRenderMaterialPathName();
-			int32 MaterialIndex = RenderMaterialPathNameArray.Find(RenderMaterialPathName);
-			if (MaterialIndex == INDEX_NONE)
-			{
-				MaterialIndex = RenderMaterialPathNameArray.Num();
-				ClothLodFacade.SetNumMaterials(MaterialIndex + 1);
-				ClothLodFacade.GetRenderMaterialPathName()[MaterialIndex] = RenderMaterialPathName;
-			}
+		// Render pattern data
+		const int32 NumRenderPatterns = bSingleRenderPattern ? 1 : ClothFacade.GetNumSimPatterns();
+		ClothFacade.SetNumRenderPatterns(NumRenderPatterns);
+		const int32 TotalNumFaces = ClothFacade.GetNumSimFaces();
+		const int32 TotalNumVertices = ClothFacade.GetNumSimVertices2D();
+		for (int32 RenderPatternIndex = 0; RenderPatternIndex < NumRenderPatterns; ++RenderPatternIndex)
+		{
+			FCollectionClothRenderPatternFacade RenderPattern = ClothFacade.GetRenderPattern(RenderPatternIndex);
+			RenderPattern.SetRenderMaterialPathName(RenderMaterialPathName);
+			RenderPattern.SetNumRenderVertices(bSingleRenderPattern ? TotalNumVertices : ClothFacade.GetSimPattern(RenderPatternIndex).GetNumSimVertices2D());
+			RenderPattern.SetNumRenderFaces(bSingleRenderPattern ? TotalNumFaces : ClothFacade.GetSimPattern(RenderPatternIndex).GetNumSimFaces());
+		}
 
-			// Copy sim data to render data
-			const TConstArrayView<FVector2f> LodSimPosition = ClothLodFacade.GetSimPosition();
-			const TConstArrayView<FVector3f> LodSimRestPosition = ClothLodFacade.GetSimRestPosition();
+		// Calculate UVs scale and zero out tangents
+		FVector2f MinPosition(TNumericLimits<float>::Max());
+		FVector2f MaxPosition(TNumericLimits<float>::Lowest());
 
-			for (int32 PatternIndex = 0; PatternIndex < ClothLodFacade.GetNumPatterns(); ++PatternIndex)
-			{
-				FCollectionClothPatternFacade ClothPatternFacade = ClothLodFacade.GetPattern(PatternIndex);
+		const TConstArrayView<FVector2f> SimPosition2D = ClothFacade.GetSimPosition2D();
+		const TConstArrayView<FVector3f> SimNormal = ClothFacade.GetSimNormal();
+		const TArrayView<FVector3f> RenderTangentU = ClothFacade.GetRenderTangentU();
+		const TArrayView<FVector3f> RenderTangentV = ClothFacade.GetRenderTangentV();
 
-				const int32 NumVertices = ClothPatternFacade.GetNumSimVertices();
-				const int32 NumFaces = ClothPatternFacade.GetNumSimFaces();
+		for (int32 VertexIndex = 0; VertexIndex < TotalNumVertices; ++VertexIndex)
+		{
+			MinPosition = FVector2f::Min(MinPosition, SimPosition2D[VertexIndex]);
+			MaxPosition = FVector2f::Max(MaxPosition, SimPosition2D[VertexIndex]);
 
-				if (!NumVertices || !NumFaces)
-				{
-					ClothPatternFacade.SetNumRenderVertices(0);
-					ClothPatternFacade.SetNumRenderFaces(0);
-					continue;
-				}
+			RenderTangentU[VertexIndex] = FVector3f::ZeroVector;
+			RenderTangentV[VertexIndex] = FVector3f::ZeroVector;
+		}
+		const FVector2f UVScale = MaxPosition - MinPosition;
+		const FVector2f UVInvScale(
+			UVScale.X < SMALL_NUMBER ? 0.f : 1.f / UVScale.X,
+			UVScale.Y < SMALL_NUMBER ? 0.f : 1.f / UVScale.Y);
+
+		// Face group (and calculating render tangents)
+		const TConstArrayView<FVector3f> SimPosition3D = ClothFacade.GetSimPosition3D();
+		const TConstArrayView<int32> SimVertex3DLookup = ClothFacade.GetSimVertex3DLookup();
+		const TConstArrayView<FIntVector3> SimIndices = ClothFacade.GetSimIndices2D();
+		const TArrayView<FIntVector3> RenderIndices = ClothFacade.GetRenderIndices();
+		for (int32 FaceIndex = 0; FaceIndex < TotalNumFaces; ++FaceIndex)
+		{
+			const FIntVector3& Face = SimIndices[FaceIndex];
+			RenderIndices[FaceIndex] = Face;
 
-				ClothPatternFacade.SetNumRenderVertices(NumVertices);
-				ClothPatternFacade.SetNumRenderFaces(NumFaces);
+			const FVector3f Pos01 = SimPosition3D[SimVertex3DLookup[Face[1]]] - SimPosition3D[SimVertex3DLookup[Face[0]]];
+			const FVector3f Pos02 = SimPosition3D[SimVertex3DLookup[Face[2]]] - SimPosition3D[SimVertex3DLookup[Face[0]]];
+			const FVector2f UV01 = SimPosition2D[Face[1]] - SimPosition2D[Face[0]];
+			const FVector2f UV02 = SimPosition2D[Face[2]] - SimPosition2D[Face[0]];
 
-				const TConstArrayView<FVector2f> SimPosition = ClothPatternFacade.GetSimPosition();
-				const TConstArrayView<FVector3f> SimRestPosition = ClothPatternFacade.GetSimRestPosition();
-				const TConstArrayView<FVector3f> SimRestNormal = ClothPatternFacade.GetSimRestNormal();
-				const TArrayView<FVector3f> RenderTangentU = ClothPatternFacade.GetRenderTangentU();
-				const TArrayView<FVector3f> RenderTangentV = ClothPatternFacade.GetRenderTangentV();
+			const float Denom = UV01.X * UV02.Y - UV01.Y * UV02.X;
+			const float InvDenom = (Denom < SMALL_NUMBER) ? 0.f : 1.f / Denom;
+			const FVector3f TangentU = (Pos01 * UV02.Y - Pos02 * UV01.Y) * InvDenom;
+			const FVector3f TangentV = (Pos02 * UV01.X - Pos01 * UV02.X) * InvDenom;
 
-				FVector2f MinPosition(TNumericLimits<float>::Max());
-				FVector2f MaxPosition(TNumericLimits<float>::Lowest());
+			for (int32 PointIndex = 0; PointIndex < 3; ++PointIndex)
+			{
+				RenderTangentU[Face[PointIndex]] += TangentU;
+				RenderTangentV[Face[PointIndex]] += TangentV;
+			}
+		}
 
-				// Calculate UVs scale and zero out tangents
-				for (int32 VertexIndex = 0; VertexIndex < NumVertices; ++VertexIndex)
-				{
-					MinPosition = FVector2f::Min(MinPosition, SimPosition[VertexIndex]);
-					MaxPosition = FVector2f::Max(MaxPosition, SimPosition[VertexIndex]);
+		// Vertex group
+		const TArrayView<FVector3f> RenderPosition = ClothFacade.GetRenderPosition();
+		const TArrayView<FVector3f> RenderNormal = ClothFacade.GetRenderNormal();
+		const TArrayView<TArray<FVector2f>> RenderUVs = ClothFacade.GetRenderUVs();
+		const TArrayView<FLinearColor> RenderColor = ClothFacade.GetRenderColor();
+		const TArrayView<TArray<int32>> RenderBoneIndices = ClothFacade.GetRenderBoneIndices();
+		const TArrayView<TArray<float>> RenderBoneWeights = ClothFacade.GetRenderBoneWeights();
 
-					RenderTangentU[VertexIndex] = FVector3f::ZeroVector;
-					RenderTangentV[VertexIndex] = FVector3f::ZeroVector;
-				}
-				const FVector2f UVScale = MaxPosition - MinPosition;
-				const FVector2f UVInvScale(
-					UVScale.X < SMALL_NUMBER ? 0.f : 1.f / UVScale.X,
-					UVScale.Y < SMALL_NUMBER ? 0.f : 1.f / UVScale.Y);
-
-				// Face group
-				const TConstArrayView<FIntVector3> SimIndices = ClothPatternFacade.GetSimIndices();
-				const TArrayView<FIntVector3> RenderIndices = ClothPatternFacade.GetRenderIndices();
-				const TArrayView<int32> RenderMaterialIndex = ClothPatternFacade.GetRenderMaterialIndex();
-				const TArrayView<FVector3f> LodRenderTangentU = ClothLodFacade.GetRenderTangentU();
-				const TArrayView<FVector3f> LodRenderTangentV = ClothLodFacade.GetRenderTangentV();
-				for (int32 FaceIndex = 0; FaceIndex < NumFaces; ++FaceIndex)
-				{
-					const FIntVector3& Face = SimIndices[FaceIndex];
-
-					RenderIndices[FaceIndex] = Face;
-					RenderMaterialIndex[FaceIndex] = MaterialIndex;
-
-					const FVector3f Pos01 = LodSimRestPosition[Face[1]] - LodSimRestPosition[Face[0]];
-					const FVector3f Pos02 = LodSimRestPosition[Face[2]] - LodSimRestPosition[Face[0]];
-					const FVector2f UV01 = LodSimPosition[Face[1]] - LodSimPosition[Face[0]];
-					const FVector2f UV02 = LodSimPosition[Face[2]] - LodSimPosition[Face[0]];
-					
-					const float Denom = UV01.X * UV02.Y - UV01.Y * UV02.X;
-					const float InvDenom = (Denom < SMALL_NUMBER) ? 0.f : 1.f / Denom;
-					const FVector3f TangentU = (Pos01 * UV02.Y - Pos02 * UV01.Y) * InvDenom;
-					const FVector3f TangentV = (Pos02 * UV01.X - Pos01 * UV02.X) * InvDenom;
-
-					for (int32 PointIndex = 0; PointIndex < 3; ++PointIndex)
-					{
-						LodRenderTangentU[Face[PointIndex]] += TangentU;
-						LodRenderTangentV[Face[PointIndex]] += TangentV;
-					}
-				}
+		// NOTE: This sim data is stored on welded vertices.
+		const TArrayView<TArray<int32>> SimBoneIndices = ClothFacade.GetSimBoneIndices();
+		const TArrayView<TArray<float>> SimBoneWeights = ClothFacade.GetSimBoneWeights();
 
-				// Vertex group
-				const TArrayView<FVector3f> RenderPosition = ClothPatternFacade.GetRenderPosition();
-				const TArrayView<FVector3f> RenderNormal = ClothPatternFacade.GetRenderNormal();
-				const TArrayView<TArray<FVector2f>> RenderUVs = ClothPatternFacade.GetRenderUVs();
-				const TArrayView<FLinearColor> RenderColor = ClothPatternFacade.GetRenderColor();
-				const TArrayView<int32> RenderNumBoneInfluences = ClothPatternFacade.GetRenderNumBoneInfluences();
-				const TArrayView<TArray<int32>> RenderBoneIndices = ClothPatternFacade.GetRenderBoneIndices();
-				const TArrayView<TArray<float>> RenderBoneWeights = ClothPatternFacade.GetRenderBoneWeights();
-				const TArrayView<int32> SimNumBoneInfluences = ClothPatternFacade.GetSimNumBoneInfluences();
-				const TArrayView<TArray<int32>> SimBoneIndices = ClothPatternFacade.GetSimBoneIndices();
-				const TArrayView<TArray<float>> SimBoneWeights = ClothPatternFacade.GetSimBoneWeights();
-
-				for (int32 VertexIndex = 0; VertexIndex < NumVertices; ++VertexIndex)
-				{
-					RenderPosition[VertexIndex] = SimRestPosition[VertexIndex];
-					RenderNormal[VertexIndex] = -SimRestNormal[VertexIndex];  // Simulation normals use reverse normals
-					RenderUVs[VertexIndex] = { SimPosition[VertexIndex] * UVInvScale };
-					RenderColor[VertexIndex] = FLinearColor::White;
-					RenderTangentU[VertexIndex].Normalize();
-					RenderTangentV[VertexIndex].Normalize();
-					RenderNumBoneInfluences[VertexIndex] = SimNumBoneInfluences[VertexIndex];
-					RenderBoneIndices[VertexIndex] = SimBoneIndices[VertexIndex];
-					RenderBoneWeights[VertexIndex] = SimBoneWeights[VertexIndex];
-				}
-			}
+		for (int32 VertexIndex = 0; VertexIndex < TotalNumVertices; ++VertexIndex)
+		{
+			const int32 VertexIndex3D = SimVertex3DLookup[VertexIndex];
+
+			RenderPosition[VertexIndex] = SimPosition3D[VertexIndex3D];
+			RenderNormal[VertexIndex] = -SimNormal[VertexIndex3D];  // Simulation normals use reverse normals
+			RenderUVs[VertexIndex] = { SimPosition2D[VertexIndex] * UVInvScale };
+			RenderColor[VertexIndex] = FLinearColor::White;
+			RenderTangentU[VertexIndex].Normalize();
+			RenderTangentV[VertexIndex].Normalize();
+			RenderBoneIndices[VertexIndex] = SimBoneIndices[VertexIndex3D];
+			RenderBoneWeights[VertexIndex] = SimBoneWeights[VertexIndex3D];
 		}
 	}
 
@@ -220,13 +568,14 @@ namespace UE::Chaos::ClothAsset
 		bool bReverseSimMeshWindingOrder,
 		bool bReverseRenderMeshNormals,
 		bool bReverseRenderMeshWindingOrder,
-		const TArray<int32>& PatternSelection)
+		const TArray<int32>& SimPatternSelection,
+		const TArray<int32>& RenderPatternSelection)
 	{
-		auto ReverseSimNormals = [](const TArrayView<FVector3f>& SimRestNormal)
+		auto ReverseSimNormals = [](const TArrayView<FVector3f>& SimNormal)
 			{
-				for (int32 VertexIndex = 0; VertexIndex < SimRestNormal.Num(); ++VertexIndex)
+				for (int32 VertexIndex = 0; VertexIndex < SimNormal.Num(); ++VertexIndex)
 				{
-					SimRestNormal[VertexIndex] = -SimRestNormal[VertexIndex];
+					SimNormal[VertexIndex] = -SimNormal[VertexIndex];
 				}
 			};
 		auto ReverseRenderNormals = [](const TArrayView<FVector3f>& RenderNormal, const TArrayView<FVector3f>& RenderTangentU)
@@ -248,54 +597,82 @@ namespace UE::Chaos::ClothAsset
 
 		FCollectionClothFacade ClothFacade(ClothCollection);
 
-		for (int32 LodIndex = 0; LodIndex < ClothFacade.GetNumLods(); ++LodIndex)
+		if (SimPatternSelection.IsEmpty())
 		{
-			FCollectionClothLodFacade ClothLodFacade = ClothFacade.GetLod(LodIndex);
+			if (bReverseSimMeshNormals)
+			{
+				ReverseSimNormals(ClothFacade.GetSimNormal());
+			}
 
-			if (PatternSelection.IsEmpty())
+			if (bReverseSimMeshWindingOrder)
 			{
-				if (bReverseSimMeshNormals)
-				{
-					ReverseSimNormals(ClothLodFacade.GetSimRestNormal());
-				}
-				if (bReverseSimMeshWindingOrder)
-				{
-					ReverseWindingOrder(ClothLodFacade.GetSimIndices());
-				}
-				if (bReverseRenderMeshNormals)
-				{
-					ReverseRenderNormals(ClothLodFacade.GetRenderNormal(), ClothLodFacade.GetRenderTangentU());
-				}
-				if (bReverseRenderMeshWindingOrder)
-				{
-					ReverseWindingOrder(ClothLodFacade.GetRenderIndices());
-				}
+				ReverseWindingOrder(ClothFacade.GetSimIndices2D());
+				ReverseWindingOrder(ClothFacade.GetSimIndices3D());
+			}
+		}
+		else
+		{
+			// Sim Normals live on welded vertices. We don't want to double flip normals that live in multiple patterns
+			TBitArray AlreadyFlippedNormal;
+			if (bReverseSimMeshNormals)
+			{
+				AlreadyFlippedNormal.Init(false, ClothFacade.GetNumSimVertices3D());
 			}
-			else
+			const TArrayView<FVector3f> AllSimNormals = ClothFacade.GetSimNormal();
+			for (int32 PatternIndex = 0; PatternIndex < ClothFacade.GetNumSimPatterns(); ++PatternIndex)
 			{
-				for (int32 PatternIndex = 0; PatternIndex < ClothLodFacade.GetNumPatterns(); ++PatternIndex)
+				if (SimPatternSelection.Find(PatternIndex) != INDEX_NONE)
 				{
-					if (PatternSelection.Find(PatternIndex) != INDEX_NONE)
-					{
-						FCollectionClothPatternFacade ClothPatternFacade = ClothLodFacade.GetPattern(PatternIndex);
+					FCollectionClothSimPatternFacade ClothPatternFacade = ClothFacade.GetSimPattern(PatternIndex);
 
-						if (bReverseSimMeshNormals)
-						{
-							ReverseSimNormals(ClothPatternFacade.GetSimRestNormal());
-						}
-						if (bReverseSimMeshWindingOrder)
-						{
-							ReverseWindingOrder(ClothPatternFacade.GetSimIndices());
-						}
-						if (bReverseRenderMeshNormals)
-						{
-							ReverseRenderNormals(ClothPatternFacade.GetRenderNormal(), ClothPatternFacade.GetRenderTangentU());
-						}
-						if (bReverseRenderMeshWindingOrder)
+					if (bReverseSimMeshNormals)
+					{
+						const TConstArrayView<int32> SimVertex3DLookup = static_cast<FCollectionClothSimPatternConstFacade&>(ClothPatternFacade).GetSimVertex3DLookup();
+						for (int32 VertexIndex3D : SimVertex3DLookup)
 						{
-							ReverseWindingOrder(ClothPatternFacade.GetRenderIndices());
+							if (!AlreadyFlippedNormal[VertexIndex3D])
+							{
+								AllSimNormals[VertexIndex3D] = -AllSimNormals[VertexIndex3D];
+								AlreadyFlippedNormal[VertexIndex3D] = true;
+							}
 						}
 					}
+
+					if (bReverseSimMeshWindingOrder)
+					{
+						ReverseWindingOrder(ClothPatternFacade.GetSimIndices2D());
+						ReverseWindingOrder(ClothPatternFacade.GetSimIndices3D());
+					}
+				}
+			}
+		}
+
+		if (RenderPatternSelection.IsEmpty())
+		{
+			if (bReverseRenderMeshNormals)
+			{
+				ReverseRenderNormals(ClothFacade.GetRenderNormal(), ClothFacade.GetRenderTangentU());
+			}
+			if (bReverseRenderMeshWindingOrder)
+			{
+				ReverseWindingOrder(ClothFacade.GetRenderIndices());
+			}
+		}
+		else
+		{
+			for (int32 PatternIndex = 0; PatternIndex < ClothFacade.GetNumRenderPatterns(); ++PatternIndex)
+			{
+				if (RenderPatternSelection.Find(PatternIndex) != INDEX_NONE)
+				{
+					FCollectionClothRenderPatternFacade ClothPatternFacade = ClothFacade.GetRenderPattern(PatternIndex);
+					if (bReverseRenderMeshNormals)
+					{
+						ReverseRenderNormals(ClothPatternFacade.GetRenderNormal(), ClothPatternFacade.GetRenderTangentU());
+					}
+					if (bReverseRenderMeshWindingOrder)
+					{
+						ReverseWindingOrder(ClothPatternFacade.GetRenderIndices());
+					}
 				}
 			}
 		}
@@ -304,103 +681,77 @@ namespace UE::Chaos::ClothAsset
 	void FClothGeometryTools::BindMeshToRootBone(
 		const TSharedPtr<FManagedArrayCollection>& ClothCollection,
 		bool bBindSimMesh,
-		bool bBindRenderMesh,
-		const TArray<int32> Lods)
+		bool bBindRenderMesh)
 	{
 		if (!bBindSimMesh && !bBindRenderMesh)
 		{
 			return;
 		}
-		
-		FCollectionClothFacade ClothFacade(ClothCollection);
 
-		TArray<int32> LodsToBind;
-		if (Lods.IsEmpty())
+		FCollectionClothFacade ClothFacade(ClothCollection);
+		if (bBindSimMesh)
 		{
-			LodsToBind.Reserve(ClothFacade.GetNumLods());
-			
-			for (int32 LodIndex = 0; LodIndex < ClothFacade.GetNumLods(); ++LodIndex)
+			const int32 NumVertices = ClothFacade.GetNumSimVertices3D();
+			TArrayView<TArray<int32>> BoneIndices = ClothFacade.GetSimBoneIndices();
+			TArrayView<TArray<float>> BoneWeights = ClothFacade.GetSimBoneWeights();
+
+			for (int32 VertexIndex = 0; VertexIndex < NumVertices; ++VertexIndex)
 			{
-				LodsToBind.Add(LodIndex);
-			}
-		}
-		else
-		{	
-			LodsToBind.Reserve(Lods.Num());
-			
-			for (const int32 Lod : Lods)
-			{	
-				// Make sure the Lod indices are valid.
-				if (Lod < ClothFacade.GetNumLods())
-				{
-					LodsToBind.Add(Lod);
-				}
+				BoneIndices[VertexIndex] = { 0 };
+				BoneWeights[VertexIndex] = { 1.0f };
 			}
 		}
 
-		for (const int32 LodIndex : LodsToBind)
+		if (bBindRenderMesh)
 		{
-			FCollectionClothLodFacade ClothLodFacade = ClothFacade.GetLod(LodIndex);
-			
-			if (bBindSimMesh)
-			{	
-				const int32 NumVertices = ClothLodFacade.GetNumSimVertices();
-				TArrayView<int32> NumBoneInfluences = ClothLodFacade.GetSimNumBoneInfluences();
-				TArrayView<TArray<int32>> BoneIndices = ClothLodFacade.GetSimBoneIndices();
-				TArrayView<TArray<float>> BoneWeights = ClothLodFacade.GetSimBoneWeights();
-
-				for (int32 VertexIndex = 0; VertexIndex < NumVertices; ++VertexIndex)
-				{
-					NumBoneInfluences[VertexIndex] = 1;
-					BoneIndices[VertexIndex] = {0};
-					BoneWeights[VertexIndex] = {1.0f};
-				}
-			}
+			const int32 NumVertices = ClothFacade.GetNumRenderVertices();
+			TArrayView<TArray<int32>> BoneIndices = ClothFacade.GetRenderBoneIndices();
+			TArrayView<TArray<float>> BoneWeights = ClothFacade.GetRenderBoneWeights();
 
-			if (bBindRenderMesh)
+			for (int32 VertexIndex = 0; VertexIndex < NumVertices; ++VertexIndex)
 			{
-				const int32 NumVertices = ClothLodFacade.GetNumRenderVertices();
-				TArrayView<int32> NumBoneInfluences = ClothLodFacade.GetRenderNumBoneInfluences();
-				TArrayView<TArray<int32>> BoneIndices = ClothLodFacade.GetRenderBoneIndices();
-				TArrayView<TArray<float>> BoneWeights = ClothLodFacade.GetRenderBoneWeights();
-
-				for (int32 VertexIndex = 0; VertexIndex < NumVertices; ++VertexIndex)
-				{
-					NumBoneInfluences[VertexIndex] = 1;
-					BoneIndices[VertexIndex] = {0};
-					BoneWeights[VertexIndex] = {1.0f};
-				}
+				BoneIndices[VertexIndex] = { 0 };
+				BoneWeights[VertexIndex] = { 1.0f };
 			}
 		}
 	}
 
-	TArray<float> FClothGeometryTools::BuildWeldedWeightMapForLod(
-		const TSharedPtr<const FManagedArrayCollection>& ClothCollection,
-		int32 LodIndex,
-		const FName& WeightMapName,
-		const TArray<TArray<int32>>& WeldedToUnweldedMap)
+	void FClothGeometryTools::BuildSimMeshFromDynamicMesh(
+		const TSharedPtr<FManagedArrayCollection>& ClothCollection,
+		const UE::Geometry::FDynamicMesh3& DynamicMesh, int32 UVChannelIndex, const FVector2f& UVScale)
 	{
-		FCollectionClothConstFacade ClothFacade(ClothCollection);
-		FCollectionClothLodConstFacade ClothLodFacade = ClothFacade.GetLod(LodIndex);
+		using namespace Private::SimMeshBuilder;
 
-		TConstArrayView<float> ClothLodWeightMap = ClothLodFacade.GetWeightMap(WeightMapName);
+		DeleteSimMesh(ClothCollection);
 
-		TArray<float> WeldedWeightMap;
-		WeldedWeightMap.SetNumZeroed(WeldedToUnweldedMap.Num());
+		const UE::Geometry::FDynamicMeshAttributeSet* const AttributeSet = DynamicMesh.Attributes();
+		const UE::Geometry::FDynamicMeshUVOverlay* const UVOverlay = AttributeSet ? AttributeSet->GetUVLayer(UVChannelIndex) : nullptr;
 
-		for (int32 WeldedIndex = 0; WeldedIndex < WeldedToUnweldedMap.Num(); ++WeldedIndex)
+		TArray<FIsland> Islands;
+		if (UVOverlay)
+		{
+			BuildIslandsFromDynamicMeshUVs(*UVOverlay, UVScale, Islands);
+		}
+		else
 		{
-			const TArray<int32>& SourceVertices = WeldedToUnweldedMap[WeldedIndex];
-			check(SourceVertices.Num());
-			for (const int32 SourceVertex : SourceVertices)
+			UnwrapDynamicMesh(DynamicMesh, Islands);
+		}
+
+		FCollectionClothFacade Cloth(ClothCollection);
+		for (FIsland& Island : Islands)
+		{
+			if (Island.Indices.Num() && Island.Positions2D.Num() && Island.Positions3D.Num())
 			{
-				WeldedWeightMap[WeldedIndex] += ClothLodWeightMap[SourceVertex];
+				FCollectionClothSimPatternFacade Pattern = Cloth.AddGetSimPattern();
+				Pattern.Initialize(Island.Positions2D, Island.Positions3D, Island.Indices);
 			}
-
-			WeldedWeightMap[WeldedIndex] /= (float)SourceVertices.Num();
 		}
 
-		return WeldedWeightMap;
+		TArray<FSeam> Seams;
+		BuildSeams(Islands, DynamicMesh, Seams);  // Build the seam information as to be able to re-weld the mesh for simulation
+		for (const FSeam& Seam : Seams)
+		{
+			AddSeam(Cloth, Seam);
+		}
 	}
-
 }  // End namespace UE::Chaos::ClothAsset
