ClothCollection Schema rework. Existing cloth assets will be emptied, but their associated Dataflow graphs and nodes are still valid (with the exception of AddWeightMapNodes which need to be repainted

diff --git a/Engine/Plugins/Experimental/ChaosClothAsset/Source/ChaosClothAsset/Private/ChaosClothAsset/CollectionClothFacade.cpp b/Engine/Plugins/Experimental/ChaosClothAsset/Source/ChaosClothAsset/Private/ChaosClothAsset/CollectionClothFacade.cpp
--- a/Engine/Plugins/Experimental/ChaosClothAsset/Source/ChaosClothAsset/Private/ChaosClothAsset/CollectionClothFacade.cpp
+++ b/Engine/Plugins/Experimental/ChaosClothAsset/Source/ChaosClothAsset/Private/ChaosClothAsset/CollectionClothFacade.cpp
@@ -1,8 +1,8 @@
 // Copyright Epic Games, Inc. All Rights Reserved.
 
 #include "ChaosClothAsset/CollectionClothFacade.h"
-#include "ChaosClothAsset/CollectionClothLodFacade.h"
 #include "ChaosClothAsset/ClothCollection.h"
+#include "ChaosClothAsset/ClothGeometryTools.h"
 #include "Chaos/ChaosArchive.h"
 #include "UObject/UE5MainStreamObjectVersion.h"
 
@@ -13,30 +13,266 @@ namespace UE::Chaos::ClothAsset
 	{
 	}
 
+	FCollectionClothConstFacade::FCollectionClothConstFacade(const TSharedPtr<const FClothCollection>& ClothCollection)
+		:ClothCollection(ClothCollection)
+	{
+	}
+
 	bool FCollectionClothConstFacade::IsValid() const
 	{
-		return ClothCollection->IsValid();
+		return ClothCollection->IsValid() && ClothCollection->GetNumElements(FClothCollection::LodsGroup) == 1;
+	}
+
+	const FString& FCollectionClothConstFacade::GetPhysicsAssetPathName() const
+	{
+		static const FString EmptyString;
+		return ClothCollection->GetPhysicsAssetPathName() && ClothCollection->GetNumElements(FClothCollection::LodsGroup) > 0 ? (*ClothCollection->GetPhysicsAssetPathName())[0] : EmptyString;
+	}
+
+	const FString& FCollectionClothConstFacade::GetSkeletonAssetPathName() const
+	{
+		static const FString EmptyString;
+		return ClothCollection->GetSkeletonAssetPathName() && ClothCollection->GetNumElements(FClothCollection::LodsGroup) > 0 ? (*ClothCollection->GetSkeletonAssetPathName())[0] : EmptyString;
+	}
+
+	int32 FCollectionClothConstFacade::GetNumSimVertices2D() const
+	{
+		return ClothCollection->GetNumElements(FClothCollection::SimVertices2DGroup);
+	}
+
+	TConstArrayView<FVector2f> FCollectionClothConstFacade::GetSimPosition2D() const
+	{
+		return ClothCollection->GetElements(ClothCollection->GetSimPosition2D());
+	}
+
+	TConstArrayView<int32> FCollectionClothConstFacade::GetSimVertex3DLookup() const
+	{
+		return ClothCollection->GetElements(ClothCollection->GetSimVertex3DLookup());
+	}
+
+	int32 FCollectionClothConstFacade::GetNumSimVertices3D() const
+	{
+		return ClothCollection->GetNumElements(FClothCollection::SimVertices3DGroup);
+	}
+
+	TConstArrayView<FVector3f> FCollectionClothConstFacade::GetSimPosition3D() const
+	{
+		return ClothCollection->GetElements(ClothCollection->GetSimPosition3D());
+	}
+
+	TConstArrayView<FVector3f> FCollectionClothConstFacade::GetSimNormal() const
+	{
+		return ClothCollection->GetElements(ClothCollection->GetSimNormal());
+	}
+
+	TConstArrayView<TArray<int32>> FCollectionClothConstFacade::GetSimBoneIndices() const
+	{
+		return ClothCollection->GetElements(ClothCollection->GetSimBoneIndices());
+	}
+
+	TConstArrayView<TArray<float>> FCollectionClothConstFacade::GetSimBoneWeights() const
+	{
+		return ClothCollection->GetElements(ClothCollection->GetSimBoneWeights());
+	}
+
+	TConstArrayView<TArray<int32>> FCollectionClothConstFacade::GetTetherKinematicIndex() const
+	{
+		return ClothCollection->GetElements(ClothCollection->GetTetherKinematicIndex());
+	}
+
+	TConstArrayView<TArray<float>> FCollectionClothConstFacade::GetTetherReferenceLength() const
+	{
+		return ClothCollection->GetElements(ClothCollection->GetTetherReferenceLength());
+	}
+
+	TConstArrayView<TArray<int32>> FCollectionClothConstFacade::GetSimVertex2DLookup() const
+	{
+		return ClothCollection->GetElements(ClothCollection->GetSimVertex2DLookup());
+	}
+
+	TConstArrayView<TArray<int32>> FCollectionClothConstFacade::GetSeamStitchLookup() const
+	{
+		return ClothCollection->GetElements(ClothCollection->GetSeamStitchLookup());
+	}
+
+	int32 FCollectionClothConstFacade::GetNumSimFaces() const
+	{
+		return ClothCollection->GetNumElements(FClothCollection::SimFacesGroup);
+	}
+
+	TConstArrayView<FIntVector3> FCollectionClothConstFacade::GetSimIndices2D() const
+	{
+		return ClothCollection->GetElements(ClothCollection->GetSimIndices2D());
 	}
 
-	FCollectionClothLodConstFacade FCollectionClothConstFacade::GetLod(int32 LodIndex) const
+	TConstArrayView<FIntVector3> FCollectionClothConstFacade::GetSimIndices3D() const
 	{
-		check(LodIndex >= 0 && LodIndex < GetNumLods());
-		return FCollectionClothLodConstFacade(ClothCollection, LodIndex);
+		return ClothCollection->GetElements(ClothCollection->GetSimIndices3D());
 	}
 
-	int32 FCollectionClothConstFacade::GetNumLods() const
+	int32 FCollectionClothConstFacade::GetNumSimPatterns() const
 	{
-		return ClothCollection->GetNumElements(FClothCollection::LodsGroup);
+		return ClothCollection->GetNumElements(FClothCollection::SimPatternsGroup);
+	}
+
+	FCollectionClothSimPatternConstFacade FCollectionClothConstFacade::GetSimPattern(int32 PatternIndex) const
+	{
+		return FCollectionClothSimPatternConstFacade(ClothCollection, PatternIndex);
+	}
+
+	int32 FCollectionClothConstFacade::FindSimPatternByVertex2D(int32 Vertex2DIndex) const
+	{
+		return ClothCollection->GetArrayIndexForContainedElement(ClothCollection->GetSimVertices2DStart(), ClothCollection->GetSimVertices2DEnd(), Vertex2DIndex);
+	}
+
+	int32 FCollectionClothConstFacade::FindSimPatternByFaceIndex(int32 FaceIndex) const
+	{
+		return ClothCollection->GetArrayIndexForContainedElement(ClothCollection->GetSimFacesStart(), ClothCollection->GetSimFacesEnd(), FaceIndex);
+	}
+
+	int32 FCollectionClothConstFacade::GetNumRenderPatterns() const
+	{
+		return ClothCollection->GetNumElements(FClothCollection::RenderPatternsGroup);
+	}
+
+	FCollectionClothRenderPatternConstFacade FCollectionClothConstFacade::GetRenderPattern(int32 PatternIndex) const
+	{
+		return FCollectionClothRenderPatternConstFacade(ClothCollection, PatternIndex);
+	}
+
+	TConstArrayView<FString> FCollectionClothConstFacade::GetRenderMaterialPathName() const
+	{
+		return ClothCollection->GetElements(ClothCollection->GetRenderMaterialPathName());
+	}
+
+	int32 FCollectionClothConstFacade::FindRenderPatternByVertex(int32 VertexIndex) const
+	{
+		return ClothCollection->GetArrayIndexForContainedElement(ClothCollection->GetRenderVerticesStart(), ClothCollection->GetRenderVerticesEnd(), VertexIndex);
+	}
+
+	int32 FCollectionClothConstFacade::FindRenderPatternByFaceIndex(int32 FaceIndex) const
+	{
+		return ClothCollection->GetArrayIndexForContainedElement(ClothCollection->GetRenderFacesStart(), ClothCollection->GetRenderFacesEnd(), FaceIndex);
+	}
+
+	int32 FCollectionClothConstFacade::GetNumSeams() const
+	{
+		return ClothCollection->GetNumElements(FClothCollection::SeamsGroup);
+	}
+
+	FCollectionClothSeamConstFacade FCollectionClothConstFacade::GetSeam(int32 SeamIndex) const
+	{
+		return FCollectionClothSeamConstFacade(ClothCollection, SeamIndex);
+	}
+
+	int32 FCollectionClothConstFacade::GetNumRenderVertices() const
+	{
+		return ClothCollection->GetNumElements(FClothCollection::RenderVerticesGroup);
+	}
+
+	TConstArrayView<FVector3f> FCollectionClothConstFacade::GetRenderPosition() const
+	{
+		return ClothCollection->GetElements(ClothCollection->GetRenderPosition());
+	}
+
+	TConstArrayView<FVector3f> FCollectionClothConstFacade::GetRenderNormal() const
+	{
+		return ClothCollection->GetElements(ClothCollection->GetRenderNormal());
+	}
+
+	TConstArrayView<FVector3f> FCollectionClothConstFacade::GetRenderTangentU() const
+	{
+		return ClothCollection->GetElements(ClothCollection->GetRenderTangentU());
+	}
+
+	TConstArrayView<FVector3f> FCollectionClothConstFacade::GetRenderTangentV() const
+	{
+		return ClothCollection->GetElements(ClothCollection->GetRenderTangentV());
+	}
+
+	TConstArrayView<TArray<FVector2f>> FCollectionClothConstFacade::GetRenderUVs() const
+	{
+		return ClothCollection->GetElements(ClothCollection->GetRenderUVs());
+	}
+
+	TConstArrayView<FLinearColor> FCollectionClothConstFacade::GetRenderColor() const
+	{
+		return ClothCollection->GetElements(ClothCollection->GetRenderColor());
+	}
+
+	TConstArrayView<TArray<int32>> FCollectionClothConstFacade::GetRenderBoneIndices() const
+	{
+		return ClothCollection->GetElements(ClothCollection->GetRenderBoneIndices());
+	}
+
+	TConstArrayView<TArray<float>> FCollectionClothConstFacade::GetRenderBoneWeights() const
+	{
+		return ClothCollection->GetElements(ClothCollection->GetRenderBoneWeights());
+	}
+
+	int32 FCollectionClothConstFacade::GetNumRenderFaces() const
+	{
+		return ClothCollection->GetNumElements(FClothCollection::RenderFacesGroup);
+	}
+
+	TConstArrayView<FIntVector3> FCollectionClothConstFacade::GetRenderIndices() const
+	{
+		return ClothCollection->GetElements(ClothCollection->GetRenderIndices());
 	}
 
 	bool FCollectionClothConstFacade::HasWeightMap(const FName& Name) const
 	{
-		return ClothCollection->HasUserDefinedAttribute<float>(Name, FClothCollection::SimVerticesGroup);
+		return ClothCollection->HasUserDefinedAttribute<float>(Name, FClothCollection::SimVertices3DGroup);
 	}
 
 	TArray<FName> FCollectionClothConstFacade::GetWeightMapNames() const
 	{
-		return ClothCollection->GetUserDefinedAttributeNames<float>(FClothCollection::SimVerticesGroup);
+		return ClothCollection->GetUserDefinedAttributeNames<float>(FClothCollection::SimVertices3DGroup);
+	}
+
+	TConstArrayView<float> FCollectionClothConstFacade::GetWeightMap(const FName& Name) const
+	{
+		return ClothCollection->GetElements(ClothCollection->GetUserDefinedAttribute<float>(Name, FClothCollection::SimVertices3DGroup));
+	}
+
+	void FCollectionClothConstFacade::BuildSimulationMesh(TArray<FVector3f>& Positions, TArray<FVector3f>& Normals, TArray<uint32>& Indices, TArray<FVector2f>& PatternsPositions, TArray<uint32>& PatternsIndices,
+		TArray<uint32>& PatternToWeldedIndices, TArray<TArray<int32>>* OptionalWeldedToPatternIndices) const
+	{
+		Positions = GetSimPosition3D();
+		Normals = GetSimNormal();
+		PatternsPositions = GetSimPosition2D();
+		PatternToWeldedIndices = GetSimVertex3DLookup();
+		if (OptionalWeldedToPatternIndices)
+		{
+			*OptionalWeldedToPatternIndices = GetSimVertex2DLookup();
+		}
+
+		for (const FIntVector3& Face : GetSimIndices2D())
+		{
+			PatternsIndices.Add(Face[0]);
+			PatternsIndices.Add(Face[1]);
+			PatternsIndices.Add(Face[2]);
+		}
+
+		// It's possible that welding created degenerate 3D triangles. Remove them when copying over. Keep 2D and 3D triangles in sync
+		PatternsIndices.Reset(GetNumSimFaces() * 3);
+		Indices.Reset(GetNumSimFaces() * 3);
+		for(int32 FaceIdx = 0; FaceIdx < GetNumSimFaces(); ++FaceIdx)
+		{
+			const FIntVector3 Face3D = GetSimIndices3D()[FaceIdx];
+
+			if (Face3D[0] != Face3D[1] && Face3D[0] != Face3D[2] && Face3D[1] != Face3D[2])
+			{
+				Indices.Add(Face3D[0]);
+				Indices.Add(Face3D[1]);
+				Indices.Add(Face3D[2]);
+
+				const FIntVector3 Face2D = GetSimIndices2D()[FaceIdx];
+				PatternsIndices.Add(Face2D[0]);
+				PatternsIndices.Add(Face2D[1]);
+				PatternsIndices.Add(Face2D[2]);
+			}
+		}
 	}
 
 	FCollectionClothFacade::FCollectionClothFacade(const TSharedPtr<FManagedArrayCollection>& ManagedArrayCollection)
@@ -44,95 +280,347 @@ namespace UE::Chaos::ClothAsset
 	{
 	}
 
+	FCollectionClothFacade::FCollectionClothFacade(const TSharedPtr<FClothCollection>& InClothCollection)
+		: FCollectionClothConstFacade(InClothCollection)
+	{
+	}
+
 	void FCollectionClothFacade::DefineSchema()
 	{
 		GetClothCollection()->DefineSchema();
+		SetDefaults();
 	}
 
 	void FCollectionClothFacade::Reset()
 	{
-		check(IsValid());
+		// Hack to reset all Lods data
+		GetClothCollection()->SetNumElements(0, FClothCollection::LodsGroup);
+		GetClothCollection()->SetNumElements(1, FClothCollection::LodsGroup);
+		RemoveAllSimVertices3D();
+		SetNumSimPatterns(0);
+		SetNumRenderPatterns(0);
+		SetNumSeams(0); // Do this after removing SimVertices3D and SimPatterns. Otherwise, Seams will do a bunch of unnecessary work to unseam stuff.
+	}
+
+	void FCollectionClothFacade::Initialize(const FCollectionClothConstFacade& Other)
+	{
+		Reset();
+
+		// LODs Group
+		SetPhysicsAssetPathName(Other.GetPhysicsAssetPathName());
+		SetSkeletonAssetPathName(Other.GetSkeletonAssetPathName());
+		
+		Append(Other);
+	}
+
+	void FCollectionClothFacade::Append(const FCollectionClothConstFacade& Other)
+	{
+		// LODs Group 
+		// Just keep original data.
 
-		const int32 NumLods = GetNumLods();
-		for (int32 LodIndex = 0; LodIndex < NumLods; ++LodIndex)
+		// Very important order of operations to ensure indices don't get messed up:
+		// 1) Append 3D Vertices, but don't set 2D Lookups or SeamStitch Lookups because those indices don't exist yet.
+		// 2) Append Sim Patterns (includes 2D Vertices, have 3D dependency)
+		// 3) Append Seams (have 2D and 3D dependency)
+		// 4) Append 2DLookups (2D dep)  and SeamStitchLookups (SeamStitch Dep)
+
+		// Sim Vertices 3D Group
+		const int32 StartNumSimVertices3D = GetNumSimVertices3D();
+		const int32 OtherNumSimVertices3D = Other.GetNumSimVertices3D();
+		GetClothCollection()->SetNumElements(StartNumSimVertices3D + OtherNumSimVertices3D, FClothCollection::SimVertices3DGroup);
+		FClothCollection::CopyArrayViewData(GetSimPosition3D().Right(OtherNumSimVertices3D), Other.GetSimPosition3D());
+		FClothCollection::CopyArrayViewData(GetSimNormal().Right(OtherNumSimVertices3D), Other.GetSimNormal());
+		FClothCollection::CopyArrayViewData(GetSimBoneIndices().Right(OtherNumSimVertices3D), Other.GetSimBoneIndices());
+		FClothCollection::CopyArrayViewData(GetSimBoneWeights().Right(OtherNumSimVertices3D), Other.GetSimBoneWeights());
+		FClothCollection::CopyArrayViewDataAndApplyOffset(GetTetherKinematicIndex().Right(OtherNumSimVertices3D), Other.GetTetherKinematicIndex(), StartNumSimVertices3D);
+		FClothCollection::CopyArrayViewData(GetTetherReferenceLength().Right(OtherNumSimVertices3D), Other.GetTetherReferenceLength());
+
+		// Sim Patterns Group
+		const int32 StartNumSimVertices2D = GetNumSimVertices2D();
+		const int32 StartNumSimPatterns = GetNumSimPatterns();
+		const int32 OtherNumSimPatterns = Other.GetNumSimPatterns();
+		SetNumSimPatterns(StartNumSimPatterns + OtherNumSimPatterns);
+		for (int32 PatternIndex = 0; PatternIndex < OtherNumSimPatterns; ++PatternIndex)
 		{
-			GetLod(LodIndex).Reset();
+			GetSimPattern(StartNumSimPatterns + PatternIndex).Initialize(Other.GetSimPattern(PatternIndex), StartNumSimVertices3D);
 		}
 
-		GetClothCollection()->SetNumElements(0, FClothCollection::LodsGroup);
+		// Seams Group
+		const int32 StartNumSeams = GetNumSeams();
+		const int32 OtherNumSeams = Other.GetNumSeams();
+		SetNumSeams(StartNumSeams + OtherNumSeams);
+		for (int32 SeamIndex = 0; SeamIndex < OtherNumSeams; ++SeamIndex)
+		{
+			GetSeam(SeamIndex + StartNumSeams).Initialize(Other.GetSeam(SeamIndex), StartNumSimVertices2D, StartNumSimVertices3D);
+		}
+
+		// Sim Vertices 3D Group (lookups)
+		FClothCollection::CopyArrayViewDataAndApplyOffset(
+			GetClothCollection()->GetElements(GetClothCollection()->GetSimVertex2DLookup()).Right(OtherNumSimVertices3D), 
+			Other.GetSimVertex2DLookup(), StartNumSimVertices2D);
+		FClothCollection::CopyArrayViewDataAndApplyOffset(
+			GetClothCollection()->GetElements(GetClothCollection()->GetSeamStitchLookup()).Right(OtherNumSimVertices3D), 
+			Other.GetSeamStitchLookup(), StartNumSeams);
+
+		// Render Patterns Group
+		const int32 StartNumRenderPatterns = GetNumRenderPatterns();
+		const int32 OtherNumRenderPatterns = Other.GetNumRenderPatterns();
+		SetNumRenderPatterns(StartNumRenderPatterns + OtherNumRenderPatterns);
+		for (int32 PatternIndex = 0; PatternIndex < OtherNumRenderPatterns; ++PatternIndex)
+		{
+			GetRenderPattern(PatternIndex + StartNumRenderPatterns).Initialize(Other.GetRenderPattern(PatternIndex));
+		}
+
+		// Weight maps
+		const TArray<FName> WeightMapNames = Other.GetWeightMapNames();
+		for (const FName& WeightMapName : WeightMapNames)
+		{
+			AddWeightMap(WeightMapName);
+			FClothCollection::CopyArrayViewData(GetWeightMap(WeightMapName).Right(OtherNumSimVertices3D), Other.GetWeightMap(WeightMapName));
+		}
 	}
 
-	void FCollectionClothFacade::PostSerialize(const ::Chaos::FChaosArchive& Ar)
+	void FCollectionClothFacade::SetPhysicsAssetPathName(const FString& PathName)
 	{
-		if (Ar.CustomVer(FUE5MainStreamObjectVersion::GUID) < FUE5MainStreamObjectVersion::ClothCollectionTetherInitialization)
+		if (ClothCollection->GetNumElements(FClothCollection::LodsGroup))
 		{
-			const int32 NumLods = GetNumLods();
-			for (int32 LodIndex = 0; LodIndex < NumLods; ++LodIndex)
-			{
-				FCollectionClothLodFacade LodFacade = GetLod(LodIndex);
-				const int32 NumTetherBatches = LodFacade.GetNumTetherBatches();
-				for (int32 TetherBatchIndex = 0; TetherBatchIndex < NumTetherBatches; ++TetherBatchIndex)
-				{
-					// The tethers were never set to be empty, so do this manually.
-					const int32 TetherBatchElementIndex = (*ClothCollection->GetTetherBatchStart())[LodIndex] + TetherBatchIndex;
-
-					(*GetClothCollection()->GetTetherStart())[TetherBatchElementIndex] = INDEX_NONE;
-					(*GetClothCollection()->GetTetherEnd())[TetherBatchElementIndex] = INDEX_NONE;
-				}
-
-				// Clear any (now empty) tether batches
-				LodFacade.SetNumTetherBatches(0);
-			}
+			(*GetClothCollection()->GetPhysicsAssetPathName())[0] = PathName;
+		}
+
+	}
+	void FCollectionClothFacade::SetSkeletonAssetPathName(const FString& PathName)
+	{
+		if (ClothCollection->GetNumElements(FClothCollection::LodsGroup))
+		{
+			(*GetClothCollection()->GetSkeletonAssetPathName())[0] = PathName;
 		}
 	}
 
-	int32 FCollectionClothFacade::AddLod()
+	TArrayView<FVector2f> FCollectionClothFacade::GetSimPosition2D()
 	{
-		check(IsValid());
+		return GetClothCollection()->GetElements(GetClothCollection()->GetSimPosition2D());
+	}
+
+	TArrayView<int32> FCollectionClothFacade::GetSimVertex3DLookupPrivate()
+	{
+		return GetClothCollection()->GetElements(GetClothCollection()->GetSimVertex3DLookup());
+	}
+
+	TArrayView<FVector3f> FCollectionClothFacade::GetSimPosition3D()
+	{
+		return GetClothCollection()->GetElements(GetClothCollection()->GetSimPosition3D());
+	}
 
-		const int32 LodIndex = GetClothCollection()->GetNumElements(FClothCollection::LodsGroup);
-		GetClothCollection()->SetNumElements(LodIndex + 1, FClothCollection::LodsGroup);
+	TArrayView<FVector3f> FCollectionClothFacade::GetSimNormal()
+	{
+		return GetClothCollection()->GetElements(GetClothCollection()->GetSimNormal());
+	}
 
-		FCollectionClothLodFacade(GetClothCollection(), LodIndex).SetDefaults();
+	TArrayView<TArray<int32>> FCollectionClothFacade::GetSimBoneIndices()
+	{
+		return GetClothCollection()->GetElements(GetClothCollection()->GetSimBoneIndices());
+	}
 
-		return LodIndex;
+	TArrayView<TArray<float>> FCollectionClothFacade::GetSimBoneWeights()
+	{
+		return GetClothCollection()->GetElements(GetClothCollection()->GetSimBoneWeights());
 	}
 
-	FCollectionClothLodFacade FCollectionClothFacade::GetLod(int32 LodIndex)
+	TArrayView<TArray<int32>> FCollectionClothFacade::GetTetherKinematicIndex()
 	{
-		check(IsValid());
-		check(LodIndex < GetNumLods());
-		return FCollectionClothLodFacade(GetClothCollection(), LodIndex);
+		return GetClothCollection()->GetElements(GetClothCollection()->GetTetherKinematicIndex());
 	}
 
-	void FCollectionClothFacade::SetNumLods(int32 InNumLods)
+	TArrayView<TArray<float>> FCollectionClothFacade::GetTetherReferenceLength()
 	{
-		check(IsValid());
+		return GetClothCollection()->GetElements(GetClothCollection()->GetTetherReferenceLength());
+	}
+
+	TArrayView<TArray<int32>> FCollectionClothFacade::GetSeamStitchLookupPrivate()
+	{
+		return GetClothCollection()->GetElements(GetClothCollection()->GetSeamStitchLookup());
+	}
+
+	TArrayView<TArray<int32>> FCollectionClothFacade::GetSimVertex2DLookupPrivate()
+	{
+		return GetClothCollection()->GetElements(GetClothCollection()->GetSimVertex2DLookup());
+	}
+
+	void FCollectionClothFacade::RemoveSimVertices3D(int32 InNumSimVertices)
+	{
+		const int32 NumSimVertices = GetNumSimVertices3D();
+		check(InNumSimVertices >= NumSimVertices);
+		GetClothCollection()->SetNumElements(NumSimVertices - InNumSimVertices, FClothCollection::SimVertices3DGroup);
+	}
+
+	TArrayView<FIntVector3> FCollectionClothFacade::GetSimIndices2D()
+	{
+		return GetClothCollection()->GetElements(GetClothCollection()->GetSimIndices2D());
+	}
 
-		const int32 NumLods = GetNumLods();
+	TArrayView<FIntVector3> FCollectionClothFacade::GetSimIndices3D()
+	{
+		return GetClothCollection()->GetElements(GetClothCollection()->GetSimIndices3D());
+	}
+
+	void FCollectionClothFacade::SetNumSimPatterns(int32 InNumPatterns)
+	{
+		const int32 NumPatterns = GetNumSimPatterns();
 
-		for (int32 LodIndex = InNumLods; LodIndex < NumLods; ++LodIndex)
+		for (int32 PatternIndex = InNumPatterns; PatternIndex < NumPatterns; ++PatternIndex)
 		{
-			GetLod(LodIndex).Reset();
+			GetSimPattern(PatternIndex).Reset();
 		}
 
-		GetClothCollection()->SetNumElements(InNumLods, FClothCollection::LodsGroup);
+		GetClothCollection()->SetNumElements(InNumPatterns, FClothCollection::SimPatternsGroup);
 
-		for (int32 LodIndex = NumLods; LodIndex < InNumLods; ++LodIndex)
+		for (int32 PatternIndex = NumPatterns; PatternIndex < InNumPatterns; ++PatternIndex)
 		{
-			GetLod(LodIndex).SetDefaults();
+			GetSimPattern(PatternIndex).SetDefaults();
 		}
 	}
 
+	int32 FCollectionClothFacade::AddSimPattern()
+	{
+		const int32 PatternIndex = GetNumSimPatterns();
+		SetNumSimPatterns(PatternIndex + 1);
+		return PatternIndex;
+	}
+
+	FCollectionClothSimPatternFacade FCollectionClothFacade::GetSimPattern(int32 PatternIndex)
+	{
+		return FCollectionClothSimPatternFacade(GetClothCollection(), PatternIndex);
+	}
+
+	void FCollectionClothFacade::SetNumRenderPatterns(int32 InNumPatterns)
+	{
+		const int32 NumPatterns = GetNumRenderPatterns();
+
+		for (int32 PatternIndex = InNumPatterns; PatternIndex < NumPatterns; ++PatternIndex)
+		{
+			GetRenderPattern(PatternIndex).Reset();
+		}
+
+		GetClothCollection()->SetNumElements(InNumPatterns, FClothCollection::RenderPatternsGroup);
+
+		for (int32 PatternIndex = NumPatterns; PatternIndex < InNumPatterns; ++PatternIndex)
+		{
+			GetRenderPattern(PatternIndex).SetDefaults();
+		}
+	}
+
+	int32 FCollectionClothFacade::AddRenderPattern()
+	{
+		const int32 PatternIndex = GetNumRenderPatterns();
+		SetNumRenderPatterns(PatternIndex + 1);
+		return PatternIndex;
+	}
+
+	FCollectionClothRenderPatternFacade FCollectionClothFacade::GetRenderPattern(int32 PatternIndex)
+	{
+		return FCollectionClothRenderPatternFacade(GetClothCollection(), PatternIndex);
+	}
+
+	TArrayView<FString> FCollectionClothFacade::GetRenderMaterialPathName()
+	{
+		return GetClothCollection()->GetElements(GetClothCollection()->GetRenderMaterialPathName());
+	}
+
+	void FCollectionClothFacade::SetNumSeams(int32 InNumSeams)
+	{
+		const int32 NumSeams = GetNumSeams();
+
+		for (int32 SeamIndex = InNumSeams; SeamIndex < NumSeams; ++SeamIndex)
+		{
+			GetSeam(SeamIndex).Reset();
+		}
+
+		GetClothCollection()->SetNumElements(InNumSeams, FClothCollection::SeamsGroup);
+
+		for (int32 SeamIndex = NumSeams; SeamIndex < InNumSeams; ++SeamIndex)
+		{
+			GetSeam(SeamIndex).SetDefaults();
+		}
+	}
+
+	int32 FCollectionClothFacade::AddSeam()
+	{
+		const int32 SeamIndex = GetNumSeams();
+		SetNumSeams(SeamIndex + 1);
+		return SeamIndex;
+	}
+
+	FCollectionClothSeamFacade FCollectionClothFacade::GetSeam(int32 SeamIndex)
+	{
+		return FCollectionClothSeamFacade(GetClothCollection(), SeamIndex);
+	}
+
+	//~ Render Vertices Group
+	TArrayView<FVector3f> FCollectionClothFacade::GetRenderPosition()
+	{
+		return GetClothCollection()->GetElements(GetClothCollection()->GetRenderPosition());
+	}
+
+	TArrayView<FVector3f> FCollectionClothFacade::GetRenderNormal()
+	{
+		return GetClothCollection()->GetElements(GetClothCollection()->GetRenderNormal());
+	}
+
+	TArrayView<FVector3f> FCollectionClothFacade::GetRenderTangentU()
+	{
+		return GetClothCollection()->GetElements(GetClothCollection()->GetRenderTangentU());
+	}
+
+	TArrayView<FVector3f> FCollectionClothFacade::GetRenderTangentV()
+	{
+		return GetClothCollection()->GetElements(GetClothCollection()->GetRenderTangentV());
+	}
+
+	TArrayView<TArray<FVector2f>> FCollectionClothFacade::GetRenderUVs()
+	{
+		return GetClothCollection()->GetElements(GetClothCollection()->GetRenderUVs());
+	}
+
+	TArrayView<FLinearColor> FCollectionClothFacade::GetRenderColor()
+	{
+		return GetClothCollection()->GetElements(GetClothCollection()->GetRenderColor());
+	}
+
+	TArrayView<TArray<int32>> FCollectionClothFacade::GetRenderBoneIndices()
+	{
+		return GetClothCollection()->GetElements(GetClothCollection()->GetRenderBoneIndices());
+	}
+
+	TArrayView<TArray<float>> FCollectionClothFacade::GetRenderBoneWeights()
+	{
+		return GetClothCollection()->GetElements(GetClothCollection()->GetRenderBoneWeights());
+	}
+
+	TArrayView<FIntVector3> FCollectionClothFacade::GetRenderIndices()
+	{
+		return GetClothCollection()->GetElements(GetClothCollection()->GetRenderIndices());
+	}
+
 	void FCollectionClothFacade::AddWeightMap(const FName& Name)
 	{
 		check(IsValid());
-		GetClothCollection()->AddUserDefinedAttribute<float>(Name, FClothCollection::SimVerticesGroup);
+		GetClothCollection()->AddUserDefinedAttribute<float>(Name, FClothCollection::SimVertices3DGroup);
 	}
 
 	void FCollectionClothFacade::RemoveWeightMap(const FName& Name)
 	{
 		check(IsValid());
-		GetClothCollection()->RemoveUserDefinedAttribute(Name, FClothCollection::SimVerticesGroup);
+		GetClothCollection()->RemoveUserDefinedAttribute(Name, FClothCollection::SimVertices3DGroup);
+	}
+
+	TArrayView<float> FCollectionClothFacade::GetWeightMap(const FName& Name)
+	{
+		return GetClothCollection()->GetElements(GetClothCollection()->GetUserDefinedAttribute<float>(Name, FClothCollection::SimVertices3DGroup));
 	}
+
+	void FCollectionClothFacade::SetDefaults()
+	{
+		GetClothCollection()->SetNumElements(1, FClothCollection::LodsGroup);
+	}
+
 } // End namespace UE::Chaos::ClothAsset
diff --git a/Engine/Plugins/Experimental/ChaosClothAsset/Source/ChaosClothAsset/Public/ChaosClothAsset/ClothCollection.h b/Engine/Plugins/Experimental/ChaosClothAsset/Source/ChaosClothAsset/Public/ChaosClothAsset/ClothCollection.h
--- a/Engine/Plugins/Experimental/ChaosClothAsset/Source/ChaosClothAsset/Public/ChaosClothAsset/ClothCollection.h
+++ b/Engine/Plugins/Experimental/ChaosClothAsset/Source/ChaosClothAsset/Public/ChaosClothAsset/ClothCollection.h
@@ -22,16 +22,19 @@ namespace UE::Chaos::ClothAsset
 	class FClothCollection final
 	{
 	public:
-		static const FName LodsGroup;  // LOD information
-		static const FName MaterialsGroup;  // Materials information
-		static const FName TetherBatchesGroup;  // Tethers parallel batch processing information
-		static const FName TethersGroup;  // Tethers information
-		static const FName SeamsGroup;  // Contains pairs of stitched sim vertex indices and matching pattern indices
-		static const FName PatternsGroup;  // Contains pattern relationships to other groups
-		static const FName SimFacesGroup;  // Contains indices to sim vertex
-		static const FName SimVerticesGroup;  // Contains patterns' 2D positions, 3D draped position (rest)
+		static const FName LodsGroup;  // LOD information (only one LOD per collection)
+		static const FName SeamsGroup;  // Collection of seam stitches
+		static const FName SeamStitchesGroup;  // Contains pairs of stitched sim vertex indices
+		static const FName SimPatternsGroup;  // Contains sim pattern relationships to other groups
+		static const FName RenderPatternsGroup;  // Contains render pattern relationships to other groups
+		static const FName SimFacesGroup;  // Contains indices to sim vertices
+		static const FName SimVertices2DGroup;  // Contains 2D positions
+		static const FName SimVertices3DGroup;  // Contains 3D positions
 		static const FName RenderFacesGroup;  // Contains indices to render vertex
-		static const FName RenderVerticesGroup;  // Contains pattern's 3D render model
+		static const FName RenderVerticesGroup;  // Contains 3D render model
+
+		static constexpr int8 MaxNumBoneInfluences = 12; // This should be <= MAX_TOTAL_INFLUENCES defined in GPUSkinPublicDefs.h 
+		static constexpr int8 MaxNumTetherAttachments = 4; // This should be <= FClothTetherDataPrivate::MaxNumAttachments
 
 		explicit FClothCollection(const TSharedPtr<FManagedArrayCollection>& InManagedArrayCollection);
 		~FClothCollection() = default;
@@ -59,49 +62,56 @@ namespace UE::Chaos::ClothAsset
 		/** Set the number of elements to one of the sub groups while maintaining the correct order of the data, and return the first index of the range. */
 		int32 SetNumElements(int32 InNumElements, const FName& GroupName, TManagedArray<int32>* StartArray, TManagedArray<int32>* EndArray, int32 Index);
 
+		/** Remove Elements */
+		void RemoveElements(const FName& Group, const TArray<int32>& SortedDeletionList);
+
+		template<typename T>
+		static inline TConstArrayView<T> GetElements(const TManagedArray<T>* ElementArray, const TManagedArray<int32>* StartArray, const TManagedArray<int32>* EndArray, int32 ArrayIndex);
+
+		template<typename T>
+		static inline TArrayView<T> GetElements(TManagedArray<T>* ElementArray, const TManagedArray<int32>* StartArray, const TManagedArray<int32>* EndArray, int32 ArrayIndex);
+
 		template<typename T>
-		inline TConstArrayView<T> GetElements(const TManagedArray<T>* ElementArray, const TManagedArray<int32>* StartArray, const TManagedArray<int32>* EndArray, int32 ArrayIndex) const;
+		static inline TConstArrayView<T> GetElements(const TManagedArray<T>* ElementArray);
 
 		template<typename T>
-		inline TArrayView<T> GetElements(TManagedArray<T>* ElementArray, const TManagedArray<int32>* StartArray, const TManagedArray<int32>* EndArray, int32 ArrayIndex);
+		static inline TArrayView<T> GetElements(TManagedArray<T>* ElementArray);
 
 		/**
 		 * Return the difference between the start index of an element to the start index of the first sub-element in the group (base).
 		 * Usefull for getting back and forth between LOD/pattern indexation modes.
 		 */
-		int32 GetElementsOffset(const TManagedArray<int32>* StartArray, int32 BaseElementIndex, int32 ElementIndex) const;
+		static int32 GetElementsOffset(const TManagedArray<int32>* StartArray, int32 BaseElementIndex, int32 ElementIndex);
 
-		int32 GetNumSubElements(
+		/**
+		* Return the ArrayIndex (into StartArray/EndArray) which corresponds with the subarray that contains ElementIndex.
+		* Useful for finding which pattern contains a vertex or face index. Returns INDEX_NONE if not found.
+		* NOTE: typically we have a small number of patterns (case where this is used), so just doing a linear search since
+		* a bisection search would get confused by the empty patterns.
+		*/
+		static int32 GetArrayIndexForContainedElement(
+			const TManagedArray<int32>* StartArray,
+			const TManagedArray<int32>* EndArray,
+			int32 ElementIndex);
+
+		static int32 GetNumSubElements(
 			const TManagedArray<int32>* StartArray,
 			const TManagedArray<int32>* EndArray,
 			const TManagedArray<int32>* StartSubArray,
 			const TManagedArray<int32>* EndSubArray,
-			int32 LodIndex) const;
-
-		UE_DEPRECATED(5.3, "Use GetNumSubElements instead.")
-		int32 GetPatternsNumElements(const TManagedArray<int32>* StartSubArray, const TManagedArray<int32>* EndSubArray, int32 LodIndex) const
-		{
-			return GetNumSubElements( PatternStart, PatternEnd, StartSubArray, EndSubArray, LodIndex);
-		}
+			int32 LodIndex);
 
 		template<typename T>
-		inline TConstArrayView<T> GetSubElements(
+		static inline TConstArrayView<T> GetSubElements(
 			const TManagedArray<T>* SubElementArray,
 			const TManagedArray<int32>* StartArray,
 			const TManagedArray<int32>* EndArray,
 			const TManagedArray<int32>* StartSubArray,
 			const TManagedArray<int32>* EndSubArray,
-			int32 ArrayIndex) const;
-
-		template<typename T>
-		UE_DEPRECATED(5.3, "Use GetSubElements instead.")
-		TConstArrayView<T> GetPatternElements(const TManagedArray<T>* SubElementArray, const TManagedArray<int32>* StartSubArray, const TManagedArray<int32>* EndSubArray, int32 LodIndex) const
-		{
-			return GetSubElements(SubElementArray, PatternStart, PatternEnd, StartSubArray, EndSubArray,LodIndex);
-		}
+			int32 ArrayIndex);
 
 		template<typename T>
-		inline TArrayView<T> GetSubElements(
+		static inline TArrayView<T> GetSubElements(
 			TManagedArray<T>* SubElementArray,
 			const TManagedArray<int32>* StartArray,
 			const TManagedArray<int32>* EndArray,
@@ -109,27 +119,23 @@ namespace UE::Chaos::ClothAsset
 			const TManagedArray<int32>* EndSubArray,
 			int32 ArrayIndex);
 
-		template<typename T>
-		UE_DEPRECATED(5.3, "Use GetSubElements instead.")
-		TArrayView<T> GetPatternElements(TManagedArray<T>* SubElementArray, const TManagedArray<int32>* StartSubArray, const TManagedArray<int32>* EndSubArray, int32 LodIndex)
-		{
-			return GetSubElements(SubElementArray, PatternStart, PatternEnd, StartSubArray, EndSubArray, LodIndex);
-		}
-
 		template<bool bStart = true, bool bEnd = true>
-		TTuple<int32, int32> GetSubElementsStartEnd(
+		static TTuple<int32, int32> GetSubElementsStartEnd(
 			const TManagedArray<int32>* StartArray,
 			const TManagedArray<int32>* EndArray,
 			const TManagedArray<int32>* StartSubArray,
 			const TManagedArray<int32>* EndSubArray,
-			int32 ArrayIndex) const;
+			int32 ArrayIndex);
 
-		template<bool bStart = true, bool bEnd = true>
-		UE_DEPRECATED(5.3, "Use GetSubElementsStartEnd instead.")
-		TTuple<int32, int32> GetPatternElementsStartEnd(const TManagedArray<int32>* StartSubArray, const TManagedArray<int32>* EndSubArray, int32 LodIndex) const
-		{
-			return GetSubElementsStartEnd(PatternStart, PatternEnd, StartSubArray, EndSubArray, LodIndex);
-		}
+
+		/** Useful method for copying data between ArrayViews. Putting here for now. */
+		template<typename T>
+		static inline void CopyArrayViewData(const TArrayView<T>& To, const TConstArrayView<T>& From);
+
+		template<typename T>
+		static inline void CopyArrayViewDataAndApplyOffset(const TArrayView<T>& To, const TConstArrayView<T>& From, const T Offset);
+
+		static void CopyArrayViewDataAndApplyOffset(const TArrayView<TArray<int32>>& To, const TConstArrayView<TArray<int32>>& From, const int32 Offset);
 
 		//~ Weight maps
 		template<typename T, TEMPLATE_REQUIRES(TIsUserAttributeType<T>::Value)>
@@ -149,60 +155,51 @@ namespace UE::Chaos::ClothAsset
 		template<typename T, TEMPLATE_REQUIRES(TIsUserAttributeType<T>::Value)>
 		TManagedArray<T>* GetUserDefinedAttribute(const FName& Name, const FName& GroupName);
 
-
-		//~ LODs Group
-		const TManagedArray<int32>* GetMaterialStart() const { return MaterialStart; }
-		const TManagedArray<int32>* GetMaterialEnd() const { return MaterialEnd; }
-		const TManagedArray<int32>* GetTetherBatchStart() const { return TetherBatchStart; }
-		const TManagedArray<int32>* GetTetherBatchEnd() const { return TetherBatchEnd; }
-		const TManagedArray<int32>* GetSeamStart() const { return SeamStart; }
-		const TManagedArray<int32>* GetSeamEnd() const { return SeamEnd; }
-		const TManagedArray<int32>* GetPatternStart() const { return PatternStart; }
-		const TManagedArray<int32>* GetPatternEnd() const { return PatternEnd; }
+		//~ LODs Group (There should be only one LOD per ClothCollection)
 		const TManagedArray<FString>* GetPhysicsAssetPathName() const { return PhysicsAssetPathName; }
 		const TManagedArray<FString>* GetSkeletonAssetPathName() const { return SkeletonAssetPathName; }
 
-		//~ Materials Group
-		const TManagedArray<FString>* GetRenderMaterialPathName() const { return RenderMaterialPathName; }
-
-		//~ Tether Batches Group
-		const TManagedArray<int32>* GetTetherStart() const { return TetherStart; }
-		const TManagedArray<int32>* GetTetherEnd() const { return TetherEnd; }
-
-		//~ Tethers Group
-		const TManagedArray<int32>* GetTetherKinematicIndex() const { return TetherKinematicIndex; }
-		const TManagedArray<int32>* GetTetherDynamicIndex() const { return TetherDynamicIndex; }
-		const TManagedArray<float>* GetTetherReferenceLength() const { return TetherReferenceLength; }
-
 		//~ Seam Group
-		const TManagedArray<FIntVector2>* GetSeamPatterns() const { return SeamPatterns; }
-		const TManagedArray<TArray<FIntVector2>>* GetSeamStitches() const { return SeamStitches; }
+		const TManagedArray<int32>* GetSeamStitchStart() const { return SeamStitchStart; }
+		const TManagedArray<int32>* GetSeamStitchEnd() const { return SeamStitchEnd; }
 
-		//~ Patterns Group
-		const TManagedArray<int32>* GetSimVerticesStart() const { return SimVerticesStart; }
-		const TManagedArray<int32>* GetSimVerticesEnd() const { return SimVerticesEnd; }
+		//~ Seam Stitches Group
+		const TManagedArray<FIntVector2>* GetSeamStitch2DEndIndices() const { return SeamStitch2DEndIndices; }
+		const TManagedArray<int32>* GetSeamStitch3DIndex() const { return SeamStitch3DIndex; }
+
+		//~ Sim Patterns Group
+		const TManagedArray<int32>* GetSimVertices2DStart() const { return SimVertices2DStart; }
+		const TManagedArray<int32>* GetSimVertices2DEnd() const { return SimVertices2DEnd; }
 		const TManagedArray<int32>* GetSimFacesStart() const { return SimFacesStart; }
 		const TManagedArray<int32>* GetSimFacesEnd() const { return SimFacesEnd; }
+
+		//~ Render Patterns Group
 		const TManagedArray<int32>* GetRenderVerticesStart() const { return RenderVerticesStart; }
 		const TManagedArray<int32>* GetRenderVerticesEnd() const { return RenderVerticesEnd; }
 		const TManagedArray<int32>* GetRenderFacesStart() const { return RenderFacesStart; }
 		const TManagedArray<int32>* GetRenderFacesEnd() const { return RenderFacesEnd; }
-		const TManagedArray<int32>* GetStatusFlags() const { return StatusFlags; }
+		const TManagedArray<FString>* GetRenderMaterialPathName() const { return RenderMaterialPathName; }
 
 		//~ Sim Faces Group
-		const TManagedArray<FIntVector3>* GetSimIndices() const { return SimIndices; }
+		const TManagedArray<FIntVector3>* GetSimIndices2D() const { return SimIndices2D; }
+		const TManagedArray<FIntVector3>* GetSimIndices3D() const { return SimIndices3D; }
 
-		//~ Sim Vertices Group
-		const TManagedArray<FVector2f>* GetSimPosition() const { return SimPosition; }
-		const TManagedArray<FVector3f>* GetSimRestPosition() const { return SimRestPosition; }
-		const TManagedArray<FVector3f>* GetSimRestNormal() const { return SimRestNormal; }
-		const TManagedArray<int32>* GetSimNumBoneInfluences() const { return SimNumBoneInfluences; }
+		//~ Sim Vertices 2D Group
+		const TManagedArray<FVector2f>* GetSimPosition2D() const { return SimPosition2D; }
+		const TManagedArray<int32>* GetSimVertex3DLookup() const { return SimVertex3DLookup; }
+
+		//~ Sim Vertices 3D Group
+		const TManagedArray<FVector3f>* GetSimPosition3D() const { return SimPosition3D; }
+		const TManagedArray<FVector3f>* GetSimNormal() const { return SimNormal; }
 		const TManagedArray<TArray<int32>>* GetSimBoneIndices() const { return SimBoneIndices; }
 		const TManagedArray<TArray<float>>* GetSimBoneWeights() const { return SimBoneWeights; }
+		const TManagedArray<TArray<int32>>* GetTetherKinematicIndex() const { return TetherKinematicIndex; }
+		const TManagedArray<TArray<float>>* GetTetherReferenceLength() const { return TetherReferenceLength; }
+		const TManagedArray<TArray<int32>>* GetSimVertex2DLookup() const { return SimVertex2DLookup; }
+		const TManagedArray<TArray<int32>>* GetSeamStitchLookup() const { return SeamStitchLookup; }
 
 		//~ Render Faces Group
 		const TManagedArray<FIntVector3>* GetRenderIndices() const { return RenderIndices; }
-		const TManagedArray<int32>* GetRenderMaterialIndex() const { return RenderMaterialIndex; }
 
 		//~ Render Vertices Group
 		const TManagedArray<FVector3f>* GetRenderPosition() const { return RenderPosition; }
@@ -211,63 +208,54 @@ namespace UE::Chaos::ClothAsset
 		const TManagedArray<FVector3f>* GetRenderTangentV() const { return RenderTangentV; }
 		const TManagedArray<TArray<FVector2f>>* GetRenderUVs() const { return RenderUVs; }
 		const TManagedArray<FLinearColor>* GetRenderColor() const { return RenderColor; }
-		const TManagedArray<int32>* GetRenderNumBoneInfluences() const { return RenderNumBoneInfluences; }
 		const TManagedArray<TArray<int32>>* GetRenderBoneIndices() const { return RenderBoneIndices; }
 		const TManagedArray<TArray<float>>* GetRenderBoneWeights() const { return RenderBoneWeights; }
 
-		//~ LODs Group
-		TManagedArray<int32>* GetMaterialStart() { return MaterialStart; }
-		TManagedArray<int32>* GetMaterialEnd() { return MaterialEnd; }
-		TManagedArray<int32>* GetTetherBatchStart() { return TetherBatchStart; }
-		TManagedArray<int32>* GetTetherBatchEnd() { return TetherBatchEnd; }
-		TManagedArray<int32>* GetSeamStart() { return SeamStart; }
-		TManagedArray<int32>* GetSeamEnd() { return SeamEnd; }
-		TManagedArray<int32>* GetPatternStart() { return PatternStart; }
-		TManagedArray<int32>* GetPatternEnd() { return PatternEnd; }
-		TManagedArray<FString>* GetPhysicsAssetPathName() { return PhysicsAssetPathName; }
-		TManagedArray<FString>* GetSkeletonAssetPathName() { return SkeletonAssetPathName; }
-
-		//~ Materials Group
-		TManagedArray<FString>* GetRenderMaterialPathName() { return RenderMaterialPathName; }
-
-		//~ Tether Batches Group
-		TManagedArray<int32>* GetTetherStart() { return TetherStart; }
-		TManagedArray<int32>* GetTetherEnd() { return TetherEnd; }
-
-		//~ Tethers Group
-		TManagedArray<int32>* GetTetherKinematicIndex() { return TetherKinematicIndex; }
-		TManagedArray<int32>* GetTetherDynamicIndex() { return TetherDynamicIndex; }
-		TManagedArray<float>* GetTetherReferenceLength() { return TetherReferenceLength; }
+		//~ LODs Group (There should be only one LOD per ClothCollection)
+		TManagedArray<FString>* GetPhysicsAssetPathName(){ return PhysicsAssetPathName; }
+		TManagedArray<FString>* GetSkeletonAssetPathName(){ return SkeletonAssetPathName; }
 
 		//~ Seam Group
-		TManagedArray<FIntVector2>* GetSeamPatterns() { return SeamPatterns; }
-		TManagedArray<TArray<FIntVector2>>* GetSeamStitches() { return SeamStitches; }
+		TManagedArray<int32>* GetSeamStitchStart() { return SeamStitchStart; }
+		TManagedArray<int32>* GetSeamStitchEnd() { return SeamStitchEnd; }
 
-		//~ Patterns Group
-		TManagedArray<int32>* GetSimVerticesStart() { return SimVerticesStart; }
-		TManagedArray<int32>* GetSimVerticesEnd() { return SimVerticesEnd; }
+		//~ Seam Stitches Group
+		TManagedArray<FIntVector2>* GetSeamStitch2DEndIndices() { return SeamStitch2DEndIndices; }
+		TManagedArray<int32>* GetSeamStitch3DIndex() { return SeamStitch3DIndex; }
+
+		//~ Sim Patterns Group
+		TManagedArray<int32>* GetSimVertices2DStart() { return SimVertices2DStart; }
+		TManagedArray<int32>* GetSimVertices2DEnd() { return SimVertices2DEnd; }
 		TManagedArray<int32>* GetSimFacesStart() { return SimFacesStart; }
 		TManagedArray<int32>* GetSimFacesEnd() { return SimFacesEnd; }
+
+		//~ Render Patterns Group
 		TManagedArray<int32>* GetRenderVerticesStart() { return RenderVerticesStart; }
 		TManagedArray<int32>* GetRenderVerticesEnd() { return RenderVerticesEnd; }
 		TManagedArray<int32>* GetRenderFacesStart() { return RenderFacesStart; }
 		TManagedArray<int32>* GetRenderFacesEnd() { return RenderFacesEnd; }
-		TManagedArray<int32>* GetStatusFlags() { return StatusFlags; }
+		TManagedArray<FString>* GetRenderMaterialPathName() { return RenderMaterialPathName; }
 
 		//~ Sim Faces Group
-		TManagedArray<FIntVector3>* GetSimIndices() { return SimIndices; }
+		TManagedArray<FIntVector3>* GetSimIndices2D() { return SimIndices2D; }
+		TManagedArray<FIntVector3>* GetSimIndices3D() { return SimIndices3D; }
 
-		//~ Sim Vertices Group
-		TManagedArray<FVector2f>* GetSimPosition() { return SimPosition; }
-		TManagedArray<FVector3f>* GetSimRestPosition() { return SimRestPosition; }
-		TManagedArray<FVector3f>* GetSimRestNormal() { return SimRestNormal; }
-		TManagedArray<int32>* GetSimNumBoneInfluences() { return SimNumBoneInfluences; }
+		//~ Sim Vertices 2D Group
+		TManagedArray<FVector2f>* GetSimPosition2D() { return SimPosition2D; }
+		TManagedArray<int32>* GetSimVertex3DLookup() { return SimVertex3DLookup; }
+
+		//~ Sim Vertices 3D Group
+		TManagedArray<FVector3f>* GetSimPosition3D() { return SimPosition3D; }
+		TManagedArray<FVector3f>* GetSimNormal() { return SimNormal; }
 		TManagedArray<TArray<int32>>* GetSimBoneIndices() { return SimBoneIndices; }
 		TManagedArray<TArray<float>>* GetSimBoneWeights() { return SimBoneWeights; }
+		TManagedArray<TArray<int32>>* GetTetherKinematicIndex() { return TetherKinematicIndex; }
+		TManagedArray<TArray<float>>* GetTetherReferenceLength() { return TetherReferenceLength; }
+		TManagedArray<TArray<int32>>* GetSimVertex2DLookup() { return SimVertex2DLookup; }
+		TManagedArray<TArray<int32>>* GetSeamStitchLookup() { return SeamStitchLookup; }
 
 		//~ Render Faces Group
 		TManagedArray<FIntVector3>* GetRenderIndices() { return RenderIndices; }
-		TManagedArray<int32>* GetRenderMaterialIndex() { return RenderMaterialIndex; }
 
 		//~ Render Vertices Group
 		TManagedArray<FVector3f>* GetRenderPosition() { return RenderPosition; }
@@ -276,7 +264,6 @@ namespace UE::Chaos::ClothAsset
 		TManagedArray<FVector3f>* GetRenderTangentV() { return RenderTangentV; }
 		TManagedArray<TArray<FVector2f>>* GetRenderUVs() { return RenderUVs; }
 		TManagedArray<FLinearColor>* GetRenderColor() { return RenderColor; }
-		TManagedArray<int32>* GetRenderNumBoneInfluences() { return RenderNumBoneInfluences; }
 		TManagedArray<TArray<int32>>* GetRenderBoneIndices() { return RenderBoneIndices; }
 		TManagedArray<TArray<float>>* GetRenderBoneWeights() { return RenderBoneWeights; }
 
@@ -285,58 +272,50 @@ namespace UE::Chaos::ClothAsset
 		TSharedPtr<FManagedArrayCollection> ManagedArrayCollection;
 
 		//~ LODs Group
-		TManagedArray<int32>* MaterialStart;
-		TManagedArray<int32>* MaterialEnd;
-		TManagedArray<int32>* TetherBatchStart;
-		TManagedArray<int32>* TetherBatchEnd;
-		TManagedArray<int32>* SeamStart;
-		TManagedArray<int32>* SeamEnd;
-		TManagedArray<int32>* PatternStart;
-		TManagedArray<int32>* PatternEnd;
 		TManagedArray<FString>* PhysicsAssetPathName;
 		TManagedArray<FString>* SkeletonAssetPathName;
 
-		//~ Materials Group
-		TManagedArray<FString>* RenderMaterialPathName;
-
-		//~ Tether Batches Group
-		TManagedArray<int32>* TetherStart;
-		TManagedArray<int32>* TetherEnd;
-
-		//~ Tethers Group
-		TManagedArray<int32>* TetherKinematicIndex;
-		TManagedArray<int32>* TetherDynamicIndex;
-		TManagedArray<float>* TetherReferenceLength;
-
 		//~ Seam Group
-		TManagedArray<FIntVector2>* SeamPatterns;  // Stitched pattern indices (LOD based indices) for SeamIndicesFirst
-		TManagedArray<TArray<FIntVector2>>* SeamStitches;  // Stitched vertex indices pair (LOD based indices)
+		TManagedArray<int32>* SeamStitchStart;
+		TManagedArray<int32>* SeamStitchEnd;
 
-		//~ Patterns Group
-		TManagedArray<int32>* SimVerticesStart;
-		TManagedArray<int32>* SimVerticesEnd;
+		//~ Seam Stitches Group
+		TManagedArray<FIntVector2>* SeamStitch2DEndIndices;  // Stitched 2D vertex indices pair
+		TManagedArray<int32>* SeamStitch3DIndex;  // Corresponding stitched 3D vertex
+
+		//~ Sim Patterns Group
+		TManagedArray<int32>* SimVertices2DStart;
+		TManagedArray<int32>* SimVertices2DEnd;
 		TManagedArray<int32>* SimFacesStart;
 		TManagedArray<int32>* SimFacesEnd;
+
+		//~ Render Patterns Group
 		TManagedArray<int32>* RenderVerticesStart;
 		TManagedArray<int32>* RenderVerticesEnd;
 		TManagedArray<int32>* RenderFacesStart;
 		TManagedArray<int32>* RenderFacesEnd;
-		TManagedArray<int32>* StatusFlags;  // Whether this pattern should be rendered, simulated, ...etc
+		TManagedArray<FString>* RenderMaterialPathName;
 
 		//~ Sim Faces Group
-		TManagedArray<FIntVector3>* SimIndices;  // The indices point to the elements in the Sim Vertices arrays but don't include the LOD start offset
+		TManagedArray<FIntVector3>* SimIndices2D;
+		TManagedArray<FIntVector3>* SimIndices3D;
 
-		//~ Sim Vertices Group
-		TManagedArray<FVector2f>* SimPosition;
-		TManagedArray<FVector3f>* SimRestPosition;
-		TManagedArray<FVector3f>* SimRestNormal;  // Used for capture, maxdistance, backstop authoring ...etc
-		TManagedArray<int32>* SimNumBoneInfluences;
+		//~ Sim Vertices 2D Group
+		TManagedArray<FVector2f>* SimPosition2D;
+		TManagedArray<int32>* SimVertex3DLookup; // Lookup into corresponding 3D vertices
+
+		//~ Sim Vertices 3D Group
+		TManagedArray<FVector3f>* SimPosition3D;
+		TManagedArray<FVector3f>* SimNormal;  // Used for capture, maxdistance, backstop authoring ...etc
 		TManagedArray<TArray<int32>>* SimBoneIndices;
 		TManagedArray<TArray<float>>* SimBoneWeights;
+		TManagedArray<TArray<int32>>* TetherKinematicIndex;
+		TManagedArray<TArray<float>>* TetherReferenceLength;
+		TManagedArray<TArray<int32>>* SimVertex2DLookup; // Lookup into corresponding 2D vertices
+		TManagedArray<TArray<int32>>* SeamStitchLookup; // Lookup into any seam stitches which weld this vertex
 
 		//~ Render Faces Group
-		TManagedArray<FIntVector3>* RenderIndices;  // The indices point to the elements in the Render Vertices arrays but don't include the LOD start offset
-		TManagedArray<int32>* RenderMaterialIndex;  // Render material per triangle
+		TManagedArray<FIntVector3>* RenderIndices;
 
 		//~ Render Vertices Group
 		TManagedArray<FVector3f>* RenderPosition;
@@ -345,13 +324,12 @@ namespace UE::Chaos::ClothAsset
 		TManagedArray<FVector3f>* RenderTangentV;
 		TManagedArray<TArray<FVector2f>>* RenderUVs;
 		TManagedArray<FLinearColor>* RenderColor;
-		TManagedArray<int32>* RenderNumBoneInfluences;
 		TManagedArray<TArray<int32>>* RenderBoneIndices;
 		TManagedArray<TArray<float>>* RenderBoneWeights;
 	};
 
 	template<typename T>
-	inline TConstArrayView<T> FClothCollection::GetElements(const TManagedArray<T>* ElementArray, const TManagedArray<int32>* StartArray, const TManagedArray<int32>* EndArray, int32 ArrayIndex) const
+	inline TConstArrayView<T> FClothCollection::GetElements(const TManagedArray<T>* ElementArray, const TManagedArray<int32>* StartArray, const TManagedArray<int32>* EndArray, int32 ArrayIndex)
 	{
 		if (StartArray && EndArray && ElementArray)
 		{
@@ -382,6 +360,26 @@ namespace UE::Chaos::ClothAsset
 		return TArrayView<T>();
 	}
 
+	template<typename T>
+	inline TConstArrayView<T> FClothCollection::GetElements(const TManagedArray<T>* ElementArray)
+	{
+		if (ElementArray)
+		{
+			return TConstArrayView<T>(ElementArray->GetData(), ElementArray->Num());
+		}
+		return TConstArrayView<T>();
+	}
+
+	template<typename T>
+	inline TArrayView<T> FClothCollection::GetElements(TManagedArray<T>* ElementArray)
+	{
+		if (ElementArray)
+		{
+			return TArrayView<T>(ElementArray->GetData(), ElementArray->Num());
+		}
+		return TArrayView<T>();
+	}
+
 	template<typename T>
 	inline TConstArrayView<T> FClothCollection::GetSubElements(
 		const TManagedArray<T>* SubElementArray,
@@ -389,7 +387,7 @@ namespace UE::Chaos::ClothAsset
 		const TManagedArray<int32>* EndArray,
 		const TManagedArray<int32>* StartSubArray,
 		const TManagedArray<int32>* EndSubArray,
-		int32 ArrayIndex) const
+		int32 ArrayIndex)
 	{
 		if (StartArray && EndArray && SubElementArray)
 		{
@@ -427,4 +425,24 @@ namespace UE::Chaos::ClothAsset
 		}
 		return TArrayView<T>();
 	}
+
+	template<typename T>
+	inline void FClothCollection::CopyArrayViewData(const TArrayView<T>& To, const TConstArrayView<T>& From)
+	{
+		check(To.Num() == From.Num());
+		for (int32 Index = 0; Index < To.Num(); ++Index)
+		{
+			To[Index] = From[Index];
+		}
+	}
+
+	template<typename T>
+	inline void FClothCollection::CopyArrayViewDataAndApplyOffset(const TArrayView<T>& To, const TConstArrayView<T>& From, const T Offset)
+	{
+		check(To.Num() == From.Num());
+		for (int32 Index = 0; Index < To.Num(); ++Index)
+		{
+			To[Index] = From[Index] + Offset;
+		}
+	}
 }  // End namespace UE::Chaos::ClothAsset
diff --git a/Engine/Plugins/Experimental/ChaosClothAsset/Source/ChaosClothAsset/Public/ChaosClothAsset/ClothGeometryTools.h b/Engine/Plugins/Experimental/ChaosClothAsset/Source/ChaosClothAsset/Public/ChaosClothAsset/ClothGeometryTools.h
--- a/Engine/Plugins/Experimental/ChaosClothAsset/Source/ChaosClothAsset/Public/ChaosClothAsset/ClothGeometryTools.h
+++ b/Engine/Plugins/Experimental/ChaosClothAsset/Source/ChaosClothAsset/Public/ChaosClothAsset/ClothGeometryTools.h
@@ -5,6 +5,7 @@
 #include "Templates/SharedPointer.h"
 
 struct FManagedArrayCollection;
+namespace UE::Geometry { class FDynamicMesh3; }
 
 namespace UE::Chaos::ClothAsset
 {
@@ -13,23 +14,24 @@ namespace UE::Chaos::ClothAsset
 	 */
 	struct CHAOSCLOTHASSET_API FClothGeometryTools
 	{
+
 		/** Return whether at least one pattern of this collection has any faces to simulate. */
 		static bool HasSimMesh(const TSharedPtr<const FManagedArrayCollection>& ClothCollection);
 
 		/** Return whether at least one pattern of this collection has any faces to render. */
 		static bool HasRenderMesh(const TSharedPtr<const FManagedArrayCollection>& ClothCollection);
 
-		/** Set the skeleton path name on each LOD of the specified collection. */
-		static void SetSkeletonAssetPathName(const TSharedPtr<FManagedArrayCollection>& ClothCollection, const FString& SkeletonAssetPathName);
-
-		/** Set the physics path name on each LOD of the specified collection. */
-		static void SetPhysicsAssetPathName(const TSharedPtr<FManagedArrayCollection>& ClothCollection, const FString& PhysicsAssetPathName);
-
 		/** Delete the render mesh data. */
 		static void DeleteRenderMesh(const TSharedPtr<FManagedArrayCollection>& ClothCollection);
 
+		/** Delete the sim mesh data. */
+		static void DeleteSimMesh(const TSharedPtr<FManagedArrayCollection>& ClothCollection);
+
+		/** Remove all tethers. */
+		static void DeleteTethers(const TSharedPtr<FManagedArrayCollection>& ClothCollection);
+
 		/** Turn the sim mesh portion of this ClothCollection into a render mesh. */
-		static void CopySimMeshToRenderMesh(const TSharedPtr<FManagedArrayCollection>& ClothCollection, const FString& RenderMaterialPathName);
+		static void CopySimMeshToRenderMesh(const TSharedPtr<FManagedArrayCollection>& ClothCollection, const FString& RenderMaterialPathName, bool bSingleRenderPattern);
 
 		/** Reverse the mesh normals. Will reverse all normals if pattern selection is empty. */
 		static void ReverseMesh(
@@ -38,7 +40,8 @@ namespace UE::Chaos::ClothAsset
 			bool bReverseSimMeshWindingOrder,
 			bool bReverseRenderMeshNormals,
 			bool bReverseRenderMeshWindingOrder,
-			const TArray<int32>& PatternSelection);
+			const TArray<int32>& SimPatternSelection,
+			const TArray<int32>& RenderPatternSelection);
 
 		/**
 		 * Set the skinning weights for all of the sim/render vertices in ClothCollection to be bound to the root node.
@@ -48,16 +51,13 @@ namespace UE::Chaos::ClothAsset
 		static void BindMeshToRootBone(
 			const TSharedPtr<FManagedArrayCollection>& ClothCollection,
 			bool bBindSimMesh,
-			bool bBindRenderMesh,
-			const TArray<int32> Lods = {});
+			bool bBindRenderMesh);
 
 		/**
-		* Generate a welded weight map.
+		* Unwrap and build SimMesh data from a DynamicMesh
 		*/
-		static TArray<float> BuildWeldedWeightMapForLod(
-			const TSharedPtr<const FManagedArrayCollection>& ClothCollection,
-			int32 LodIndex,
-			const FName& WeightMapName,
-			const TArray<TArray<int32>>& WeldedToUnweldedMap);
+		static void BuildSimMeshFromDynamicMesh(
+			const TSharedPtr<FManagedArrayCollection>& ClothCollection,
+			const UE::Geometry::FDynamicMesh3& DynamicMesh, int32 UVChannelIndex, const FVector2f& UVScale);
 	};
 }  // End namespace UE::Chaos::ClothAsset
\ No newline at end of file
diff --git a/Engine/Plugins/Experimental/ChaosClothAsset/Source/ChaosClothAsset/Public/ChaosClothAsset/CollectionClothFacade.h b/Engine/Plugins/Experimental/ChaosClothAsset/Source/ChaosClothAsset/Public/ChaosClothAsset/CollectionClothFacade.h
--- a/Engine/Plugins/Experimental/ChaosClothAsset/Source/ChaosClothAsset/Public/ChaosClothAsset/CollectionClothFacade.h
+++ b/Engine/Plugins/Experimental/ChaosClothAsset/Source/ChaosClothAsset/Public/ChaosClothAsset/CollectionClothFacade.h
@@ -2,7 +2,9 @@
 
 #pragma once
 
-#include "ChaosClothAsset/CollectionClothLodFacade.h"
+#include "ChaosClothAsset/CollectionClothSimPatternFacade.h"
+#include "ChaosClothAsset/CollectionClothRenderPatternFacade.h"
+#include "ChaosClothAsset/CollectionClothSeamFacade.h"
 
 class FClothCollection;
 namespace Chaos
@@ -34,11 +36,80 @@ namespace UE::Chaos::ClothAsset
 		/** Return whether the facade is defined on the collection. */
 		bool IsValid() const;
 
-		/** Return the specified LOD. */
-		FCollectionClothLodConstFacade GetLod(int32 LodIndex) const;
+		//~ LOD (single per collection) Group
+		/** Return the physics asset path names used for this collection. */
+		const FString& GetPhysicsAssetPathName() const;
+		/** Return the skeleton asset path names used for this collection. */
+		const FString& GetSkeletonAssetPathName() const;
 
-		/** Return the number of LODs contained in this Cloth. */
-		int32 GetNumLods() const;
+		//~ Sim Vertices 2D Group
+		/** Return the total number of 2D simulation vertices for this collection. */
+		int32 GetNumSimVertices2D() const;
+		TConstArrayView<FVector2f> GetSimPosition2D() const;
+		TConstArrayView<int32> GetSimVertex3DLookup() const;
+
+		//~ Sim Vertices 3D Group
+		/** Return the total number of 3D simulation vertices for this collection. */
+		int32 GetNumSimVertices3D() const;
+		TConstArrayView<FVector3f> GetSimPosition3D() const;
+		TConstArrayView<FVector3f> GetSimNormal() const;
+		TConstArrayView<TArray<int32>> GetSimBoneIndices() const;
+		TConstArrayView<TArray<float>> GetSimBoneWeights() const;
+		TConstArrayView<TArray<int32>> GetTetherKinematicIndex() const;
+		TConstArrayView<TArray<float>> GetTetherReferenceLength() const;
+		TConstArrayView<TArray<int32>> GetSimVertex2DLookup() const;
+		TConstArrayView<TArray<int32>> GetSeamStitchLookup() const;
+
+		//~ Sim Faces Group
+		/** Return the total number of simulation faces for this collection across all patterns. */
+		int32 GetNumSimFaces() const;
+		TConstArrayView<FIntVector3> GetSimIndices2D() const;
+		TConstArrayView<FIntVector3> GetSimIndices3D() const;
+
+		//~ Sim Patterns Group
+		/** Return the number of patterns in this collection. */
+		int32 GetNumSimPatterns() const;
+		/** Return a pattern facade for the specified pattern index. */
+		FCollectionClothSimPatternConstFacade GetSimPattern(int32 PatternIndex) const;
+		/** Convenience to find which sim pattern a 2D vertex belongs to */
+		int32 FindSimPatternByVertex2D(int32 Vertex2DIndex) const;
+		/** Convenience to find which sim pattern a sim face belongs to */
+		int32 FindSimPatternByFaceIndex(int32 FaceIndex) const;
+
+
+		//~ Render Patterns Group
+		/** Return the number of patterns in this collection. */
+		int32 GetNumRenderPatterns() const;
+		/** Return a pattern facade for the specified pattern index. */
+		FCollectionClothRenderPatternConstFacade GetRenderPattern(int32 PatternIndex) const;
+		/** Return a view of all the render materials used on this collection across all patterns. */
+		TConstArrayView<FString> GetRenderMaterialPathName() const;
+		/** Convenience to find which render pattern a render vertex belongs to */
+		int32 FindRenderPatternByVertex(int32 VertexIndex) const;
+		/** Convenience to find which render pattern a render face belongs to */
+		int32 FindRenderPatternByFaceIndex(int32 FaceIndex) const;
+
+		//~ Seam Group
+		/** Return the number of seams in this collection. */
+		int32 GetNumSeams() const;
+		/** Return a seam facade for the specified pattern index. */
+		FCollectionClothSeamConstFacade GetSeam(int32 SeamIndex) const;
+
+		//~ Render Vertices Group
+		/** Return the total number of render vertices for this collection. */
+		int32 GetNumRenderVertices() const;
+		TConstArrayView<FVector3f> GetRenderPosition() const;
+		TConstArrayView<FVector3f> GetRenderNormal() const;
+		TConstArrayView<FVector3f> GetRenderTangentU() const;
+		TConstArrayView<FVector3f> GetRenderTangentV() const;
+		TConstArrayView<TArray<FVector2f>> GetRenderUVs() const;
+		TConstArrayView<FLinearColor> GetRenderColor() const;
+		TConstArrayView<TArray<int32>> GetRenderBoneIndices() const;
+		TConstArrayView<TArray<float>> GetRenderBoneWeights() const;
+
+		//~ Render Faces Group
+		int32 GetNumRenderFaces() const;
+		TConstArrayView<FIntVector3> GetRenderIndices() const;
 
 		/** Return whether this cloth collection has the specified weight map. */
 		bool HasWeightMap(const FName& Name) const;
@@ -46,8 +117,17 @@ namespace UE::Chaos::ClothAsset
 		/** Return the name of all user weight maps on this cloth collection. */
 		TArray<FName> GetWeightMapNames() const;
 
+		TConstArrayView<float> GetWeightMap(const FName& Name) const;
+
+		void BuildSimulationMesh(TArray<FVector3f>& Positions, TArray<FVector3f>& Normals, TArray<uint32>& Indices, TArray<FVector2f>& PatternsPositions, TArray<uint32>& PatternsIndices, 
+			TArray<uint32>& PatternToWeldedIndices, TArray<TArray<int32>>* OptionalWeldedToPatternIndices = nullptr) const;
+
 	protected:
 		TSharedPtr<const FClothCollection> ClothCollection;
+
+		friend class FCollectionClothSeamFacade;
+		friend class FCollectionClothSeamConstFacade;
+		explicit FCollectionClothConstFacade(const TSharedPtr<const FClothCollection>& ClothCollection);
 	};
 
 	/**
@@ -74,20 +154,88 @@ namespace UE::Chaos::ClothAsset
 		/** Remove all LODs from this cloth. */
 		void Reset();
 
-		/** Post serialize function. Use to upgrade between versions.*/
-		void PostSerialize(const ::Chaos::FChaosArchive& Ar);
+		/** Initialize the cloth using another cloth collection. */
+		void Initialize(const FCollectionClothConstFacade& Other);
+
+		/** Append data from another cloth collection. */
+		void Append(const FCollectionClothConstFacade& Other);
+
+		/** Append to this cloth another cloth. */
+		//void Append(const FCollectionClothConstFacade& Other);
 
-		/** Add a new LOD to this cloth. */
-		int32 AddLod();
+		//~ LOD (single per collection) Group
+		/** Set the physics asset path name. */
+		void SetPhysicsAssetPathName(const FString& PathName);
+		/** Set the skeleton asset path name. */
+		void SetSkeletonAssetPathName(const FString& PathName);
 
-		/** Return the specified LOD. */
-		FCollectionClothLodFacade GetLod(int32 LodIndex);
+		//~ Pattern Sim Vertices 2D Group
+		/** SetNumSimVertices2D per pattern within pattern facade. */
+		TArrayView<FVector2f> GetSimPosition2D();
 
-		/** Add a new LOD to this cloth, and return the cloth LOD facade set to its index. */
-		FCollectionClothLodFacade AddGetLod() { return GetLod(AddLod()); }
+		//~ Pattern Sim Vertices 3D Group
+		TArrayView<FVector3f> GetSimPosition3D();
+		TArrayView<FVector3f> GetSimNormal();
+		TArrayView<TArray<int32>> GetSimBoneIndices();
+		TArrayView<TArray<float>> GetSimBoneWeights();
+		TArrayView<TArray<int32>> GetTetherKinematicIndex();
+		TArrayView<TArray<float>> GetTetherReferenceLength();
 
-		/** Set a new number of LODs for this cloth. */
-		void SetNumLods(int32 NumLods);
+		/** This will remove the 3D vertices, but the associated seams and 2D vertices will still exist, and point to INDEX_NONE */
+		void RemoveSimVertices3D(int32 NumSimVertices);
+		void RemoveAllSimVertices3D() { RemoveSimVertices3D(GetNumSimVertices3D()); }
+
+		//~ Pattern Sim Faces Group
+		/** SetNumSimFaces per pattern within pattern facade. */
+		TArrayView<FIntVector3> GetSimIndices2D();
+		TArrayView<FIntVector3> GetSimIndices3D();
+
+		//~ Sim Patterns Group
+		/** Set the new number of patterns to this cloth LOD. */
+		void SetNumSimPatterns(int32 NumPatterns);
+		/** Add a new pattern to this cloth LOD and return its index in the LOD pattern list. */
+		int32 AddSimPattern();
+		/** Return a pattern facade for the specified pattern index. */
+		FCollectionClothSimPatternFacade GetSimPattern(int32 PatternIndex);
+		/** Add a new pattern to this cloth LOD, and return the cloth pattern facade set to its index. */
+		FCollectionClothSimPatternFacade AddGetSimPattern() { return GetSimPattern(AddSimPattern()); }
+
+		//~ Render Patterns Group
+		/** Set the new number of patterns to this cloth LOD. */
+		void SetNumRenderPatterns(int32 NumPatterns);
+		/** Add a new pattern to this cloth LOD and return its index in the LOD pattern list. */
+		int32 AddRenderPattern();
+		/** Return a pattern facade for the specified pattern index. */
+		FCollectionClothRenderPatternFacade GetRenderPattern(int32 PatternIndex);
+		/** Add a new pattern to this cloth LOD, and return the cloth pattern facade set to its index. */
+		FCollectionClothRenderPatternFacade AddGetRenderPattern() { return GetRenderPattern(AddRenderPattern()); }
+		/** Return a view of all the render materials used on this collection across all patterns. */
+		TArrayView<FString> GetRenderMaterialPathName();
+
+		//~ Seam Group
+		/** Set the new number of seams to this cloth. */
+		void SetNumSeams(int32 NumSeams);
+		/** Add a new seam to this cloth and return its index in the seam list. */
+		int32 AddSeam();
+		/** Return a seam facade for the specified seam index. */
+		FCollectionClothSeamFacade GetSeam(int32 SeamIndex);
+		/** Add a new seam to this cloth and return the seam facade set to its index. */
+		FCollectionClothSeamFacade AddGetSeam() { return GetSeam(AddSeam()); }
+
+		//~ Render Vertices Group
+		/** SetNumRenderVertices per pattern within pattern facade. */
+		TArrayView<FVector3f> GetRenderPosition();
+		TArrayView<FVector3f> GetRenderNormal();
+		TArrayView<FVector3f> GetRenderTangentU();
+		TArrayView<FVector3f> GetRenderTangentV();
+		TArrayView<TArray<FVector2f>> GetRenderUVs();
+		TArrayView<FLinearColor> GetRenderColor();
+		TArrayView<TArray<int32>> GetRenderBoneIndices();
+		TArrayView<TArray<float>> GetRenderBoneWeights();
+
+		//~ Render Faces Group
+		/** SetNumRenderFaces per pattern within pattern facade. */
+		TArrayView<FIntVector3> GetRenderIndices();
 
 		/** Add a new weight map to this cloth. Access is then done per pattern. */
 		void AddWeightMap(const FName& Name);
@@ -95,7 +243,24 @@ namespace UE::Chaos::ClothAsset
 		/** Remove a weight map from this cloth. */
 		void RemoveWeightMap(const FName& Name);
 
+		TArrayView<float> GetWeightMap(const FName& Name);
 	private:
+
+		void SetDefaults();
+
 		TSharedPtr<FClothCollection> GetClothCollection() { return ConstCastSharedPtr<FClothCollection>(ClothCollection); }
+
+		friend class FCollectionClothSeamFacade;
+		friend class FCollectionClothSimPatternFacade;
+
+		explicit FCollectionClothFacade(const TSharedPtr<FClothCollection>& InClothCollection);
+
+		// These methods are private because they're managed by the FCollectionClothSeamFacade.
+		//~ Sim Vertices 2D Group
+		TArrayView<int32> GetSimVertex3DLookupPrivate();
+
+		//~ Sim Vertices 3D Group
+		TArrayView<TArray<int32>> GetSeamStitchLookupPrivate();
+		TArrayView<TArray<int32>> GetSimVertex2DLookupPrivate();
 	};
 }  // End namespace UE::Chaos::ClothAsset
diff --git a/Engine/Plugins/Experimental/ChaosClothAsset/Source/ChaosClothAsset/Public/ChaosClothAsset/CollectionClothRenderPatternFacade.h b/Engine/Plugins/Experimental/ChaosClothAsset/Source/ChaosClothAsset/Public/ChaosClothAsset/CollectionClothRenderPatternFacade.h
new file mode 100644
--- /dev/null
+++ b/Engine/Plugins/Experimental/ChaosClothAsset/Source/ChaosClothAsset/Public/ChaosClothAsset/CollectionClothRenderPatternFacade.h
@@ -0,0 +1,124 @@
+// Copyright Epic Games, Inc. All Rights Reserved.
+
+#pragma once
+
+#include "Templates/SharedPointer.h"
+
+struct FManagedArrayCollection;
+
+namespace UE::Chaos::ClothAsset
+{
+	class FClothCollection;
+
+	/**
+	 * Cloth Asset collection render pattern facade class to access cloth render pattern data.
+	 * Constructed from FCollectionClothConstFacade.
+	 * Const access (read only) version.
+	 */
+	class CHAOSCLOTHASSET_API FCollectionClothRenderPatternConstFacade
+	{
+	public:
+		FCollectionClothRenderPatternConstFacade() = delete;
+
+		FCollectionClothRenderPatternConstFacade(const FCollectionClothRenderPatternConstFacade&) = delete;
+		FCollectionClothRenderPatternConstFacade& operator=(const FCollectionClothRenderPatternConstFacade&) = delete;
+
+		FCollectionClothRenderPatternConstFacade(FCollectionClothRenderPatternConstFacade&&) = default;
+		FCollectionClothRenderPatternConstFacade& operator=(FCollectionClothRenderPatternConstFacade&&) = default;
+
+		virtual ~FCollectionClothRenderPatternConstFacade() = default;
+
+		/** Return the render material for this pattern. */
+		const FString& GetRenderMaterialPathName() const;
+
+		//~ Render Vertices Group
+		// Note: Use the FCollectionClothConstFacade accessors instead of these for the array indices to match the RenderIndices values
+		/** Return the total number of render vertices for this pattern. */
+		int32 GetNumRenderVertices() const;
+		/** Return the render vertices offset for this pattern in the render vertices for the collection. */
+		int32 GetRenderVerticesOffset() const;
+		TConstArrayView<FVector3f> GetRenderPosition() const;
+		TConstArrayView<FVector3f> GetRenderNormal() const;
+		TConstArrayView<FVector3f> GetRenderTangentU() const;
+		TConstArrayView<FVector3f> GetRenderTangentV() const;
+		TConstArrayView<TArray<FVector2f>> GetRenderUVs() const;
+		TConstArrayView<FLinearColor> GetRenderColor() const;
+		TConstArrayView<TArray<int32>> GetRenderBoneIndices() const;
+		TConstArrayView<TArray<float>> GetRenderBoneWeights() const;
+		
+		//~ Render Faces Group
+		// Note: RenderIndices points to the collection arrays, not the pattern arrays
+		int32 GetNumRenderFaces() const;
+		/** Return the render faces offset for this pattern in the render faces */
+		int32 GetRenderFacesOffset() const;
+		TConstArrayView<FIntVector3> GetRenderIndices() const;
+
+		/** Return the Pattern index this facade has been created with. */
+		int32 GetPatternIndex() const { return PatternIndex; }
+
+	protected:
+		friend class FCollectionClothRenderPatternFacade;  // For other instances access
+		friend class FCollectionClothConstFacade;
+		FCollectionClothRenderPatternConstFacade(const TSharedPtr<const FClothCollection>& InClothCollection, int32 InPatternIndex);
+
+		static constexpr int32 GetBaseElementIndex() { return 0; }
+		int32 GetElementIndex() const { return GetBaseElementIndex() + PatternIndex; }
+
+		TSharedPtr<const FClothCollection> ClothCollection;
+		int32 PatternIndex;
+	};
+
+	/**
+	 * Cloth Asset collection render pattern facade class to access cloth render pattern data.
+	 * Constructed from FCollectionClothFacade.
+	 * Non-const access (read/write) version.
+	 */
+	class CHAOSCLOTHASSET_API FCollectionClothRenderPatternFacade final : public FCollectionClothRenderPatternConstFacade
+	{
+	public:
+		FCollectionClothRenderPatternFacade() = delete;
+
+		FCollectionClothRenderPatternFacade(const FCollectionClothRenderPatternFacade&) = delete;
+		FCollectionClothRenderPatternFacade& operator=(const FCollectionClothRenderPatternFacade&) = delete;
+
+		FCollectionClothRenderPatternFacade(FCollectionClothRenderPatternFacade&&) = default;
+		FCollectionClothRenderPatternFacade& operator=(FCollectionClothRenderPatternFacade&&) = default;
+
+		virtual ~FCollectionClothRenderPatternFacade() override = default;
+
+		/** Remove all geometry from this cloth pattern. */
+		void Reset();
+
+		/** Initialize from another render pattern. Assumes all indices match between source and target. */
+		void Initialize(const FCollectionClothRenderPatternConstFacade& Other);
+
+		/** Set the render material for this pattern. */
+		void SetRenderMaterialPathName(const FString& PathName);
+
+		//~ Render Vertices Group
+		// Note: Use the FCollectionClothFacade accessors instead of these for the array indices to match the RenderIndices values
+		/** Grow or shrink the space reserved for render vertices for this pattern within the cloth collection and return its start index. */
+		void SetNumRenderVertices(int32 NumRenderVertices);
+		TArrayView<FVector3f> GetRenderPosition();
+		TArrayView<FVector3f> GetRenderNormal();
+		TArrayView<FVector3f> GetRenderTangentU();
+		TArrayView<FVector3f> GetRenderTangentV();
+		TArrayView<TArray<FVector2f>> GetRenderUVs();
+		TArrayView<FLinearColor> GetRenderColor();
+		TArrayView<TArray<int32>> GetRenderBoneIndices();
+		TArrayView<TArray<float>> GetRenderBoneWeights();
+
+		//~ Render Faces Group
+		/** Grow or shrink the space reserved for render faces for this pattern within the cloth collection and return its start index. */
+		void SetNumRenderFaces(int32 NumRenderFaces);
+		TArrayView<FIntVector3> GetRenderIndices();
+
+	protected:
+		friend class FCollectionClothFacade;
+		FCollectionClothRenderPatternFacade(const TSharedPtr<FClothCollection>& InClothCollection, int32 InPatternIndex);
+
+		void SetDefaults();
+
+		TSharedPtr<FClothCollection> GetClothCollection() { return ConstCastSharedPtr<FClothCollection>(ClothCollection); }
+	};
+}  // End namespace UE::Chaos::ClothAsset
diff --git a/Engine/Plugins/Experimental/ChaosClothAsset/Source/ChaosClothAsset/Public/ChaosClothAsset/CollectionClothSimPatternFacade.h b/Engine/Plugins/Experimental/ChaosClothAsset/Source/ChaosClothAsset/Public/ChaosClothAsset/CollectionClothSimPatternFacade.h
new file mode 100644
--- /dev/null
+++ b/Engine/Plugins/Experimental/ChaosClothAsset/Source/ChaosClothAsset/Public/ChaosClothAsset/CollectionClothSimPatternFacade.h
@@ -0,0 +1,120 @@
+// Copyright Epic Games, Inc. All Rights Reserved.
+
+#pragma once
+
+#include "Templates/SharedPointer.h"
+
+struct FManagedArrayCollection;
+
+namespace UE::Chaos::ClothAsset
+{
+	class FClothCollection;
+
+	/**
+	 * Cloth Asset collection sim pattern facade class to access cloth sim pattern data.
+	 * Constructed from FCollectionClothConstFacade.
+	 * Const access (read only) version.
+	 */
+	class CHAOSCLOTHASSET_API FCollectionClothSimPatternConstFacade
+	{
+	public:
+		FCollectionClothSimPatternConstFacade() = delete;
+
+		FCollectionClothSimPatternConstFacade(const FCollectionClothSimPatternConstFacade&) = delete;
+		FCollectionClothSimPatternConstFacade& operator=(const FCollectionClothSimPatternConstFacade&) = delete;
+
+		FCollectionClothSimPatternConstFacade(FCollectionClothSimPatternConstFacade&&) = default;
+		FCollectionClothSimPatternConstFacade& operator=(FCollectionClothSimPatternConstFacade&&) = default;
+
+		virtual ~FCollectionClothSimPatternConstFacade() = default;
+
+		//~ Sim Vertices 2D Group
+		// Note: Use the FCollectionClothConstFacade accessors instead of these for the array indices to match the SimIndices2D values
+		/** Return the total number of simulation vertices for this pattern. */
+		int32 GetNumSimVertices2D() const;
+		/** Return the simulation vertices offset for this pattern in the simulation vertices for the collection. */
+		int32 GetSimVertices2DOffset() const;
+		TConstArrayView<FVector2f> GetSimPosition2D() const;
+		TConstArrayView<int32> GetSimVertex3DLookup() const;
+
+		//~ Sim Faces Group
+		// Note: SimIndices points to the collection arrays, not the pattern arrays
+		int32 GetNumSimFaces() const;
+		/** Return the simulation faces offset for this pattern in the simulation faces. */
+		int32 GetSimFacesOffset() const;
+		TConstArrayView<FIntVector3> GetSimIndices2D() const;
+		TConstArrayView<FIntVector3> GetSimIndices3D() const;
+
+		/** Return the Pattern index this facade has been created with. */
+		int32 GetPatternIndex() const { return PatternIndex; }
+
+	protected:
+		friend class FCollectionClothSimPatternFacade;  // For other instances access
+		friend class FCollectionClothConstFacade;
+		FCollectionClothSimPatternConstFacade(const TSharedPtr<const FClothCollection>& InClothCollection, int32 InPatternIndex);
+
+		static constexpr int32 GetBaseElementIndex() { return 0; }
+		int32 GetElementIndex() const { return GetBaseElementIndex() + PatternIndex; }
+
+		TSharedPtr<const FClothCollection> ClothCollection;
+		int32 PatternIndex;
+	};
+
+	/**
+	 * Cloth Asset collection sim pattern facade class to access cloth sim pattern data.
+	 * Constructed from FCollectionClothFacade.
+	 * Non-const access (read/write) version.
+	 */
+	class CHAOSCLOTHASSET_API FCollectionClothSimPatternFacade final : public FCollectionClothSimPatternConstFacade
+	{
+	public:
+		FCollectionClothSimPatternFacade() = delete;
+
+		FCollectionClothSimPatternFacade(const FCollectionClothSimPatternFacade&) = delete;
+		FCollectionClothSimPatternFacade& operator=(const FCollectionClothSimPatternFacade&) = delete;
+
+		FCollectionClothSimPatternFacade(FCollectionClothSimPatternFacade&&) = default;
+		FCollectionClothSimPatternFacade& operator=(FCollectionClothSimPatternFacade&&) = default;
+
+		virtual ~FCollectionClothSimPatternFacade() override = default;
+
+		/** Remove all geometry from this cloth pattern. */
+		void Reset();
+
+		/** Initialize the cloth pattern using the specified 3D and 2D positions, and topology. */
+		template<typename IndexType>
+		void Initialize(const TArray<FVector2f>& Positions2D, const TArray<FVector3f>& Positions3D, const TArray<IndexType>& Indices);
+
+		/** Initialize this pattern using another pattern collection. */
+		void Initialize(const FCollectionClothSimPatternConstFacade& Other, const int32 SimVertex3DOffset);
+
+		//~ Sim Vertices 2D Group
+		// Note: Use the FCollectionClothConstFacade accessors instead of these for the array indices to match the SimIndices2D values
+		TArrayView<FVector2f> GetSimPosition2D();
+
+		/** This will remove the 2D vertices, but the associated seams and 3D vertices will still exist, and point to INDEX_NONE */
+		void RemoveSimVertices2D(int32 NumSimVertices);
+		void RemoveAllSimVertices2D() { RemoveSimVertices2D(GetNumSimVertices2D()); }
+
+		//~ Sim Faces Group
+		// Note: SimIndices points to the collection arrays, not the pattern arrays
+		/** Grow or shrink the space reserved for simulation faces for this pattern within the cloth collection and return its start index. */
+		void SetNumSimFaces(int32 NumSimFaces);
+		TArrayView<FIntVector3> GetSimIndices2D();
+		TArrayView<FIntVector3> GetSimIndices3D();
+
+	protected:
+		friend class FCollectionClothFacade;
+		FCollectionClothSimPatternFacade(const TSharedPtr<FClothCollection>& InClothCollection, int32 InPatternIndex);
+
+		void SetDefaults();
+
+		/** Append this many 2D and 3D SimVertices which correspond to each other. Returns {2D start index (local pattern index), 3D start index}. */
+		FIntVector2 AppendSimVertices(int32 NumSimVertices);
+
+		void SetNumSimVertices2D(int32 NumSimVertices);
+		TArrayView<int32> GetSimVertex3DLookup();
+
+		TSharedPtr<FClothCollection> GetClothCollection() { return ConstCastSharedPtr<FClothCollection>(ClothCollection); }
+	};
+}  // End namespace UE::Chaos::ClothAsset
diff --git a/Engine/Plugins/Experimental/ChaosClothAsset/Source/ChaosClothAssetEngine/Private/ChaosClothAsset/ClothAsset.cpp b/Engine/Plugins/Experimental/ChaosClothAsset/Source/ChaosClothAssetEngine/Private/ChaosClothAsset/ClothAsset.cpp
--- a/Engine/Plugins/Experimental/ChaosClothAsset/Source/ChaosClothAssetEngine/Private/ChaosClothAsset/ClothAsset.cpp
+++ b/Engine/Plugins/Experimental/ChaosClothAsset/Source/ChaosClothAssetEngine/Private/ChaosClothAsset/ClothAsset.cpp
@@ -28,18 +28,85 @@
 // In case of merge conflicts with DDC versions, you MUST generate a new GUID and set this new GUID as the version.
 #define CHAOS_CLOTH_ASSET_DERIVED_DATA_VERSION TEXT("19823D996CA54F279B9A9FA8ED7A8EB6")
 
+namespace UE::Chaos::ClothAsset::Private
+{
+::Chaos::FChaosArchive& Serialize(::Chaos::FChaosArchive& Ar, TArray<TSharedPtr<FManagedArrayCollection>>& ClothCollections)
+{
+	Ar.UsingCustomVersion(FUE5MainStreamObjectVersion::GUID);
+
+	if (Ar.IsLoading() && Ar.CustomVer(FUE5MainStreamObjectVersion::GUID) < FUE5MainStreamObjectVersion::ClothCollectionSingleLodSchema)
+	{
+		// Cloth assets before this version had a single ClothCollection with a completely different schema.
+		ClothCollections.Empty(1);
+		TSharedPtr<FManagedArrayCollection>& ClothCollection = ClothCollections.Emplace_GetRef(MakeShared<FManagedArrayCollection>());
+		ClothCollection->Serialize(Ar);
+
+		// Now we're just going to hard reset and define a new schema.
+		ClothCollection->Reset();
+		FCollectionClothFacade ClothFacade(ClothCollection);
+		ClothFacade.DefineSchema();
+
+		return Ar;
+	}
+	else
+	{
+		// This is following Serialize for Arrays
+		ClothCollections.CountBytes(Ar);
+		int32 SerializeNum = Ar.IsLoading() ? 0 : ClothCollections.Num();
+		Ar << SerializeNum;
+		if (SerializeNum == 0)
+		{
+			// if we are loading, then we have to reset the size to 0, in case it isn't currently 0
+			if (Ar.IsLoading())
+			{
+				ClothCollections.Empty();
+			}
+			return Ar;
+		}
+		check(SerializeNum >= 0);
+
+		if (Ar.IsError() || SerializeNum < 0)
+		{
+			Ar.SetError();
+			return Ar;
+		}
+		if (Ar.IsLoading())
+		{
+			// Required for resetting ArrayNum
+			ClothCollections.Empty(SerializeNum);
+
+			for (int32 i = 0; i < SerializeNum; i++)
+			{
+				TSharedPtr<FManagedArrayCollection>& ClothCollection = ClothCollections.Emplace_GetRef(MakeShared<FManagedArrayCollection>());
+				ClothCollection->Serialize(Ar);
+			}
+		}
+		else
+		{
+			check(SerializeNum == ClothCollections.Num());
+
+			for (int32 i = 0; i < SerializeNum; i++)
+			{
+				ClothCollections[i]->Serialize(Ar);
+			}
+		}
+
+		return Ar;
+	}
+}
+}
+
 UChaosClothAsset::UChaosClothAsset(const FObjectInitializer& ObjectInitializer)
 	: Super(ObjectInitializer)
 	, DisableBelowMinLodStripping(FPerPlatformBool(false))
-	, ClothCollection(MakeShared<FManagedArrayCollection>())
 #if WITH_EDITORONLY_DATA
 	, MeshModel(MakeShareable(new FSkeletalMeshModel()))
 #endif
 {
-	// Setup the collection to be a cloth collection and have at least one LOD
+	// Setup a single LOD's Cloth Collection
+	TSharedPtr<FManagedArrayCollection>& ClothCollection = ClothCollections.Emplace_GetRef(MakeShared<FManagedArrayCollection>());
 	UE::Chaos::ClothAsset::FCollectionClothFacade ClothFacade(ClothCollection);
 	ClothFacade.DefineSchema();
-	ClothFacade.AddLod();
 
 	// Set default skeleton (must be done after having added the LOD)
 	constexpr bool bRebuildModels = false;
@@ -88,9 +155,7 @@ void UChaosClothAsset::Serialize(FArchive& Ar)
 	Ar << bCooked;
 
 	Chaos::FChaosArchive ChaosArchive(Ar);
-	ClothCollection->Serialize(ChaosArchive);
-	UE::Chaos::ClothAsset::FCollectionClothFacade ClothFacade(ClothCollection);
-	ClothFacade.PostSerialize(ChaosArchive);
+	UE::Chaos::ClothAsset::Private::Serialize(ChaosArchive, ClothCollections);
 
 	Ar << GetRefSkeleton();
 	if (Ar.IsLoading())
@@ -140,17 +205,31 @@ void UChaosClothAsset::BeginPostLoadInternal(FSkinnedAssetPostLoadContext& Conte
 	FSkinnedAssetAsyncBuildScope AsyncBuildScope(this);
 
 	// Make sure that the collection is still compatible and valid
-	UE::Chaos::ClothAsset::FCollectionClothFacade ClothFacade(ClothCollection);
-	if (!ClothFacade.IsValid())
+	bool bAnyInvalidLods = false;
+	if (ClothCollections.IsEmpty())
 	{
-		UE_LOG(LogChaosClothAsset, Warning, TEXT("Invalid Cloth Collection found while loading Cloth Asset %s."), *GetFullName());
-		ClothCollection = MakeShared<FManagedArrayCollection>();
-		ClothFacade = UE::Chaos::ClothAsset::FCollectionClothFacade(ClothCollection);
+		UE_LOG(LogChaosClothAsset, Warning, TEXT("Invalid Cloth Collection (no LODs) found while loading Cloth Asset %s."), *GetFullName());
+		TSharedPtr<FManagedArrayCollection>& ClothCollection = ClothCollections.Emplace_GetRef(MakeShared<FManagedArrayCollection>());
+		UE::Chaos::ClothAsset::FCollectionClothFacade ClothFacade(ClothCollection);
 		ClothFacade.DefineSchema();
+		bAnyInvalidLods = true;
+	}
+	for(int32 LODIndex = 0; LODIndex < ClothCollections.Num(); ++LODIndex)
+	{
+		TSharedPtr<FManagedArrayCollection>& ClothCollection = ClothCollections[LODIndex];
+
+		UE::Chaos::ClothAsset::FCollectionClothFacade ClothFacade(ClothCollection);
+		if (!ClothFacade.IsValid())
+		{
+			UE_LOG(LogChaosClothAsset, Warning, TEXT("Invalid Cloth Collection found at LOD %i while loading Cloth Asset %s."), LODIndex, *GetFullName());
+			ClothCollection = MakeShared<FManagedArrayCollection>();
+			ClothFacade = UE::Chaos::ClothAsset::FCollectionClothFacade(ClothCollection);
+			ClothFacade.DefineSchema();
+			bAnyInvalidLods = true;
+		}
 	}
-	if (!ClothFacade.GetNumLods())
+	if (bAnyInvalidLods)
 	{
-		ClothFacade.AddLod();
 		constexpr bool bRebuildModels = false;
 		SetSkeleton(Skeleton, bRebuildModels);  // Re-update the collection with the skeleton information if any
 		SetPhysicsAsset(PhysicsAsset);  // Re-update the collection with the physics asset information if any
@@ -315,12 +394,10 @@ void UChaosClothAsset::CalculateBounds()
 
 	FBox BoundingBox(ForceInit);
 
-	const FCollectionClothConstFacade Cloth(ClothCollection);
-
-	for (int32 LodIndex = 0; LodIndex < Cloth.GetNumLods(); ++LodIndex)
+	for (const TSharedPtr<FManagedArrayCollection>& ClothCollection : ClothCollections)
 	{
-		const FCollectionClothLodConstFacade ClothLod = Cloth.GetLod(LodIndex);
-		const TConstArrayView<FVector3f> RenderPositionArray = ClothLod.GetRenderPosition();
+		const FCollectionClothConstFacade Cloth(ClothCollection);
+		const TConstArrayView<FVector3f> RenderPositionArray = Cloth.GetRenderPosition();
 
 		for (const FVector3f& RenderPosition : RenderPositionArray)
 		{
@@ -365,8 +442,7 @@ void UChaosClothAsset::Build()
 	CalculateBounds();
 
 	// Add LODs to the render data
-	const FCollectionClothConstFacade Cloth(ClothCollection);
-	const int32 NumLods = FMath::Max(Cloth.GetNumLods(), 1);  // The render data will always look for at least one default LOD 0
+	const int32 NumLods = FMath::Max(ClothCollections.Num(), 1);  // The render data will always look for at least one default LOD 0
 
 	// Rebuild LOD Infos
 	LODInfo.Reset(NumLods);
@@ -447,8 +523,7 @@ void UChaosClothAsset::BuildMeshModel()
 		{
 			using namespace UE::Chaos::ClothAsset;
 
-			const FCollectionClothConstFacade Cloth(GetClothCollection());
-			const int32 NumLods = Cloth.GetNumLods();
+			const int32 NumLods = ClothCollections.Num();
 
 			// Build each LOD
 			check(MeshModel);  // MeshModel should always be created in the Cloth Asset constructor WITH_EDITORONLY_DATA
@@ -470,7 +545,7 @@ void UChaosClothAsset::BuildMeshModel()
 
 void UChaosClothAsset::BuildClothSimulationModel()
 {
-	ClothSimulationModel = MakeShared<FChaosClothSimulationModel>(GetClothCollection(), GetRefSkeleton());
+	ClothSimulationModel = MakeShared<FChaosClothSimulationModel>(const_cast<const UChaosClothAsset*>(this)->GetClothCollections(), GetRefSkeleton());
 }
 
 const FMeshUVChannelInfo* UChaosClothAsset::GetUVChannelData(int32 MaterialIndex) const
@@ -606,13 +681,17 @@ void UChaosClothAsset::SetPhysicsAsset(UPhysicsAsset* InPhysicsAsset)
 	using namespace UE::Chaos::ClothAsset;
 
 	PhysicsAsset = InPhysicsAsset;
-	FClothGeometryTools::SetPhysicsAssetPathName(ClothCollection, PhysicsAsset ? PhysicsAsset->GetPathName() : FString());
+	for (TSharedPtr<FManagedArrayCollection>& ClothCollection : ClothCollections)
+	{
+		FCollectionClothFacade Cloth(ClothCollection);
+		Cloth.SetPhysicsAssetPathName(PhysicsAsset ? PhysicsAsset->GetPathName() : FString());
+	}
 }
 
 void UChaosClothAsset::SetSkeleton(USkeleton* InSkeleton, bool bRebuildModels)
 {
 	using namespace UE::Chaos::ClothAsset;
-	check(ClothCollection.IsValid());
+	check(ClothCollections.Num());
 
 	const TCHAR* const DefaultSkeletonPathName = TEXT("/Engine/EditorMeshes/SkeletalMesh/DefaultSkeletalMesh_Skeleton.DefaultSkeletalMesh_Skeleton");
 	Skeleton = InSkeleton ?
@@ -621,10 +700,16 @@ void UChaosClothAsset::SetSkeleton(USkeleton* InSkeleton, bool bRebuildModels)
 
 	RefSkeleton = Skeleton->GetReferenceSkeleton();
 
-	constexpr bool bBindSimMesh = true;
-	constexpr bool bBindRenderMesh = true;
-	FClothGeometryTools::BindMeshToRootBone(ClothCollection, bBindSimMesh, bBindRenderMesh);
-	FClothGeometryTools::SetSkeletonAssetPathName(ClothCollection, Skeleton->GetPathName());
+	for (TSharedPtr<FManagedArrayCollection>& ClothCollection : ClothCollections)
+	{
+		check(ClothCollection.IsValid());
+		constexpr bool bBindSimMesh = true;
+		constexpr bool bBindRenderMesh = true;
+		FClothGeometryTools::BindMeshToRootBone(ClothCollection, bBindSimMesh, bBindRenderMesh);
+
+		FCollectionClothFacade Cloth(ClothCollection);
+		Cloth.SetSkeletonAssetPathName(Skeleton->GetPathName());
+	}
 	if (bRebuildModels)
 	{
 		Build();
@@ -634,16 +719,14 @@ void UChaosClothAsset::SetSkeleton(USkeleton* InSkeleton, bool bRebuildModels)
 void UChaosClothAsset::UpdateSkeletonFromCollection(bool bRebuildModels)
 {
 	using namespace UE::Chaos::ClothAsset;
-	check(ClothCollection.IsValid());
 
-	FCollectionClothConstFacade ClothFacade(ClothCollection);
+	check(ClothCollections.Num());
+	FCollectionClothConstFacade ClothFacade(ClothCollections[0]);
+	check(ClothFacade.IsValid());
 
 	USkeleton* InSkeleton = nullptr;
-	if (ClothFacade.GetNumLods())
-	{
-		const FString& SkeletonAssetPathName = ClothFacade.GetLod(0).GetSkeletonAssetPathName();
-		InSkeleton = LoadObject<USkeleton>(nullptr, *SkeletonAssetPathName, nullptr, LOAD_None, nullptr);
-	}
+	const FString& SkeletonAssetPathName = ClothFacade.GetSkeletonAssetPathName();
+	InSkeleton = LoadObject<USkeleton>(nullptr, *SkeletonAssetPathName, nullptr, LOAD_None, nullptr);
 
 	if (!InSkeleton)
 	{
@@ -664,18 +747,25 @@ void UChaosClothAsset::UpdateSkeletonFromCollection(bool bRebuildModels)
 void UChaosClothAsset::CopySimMeshToRenderMesh(UMaterialInterface* Material)
 {
 	using namespace UE::Chaos::ClothAsset;
-	check(ClothCollection.IsValid());
+	check(ClothCollections.Num());
 
 	// Add a default material if none is specified
 	const FString RenderMaterialPathName = Material ?
 		Material->GetPathName() :
 		FString(TEXT("/Engine/EditorMaterials/Cloth/CameraLitDoubleSided.CameraLitDoubleSided"));
 
-	FClothGeometryTools::CopySimMeshToRenderMesh(ClothCollection, RenderMaterialPathName);
+	bool bAnyLodHasRenderMesh = false;
+	for (TSharedPtr<FManagedArrayCollection>& ClothCollection : ClothCollections)
+	{
+		check(ClothCollection.IsValid());
+		constexpr bool bSingleRenderPattern = true;
+		FClothGeometryTools::CopySimMeshToRenderMesh(ClothCollection, RenderMaterialPathName, bSingleRenderPattern);
+		bAnyLodHasRenderMesh = bAnyLodHasRenderMesh || FClothGeometryTools::HasRenderMesh(ClothCollection);
+	}
 
 	// Set new material
 	Materials.Reset(1);
-	if (FClothGeometryTools::HasRenderMesh(ClothCollection))
+	if (bAnyLodHasRenderMesh)
 	{
 		if (UMaterialInterface* const LoadedMaterial = LoadObject<UMaterialInterface>(nullptr, *RenderMaterialPathName, nullptr, LOAD_None, nullptr))
 		{
@@ -702,7 +792,11 @@ void UChaosClothAsset::ReregisterComponents()
 void UChaosClothAsset::BindSimMeshToRootBone()
 {
 	using namespace UE::Chaos::ClothAsset;
+	check(ClothCollections.Num());
 
-	check(ClothCollection.IsValid());
-	FClothGeometryTools::BindMeshToRootBone(ClothCollection, true, false);
+	for (TSharedPtr<FManagedArrayCollection>& ClothCollection : ClothCollections)
+	{
+		check(ClothCollection.IsValid());
+		FClothGeometryTools::BindMeshToRootBone(ClothCollection, true, false);
+	}
 }
\ No newline at end of file
diff --git a/Engine/Plugins/Experimental/ChaosClothAsset/Source/ChaosClothAssetEngine/Private/ChaosClothAsset/ClothComponent.cpp b/Engine/Plugins/Experimental/ChaosClothAsset/Source/ChaosClothAssetEngine/Private/ChaosClothAsset/ClothComponent.cpp
--- a/Engine/Plugins/Experimental/ChaosClothAsset/Source/ChaosClothAssetEngine/Private/ChaosClothAsset/ClothComponent.cpp
+++ b/Engine/Plugins/Experimental/ChaosClothAsset/Source/ChaosClothAssetEngine/Private/ChaosClothAsset/ClothComponent.cpp
@@ -72,9 +72,10 @@ void UChaosClothComponent::ResetConfigProperties()
 		// Create a mutable facade for our component's property collection, and use it to copy the properties from the cloth collection
 		::Chaos::Softs::FCollectionPropertyMutableFacade CollectionPropertyMutableFacade(PropertyCollection);
 
-		if (GetClothAsset() && GetClothAsset()->GetClothCollection())
+		// TODO: Make PropertyCollection also one per LOD.
+		if (GetClothAsset() && GetClothAsset()->GetClothCollections().Num() > 0 && GetClothAsset()->GetClothCollections()[0].IsValid())
 		{
-			CollectionPropertyMutableFacade.Copy(*GetClothAsset()->GetClothCollection());
+			CollectionPropertyMutableFacade.Copy(*GetClothAsset()->GetClothCollections()[0]);
 		}
 		else
 		{
diff --git a/Engine/Plugins/Experimental/ChaosClothAsset/Source/ChaosClothAssetEngine/Private/ChaosClothAsset/ClothEngineTools.cpp b/Engine/Plugins/Experimental/ChaosClothAsset/Source/ChaosClothAssetEngine/Private/ChaosClothAsset/ClothEngineTools.cpp
--- a/Engine/Plugins/Experimental/ChaosClothAsset/Source/ChaosClothAssetEngine/Private/ChaosClothAsset/ClothEngineTools.cpp
+++ b/Engine/Plugins/Experimental/ChaosClothAsset/Source/ChaosClothAssetEngine/Private/ChaosClothAsset/ClothEngineTools.cpp
@@ -9,45 +9,38 @@
 namespace UE::Chaos::ClothAsset
 {
 
-void FClothEngineTools::GenerateTethers(const TSharedPtr<FManagedArrayCollection>& ClothCollection, const FName& WeightMap, const bool bGenerateGeodesicTethers)
+void FClothEngineTools::GenerateTethers(const TSharedPtr<FManagedArrayCollection>& ClothCollection, const FName& WeightMapName, const bool bGenerateGeodesicTethers)
 {
 	FCollectionClothFacade ClothFacade(ClothCollection);
-	if (ClothFacade.HasWeightMap(WeightMap))
+	FClothGeometryTools::DeleteTethers(ClothCollection);
+	if (ClothFacade.HasWeightMap(WeightMapName))
 	{
-		for (int32 LodIndex = 0; LodIndex < ClothFacade.GetNumLods(); ++LodIndex)
+		FClothTetherData TetherData;
+		TArray<uint32> SimIndices;
+		SimIndices.Reserve(ClothFacade.GetNumSimFaces() * 3);
+		for (const FIntVector3& Face : ClothFacade.GetSimIndices3D())
 		{
-			FCollectionClothLodFacade ClothLodFacade = ClothFacade.GetLod(LodIndex);
-
-			// We want to create the tethers based on the welded mesh, so create that now
-			TArray<FVector3f> Positions;
-			TArray<FVector3f> Normals;
-			TArray<uint32> Indices;
-			TArray<FVector2f> PatternsPositions;
-			TArray<uint32> PatternsIndices;
-			TArray<uint32> PatternToWeldedIndices;
-			TArray<TArray<int32>> WeldedToPatternIndices;
-			ClothLodFacade.BuildSimulationMesh(Positions, Normals, Indices, PatternsPositions, PatternsIndices, PatternToWeldedIndices, &WeldedToPatternIndices);
-
-			const TArray<float> WeldedWeightMap = FClothGeometryTools::BuildWeldedWeightMapForLod(ClothCollection, LodIndex, WeightMap, WeldedToPatternIndices);
-
-			FClothTetherData TetherData;
-			TetherData.GenerateTethers(TConstArrayView<FVector3f>(Positions), TConstArrayView<uint32>(Indices), TConstArrayView<float>(WeldedWeightMap), bGenerateGeodesicTethers);
+			SimIndices.Add(Face[0]);
+			SimIndices.Add(Face[1]);
+			SimIndices.Add(Face[2]);
+		}
+		TetherData.GenerateTethers(ClothFacade.GetSimPosition3D(), TConstArrayView<uint32>(SimIndices), ClothFacade.GetWeightMap(WeightMapName), bGenerateGeodesicTethers);
 
-			// Set new tethers
-			const int32 NumTetherBatches = TetherData.Tethers.Num();
-			ClothLodFacade.SetNumTetherBatches(NumTetherBatches);
-			for (int32 TetherBatchIndex = 0; TetherBatchIndex < NumTetherBatches; ++TetherBatchIndex)
+		// Append new tethers
+		TArrayView<TArray<int32>> TetherKinematicIndex = ClothFacade.GetTetherKinematicIndex();
+		TArrayView<TArray<float>> TetherReferenceLength = ClothFacade.GetTetherReferenceLength();
+		for(const TArray<TTuple<int32, int32, float>>& TetherBatch : TetherData.Tethers)
+		{
+			for (const TTuple<int32, int32, float>& Tether : TetherBatch)
 			{
-				TArray<TTuple<int32, int32, float>>& TetherBatch = TetherData.Tethers[TetherBatchIndex];
-				// Remap tether data indices back to unwelded indices. Just set them to the first of the unwelded index if that welded index is formed from multiple unwelded indices.
-				for (TTuple<int32, int32, float>& Tether : TetherBatch)
-				{
-					Tether.Get<0>() = WeldedToPatternIndices[Tether.Get<0>()][0];
-					Tether.Get<1>() = WeldedToPatternIndices[Tether.Get<1>()][0];
-				}
-
-				FCollectionClothTetherBatchFacade TetherBatchFacade = ClothLodFacade.GetTetherBatch(TetherBatchIndex);
-				TetherBatchFacade.Initialize(TetherBatch);
+				// Tuple is Kinematic, Dynamic, RefLength
+				const int32 DynamicIndex = Tether.Get<1>();
+				TArray<int32>& KinematicIndex = TetherKinematicIndex[DynamicIndex];
+				TArray<float>& ReferenceLength = TetherReferenceLength[DynamicIndex];
+				check(KinematicIndex.Num() == ReferenceLength.Num());
+				checkSlow(KinematicIndex.Find(Tether.Get<0>()) == INDEX_NONE);
+				KinematicIndex.Add(Tether.Get<0>());
+				ReferenceLength.Add(Tether.Get<2>());
 			}
 		}
 	}
