Fix bug found by static analysis from recent CL in CollectionClothSeamFacade.cpp

diff --git a/Engine/Plugins/Experimental/ChaosClothAsset/Source/ChaosClothAsset/Private/ChaosClothAsset/CollectionClothSeamFacade.cpp b/Engine/Plugins/Experimental/ChaosClothAsset/Source/ChaosClothAsset/Private/ChaosClothAsset/CollectionClothSeamFacade.cpp
--- a/Engine/Plugins/Experimental/ChaosClothAsset/Source/ChaosClothAsset/Private/ChaosClothAsset/CollectionClothSeamFacade.cpp
+++ b/Engine/Plugins/Experimental/ChaosClothAsset/Source/ChaosClothAsset/Private/ChaosClothAsset/CollectionClothSeamFacade.cpp
@@ -224,7 +224,7 @@ namespace UE::Chaos::ClothAsset
 				check(Indices.Num() == Lengths.Num());
 				const int32 NumTethers = Indices.Num();
 				// Go in reverse because we're going to remove any invalid tethers while we're here.
-				for (int32 TetherIndex = NumTethers - 1; TetherIndex >= 0; ++TetherIndex)
+				for (int32 TetherIndex = NumTethers - 1; TetherIndex >= 0; --TetherIndex)
 				{
 					if (Indices[TetherIndex] == INDEX_NONE)
 					{
