Chaos Cloth - Trivial, quick tidying up changes including:
- Fixed the FClothCollection forward declaration namespace.
- Removed a redundant forward declaration.
- Made a protected into a private because the class is final.
- Removed some duplicates whitespaces in a couple of trivial locations.

#rnx

diff --git a/Engine/Plugins/Experimental/ChaosClothAsset/Source/ChaosClothAsset/Public/ChaosClothAsset/ClothCollection.h b/Engine/Plugins/Experimental/ChaosClothAsset/Source/ChaosClothAsset/Public/ChaosClothAsset/ClothCollection.h
--- a/Engine/Plugins/Experimental/ChaosClothAsset/Source/ChaosClothAsset/Public/ChaosClothAsset/ClothCollection.h
+++ b/Engine/Plugins/Experimental/ChaosClothAsset/Source/ChaosClothAsset/Public/ChaosClothAsset/ClothCollection.h
@@ -5,7 +5,6 @@
 #include "GeometryCollection/ManagedArray.h"
 
 struct FManagedArrayCollection;
-template<class InElementType> class TManagedArray;
 
 namespace UE::Chaos::ClothAsset
 {
diff --git a/Engine/Plugins/Experimental/ChaosClothAsset/Source/ChaosClothAsset/Public/ChaosClothAsset/CollectionClothSimPatternFacade.h b/Engine/Plugins/Experimental/ChaosClothAsset/Source/ChaosClothAsset/Public/ChaosClothAsset/CollectionClothSimPatternFacade.h
--- a/Engine/Plugins/Experimental/ChaosClothAsset/Source/ChaosClothAsset/Public/ChaosClothAsset/CollectionClothSimPatternFacade.h
+++ b/Engine/Plugins/Experimental/ChaosClothAsset/Source/ChaosClothAsset/Public/ChaosClothAsset/CollectionClothSimPatternFacade.h
@@ -103,7 +103,7 @@ namespace UE::Chaos::ClothAsset
 		TArrayView<FIntVector3> GetSimIndices2D();
 		TArrayView<FIntVector3> GetSimIndices3D();
 
-	protected:
+	private:
 		friend class FCollectionClothFacade;
 		FCollectionClothSimPatternFacade(const TSharedPtr<FClothCollection>& InClothCollection, int32 InPatternIndex);
 
