Chaos Cloth - Fixed copy to render mesh UVs orientation, scale and offset.

#rnx

diff --git a/Engine/Plugins/Experimental/ChaosClothAsset/Source/ChaosClothAsset/Private/ChaosClothAsset/ClothGeometryTools.cpp b/Engine/Plugins/Experimental/ChaosClothAsset/Source/ChaosClothAsset/Private/ChaosClothAsset/ClothGeometryTools.cpp
--- a/Engine/Plugins/Experimental/ChaosClothAsset/Source/ChaosClothAsset/Private/ChaosClothAsset/ClothGeometryTools.cpp
+++ b/Engine/Plugins/Experimental/ChaosClothAsset/Source/ChaosClothAsset/Private/ChaosClothAsset/ClothGeometryTools.cpp
@@ -554,7 +554,8 @@ namespace UE::Chaos::ClothAsset
 
 			RenderPosition[VertexIndex] = SimPosition3D[VertexIndex3D];
 			RenderNormal[VertexIndex] = -SimNormal[VertexIndex3D];  // Simulation normals use reverse normals
-			RenderUVs[VertexIndex] = { SimPosition2D[VertexIndex] * UVInvScale };
+			RenderUVs[VertexIndex] = { (SimPosition2D[VertexIndex] - MinPosition) * UVInvScale };
+			RenderUVs[VertexIndex][0].Y = 1.f - RenderUVs[VertexIndex][0].Y;  // Reverse Y axis 
 			RenderColor[VertexIndex] = FLinearColor::White;
 			RenderTangentU[VertexIndex].Normalize();
 			RenderTangentV[VertexIndex].Normalize();
