Chaos Cloth Asset: Fix Seam lookup indexing bug.

This is a bug from the recent schema update. The seam stitch <--> Vertex3D lookups are incorrect. This only affects the ability to split seams. Not writing a serialization upgrade, but instead just regenerating the handful of cloth assets on disk.

diff --git a/Engine/Plugins/Experimental/ChaosClothAsset/Source/ChaosClothAsset/Private/ChaosClothAsset/CollectionClothSeamFacade.cpp b/Engine/Plugins/Experimental/ChaosClothAsset/Source/ChaosClothAsset/Private/ChaosClothAsset/CollectionClothSeamFacade.cpp
--- a/Engine/Plugins/Experimental/ChaosClothAsset/Source/ChaosClothAsset/Private/ChaosClothAsset/CollectionClothSeamFacade.cpp
+++ b/Engine/Plugins/Experimental/ChaosClothAsset/Source/ChaosClothAsset/Private/ChaosClothAsset/CollectionClothSeamFacade.cpp
@@ -94,22 +94,23 @@ namespace UE::Chaos::ClothAsset
 			}
 		}
 
-		static void UpdateWeldingLookups(const TMap<int32, FWeldingGroup>& WeldingGroups, const TArrayView<int32>& SimVertex3DLookup, const TArrayView<TArray<int32>>& SimVertex2DLookup)
+		// This is used for SimVertex3D <--> SimVertex2D, as well as SimVertex3D <--> SeamStitch
+		static void UpdateWeldingLookups(const TMap<int32, FWeldingGroup>& WeldingGroups, const TArrayView<int32>& SimVertex3DLookup, const TArrayView<TArray<int32>>& ReverseLookup)
 		{
 			for (TMap<int32, FWeldingGroup>::TConstIterator GroupIter = WeldingGroups.CreateConstIterator(); GroupIter; ++GroupIter)
 			{
 				const int32 PrimaryIndex3D = GroupIter.Key();
-				TArray<int32>& PrimarySimVertex2Ds = SimVertex2DLookup[PrimaryIndex3D];
+				TArray<int32>& PrimaryReverseLookup = ReverseLookup[PrimaryIndex3D];
 				for (const TPair<int32, int32>& IndexAndWeight : GroupIter.Value())
 				{
-					const TArray<int32>& MergingSimVertex2Ds = SimVertex2DLookup[IndexAndWeight.Get<0>()];
-					for (const int32 Vertex2D : MergingSimVertex2Ds)
+					const TArray<int32>& MergingReverseLookup = ReverseLookup[IndexAndWeight.Get<0>()];
+					for (const int32 ReverseIndex : MergingReverseLookup)
 					{
-						if (Vertex2D != INDEX_NONE)
+						if (ReverseIndex != INDEX_NONE)
 						{
-							// All 2D vertices that used to point to us need to point to PrimaryIndex3D
-							SimVertex3DLookup[Vertex2D] = PrimaryIndex3D;
-							PrimarySimVertex2Ds.AddUnique(Vertex2D);
+							// All elements that used to point to us need to point to PrimaryIndex3D
+							SimVertex3DLookup[ReverseIndex] = PrimaryIndex3D;
+							PrimaryReverseLookup.AddUnique(ReverseIndex);
 						}
 					}
 				}
@@ -259,6 +260,7 @@ namespace UE::Chaos::ClothAsset
 			WeldIndexAndFloatArrays<false, FClothCollection::MaxNumTetherAttachments>(WeldingGroups, TetherKinematicIndices, TetherReferenceLengths,
 				[](const TPair<float, int32>& A, const TPair<float, int32>& B) { return A < B; });
 		}
+
 	} // namespace Private
 
 	int32 FCollectionClothSeamConstFacade::GetNumSeamStitches() const
@@ -379,7 +381,7 @@ namespace UE::Chaos::ClothAsset
 		TArrayView<TArray<int32>> SeamStitchLookup = Cloth.GetSeamStitchLookupPrivate();
 		for (int32 StitchIndex = 0; StitchIndex < NumStitches; ++StitchIndex)
 		{
-			SeamStitch3DIndex[StitchIndex] = WeldingMappedValue(WeldingMap, StitchIndex);
+			SeamStitch3DIndex[StitchIndex] = WeldingMappedValue(WeldingMap, SeamStitch3DIndex[StitchIndex]);
 			SeamStitchLookup[SeamStitch3DIndex[StitchIndex]].Add(StitchOffset + StitchIndex);
 		}
 
@@ -392,6 +394,9 @@ namespace UE::Chaos::ClothAsset
 		// Update 2D vs 3D lookups
 		UpdateWeldingLookups(WeldingGroups, Cloth.GetSimVertex3DLookupPrivate(), Cloth.GetSimVertex2DLookupPrivate());
 
+		// Weld Stitch <-> 3D vertex lookups for stitches in other seams.
+		UpdateWeldingLookups(WeldingGroups, GetClothCollection()->GetElements(GetClothCollection()->GetSeamStitch3DIndex()), Cloth.GetSeamStitchLookupPrivate());
+
 		// Weld 3D positions
 		WeldByWeightedAverage(WeldingGroups, Cloth.GetSimPosition3D());
 
@@ -421,6 +426,8 @@ namespace UE::Chaos::ClothAsset
 			WeldByWeightedAverage(WeldingGroups, Cloth.GetWeightMap(WeightMapName));
 		}
 
+
+
 		// Gather list of vertices to remove
 		TArray<int32> VerticesToRemove;
 		for (TMap<int32, int32>::TConstIterator WeldingMapIter = WeldingMap.CreateConstIterator(); WeldingMapIter; ++WeldingMapIter)
