Chaos Cloth - Moved ClothCollection.h to the private ClothAsset folder. Removed all its unnecessary header inclusions.

#rnx

diff --git a/Engine/Plugins/Experimental/ChaosClothAsset/Source/ChaosClothAsset/Private/ChaosClothAsset/ClothCollection.h b/Engine/Plugins/Experimental/ChaosClothAsset/Source/ChaosClothAsset/Private/ChaosClothAsset/ClothCollection.h
new file mode 100644
--- /dev/null
+++ b/Engine/Plugins/Experimental/ChaosClothAsset/Source/ChaosClothAsset/Private/ChaosClothAsset/ClothCollection.h
@@ -0,0 +1,447 @@
+// Copyright Epic Games, Inc. All Rights Reserved.
+
+#pragma once
+
+#include "GeometryCollection/ManagedArray.h"
+
+struct FManagedArrayCollection;
+
+namespace UE::Chaos::ClothAsset
+{
+	/** User defined attribute types. */
+	template<typename T> struct TIsUserAttributeType { static constexpr bool Value = false; };
+	template<> struct TIsUserAttributeType<bool> { static constexpr bool Value = true; };
+	template<> struct TIsUserAttributeType<int32> { static constexpr bool Value = true; };
+	template<> struct TIsUserAttributeType<float> { static constexpr bool Value = true; };
+	template<> struct TIsUserAttributeType<FVector3f> { static constexpr bool Value = true; };
+
+	/**
+	 * Cloth collection facade data.
+	 */
+	class FClothCollection final
+	{
+	public:
+		static const FName LodsGroup;  // LOD information (only one LOD per collection)
+		static const FName SeamsGroup;  // Collection of seam stitches
+		static const FName SeamStitchesGroup;  // Contains pairs of stitched sim vertex indices
+		static const FName SimPatternsGroup;  // Contains sim pattern relationships to other groups
+		static const FName RenderPatternsGroup;  // Contains render pattern relationships to other groups
+		static const FName SimFacesGroup;  // Contains indices to sim vertices
+		static const FName SimVertices2DGroup;  // Contains 2D positions
+		static const FName SimVertices3DGroup;  // Contains 3D positions
+		static const FName RenderFacesGroup;  // Contains indices to render vertex
+		static const FName RenderVerticesGroup;  // Contains 3D render model
+
+		static constexpr int8 MaxNumBoneInfluences = 12; // This should be <= MAX_TOTAL_INFLUENCES defined in GPUSkinPublicDefs.h 
+		static constexpr int8 MaxNumTetherAttachments = 4; // This should be <= FClothTetherDataPrivate::MaxNumAttachments
+
+		explicit FClothCollection(const TSharedPtr<FManagedArrayCollection>& InManagedArrayCollection);
+		~FClothCollection() = default;
+
+		FClothCollection(const FClothCollection&) = delete;
+		FClothCollection& operator=(const FClothCollection&) = delete;
+		FClothCollection(FClothCollection&&) = delete;
+		FClothCollection& operator=(FClothCollection&&) = delete;
+
+		/** Return whether the underlying collection is a valid cloth collection. */
+		bool IsValid() const;
+
+		/** Make the underlying collection a cloth collection. */
+		void DefineSchema();
+
+		/** Get the number of elements of a group. */
+		int32 GetNumElements(const FName& GroupName) const;
+
+		/** Get the number of elements of one of the sub groups that have start/end indices. */
+		int32 GetNumElements(const TManagedArray<int32>* StartArray, const TManagedArray<int32>* EndArray, int32 Index) const;
+
+		/** Set the number of elements of a group. */
+		void SetNumElements(int32 InNumElements, const FName& GroupName);
+
+		/** Set the number of elements to one of the sub groups while maintaining the correct order of the data, and return the first index of the range. */
+		int32 SetNumElements(int32 InNumElements, const FName& GroupName, TManagedArray<int32>* StartArray, TManagedArray<int32>* EndArray, int32 Index);
+
+		/** Remove Elements */
+		void RemoveElements(const FName& Group, const TArray<int32>& SortedDeletionList);
+
+		template<typename T>
+		static inline TConstArrayView<T> GetElements(const TManagedArray<T>* ElementArray, const TManagedArray<int32>* StartArray, const TManagedArray<int32>* EndArray, int32 ArrayIndex);
+
+		template<typename T>
+		static inline TArrayView<T> GetElements(TManagedArray<T>* ElementArray, const TManagedArray<int32>* StartArray, const TManagedArray<int32>* EndArray, int32 ArrayIndex);
+
+		template<typename T>
+		static inline TConstArrayView<T> GetElements(const TManagedArray<T>* ElementArray);
+
+		template<typename T>
+		static inline TArrayView<T> GetElements(TManagedArray<T>* ElementArray);
+
+		/**
+		 * Return the difference between the start index of an element to the start index of the first sub-element in the group (base).
+		 * Usefull for getting back and forth between LOD/pattern indexation modes.
+		 */
+		static int32 GetElementsOffset(const TManagedArray<int32>* StartArray, int32 BaseElementIndex, int32 ElementIndex);
+
+		/**
+		* Return the ArrayIndex (into StartArray/EndArray) which corresponds with the subarray that contains ElementIndex.
+		* Useful for finding which pattern contains a vertex or face index. Returns INDEX_NONE if not found.
+		* NOTE: typically we have a small number of patterns (case where this is used), so just doing a linear search since
+		* a bisection search would get confused by the empty patterns.
+		*/
+		static int32 GetArrayIndexForContainedElement(
+			const TManagedArray<int32>* StartArray,
+			const TManagedArray<int32>* EndArray,
+			int32 ElementIndex);
+
+		static int32 GetNumSubElements(
+			const TManagedArray<int32>* StartArray,
+			const TManagedArray<int32>* EndArray,
+			const TManagedArray<int32>* StartSubArray,
+			const TManagedArray<int32>* EndSubArray,
+			int32 LodIndex);
+
+		template<typename T>
+		static inline TConstArrayView<T> GetSubElements(
+			const TManagedArray<T>* SubElementArray,
+			const TManagedArray<int32>* StartArray,
+			const TManagedArray<int32>* EndArray,
+			const TManagedArray<int32>* StartSubArray,
+			const TManagedArray<int32>* EndSubArray,
+			int32 ArrayIndex);
+
+		template<typename T>
+		static inline TArrayView<T> GetSubElements(
+			TManagedArray<T>* SubElementArray,
+			const TManagedArray<int32>* StartArray,
+			const TManagedArray<int32>* EndArray,
+			const TManagedArray<int32>* StartSubArray,
+			const TManagedArray<int32>* EndSubArray,
+			int32 ArrayIndex);
+
+		template<bool bStart = true, bool bEnd = true>
+		static TTuple<int32, int32> GetSubElementsStartEnd(
+			const TManagedArray<int32>* StartArray,
+			const TManagedArray<int32>* EndArray,
+			const TManagedArray<int32>* StartSubArray,
+			const TManagedArray<int32>* EndSubArray,
+			int32 ArrayIndex);
+
+
+		/** Useful method for copying data between ArrayViews. Putting here for now. */
+		template<typename T>
+		static inline void CopyArrayViewData(const TArrayView<T>& To, const TConstArrayView<T>& From);
+
+		template<typename T>
+		static inline void CopyArrayViewDataAndApplyOffset(const TArrayView<T>& To, const TConstArrayView<T>& From, const T Offset);
+
+		static void CopyArrayViewDataAndApplyOffset(const TArrayView<TArray<int32>>& To, const TConstArrayView<TArray<int32>>& From, const int32 Offset);
+
+		//~ Weight maps
+		template<typename T, TEMPLATE_REQUIRES(TIsUserAttributeType<T>::Value)>
+		TArray<FName> GetUserDefinedAttributeNames(const FName& GroupName) const;
+
+		template<typename T, TEMPLATE_REQUIRES(TIsUserAttributeType<T>::Value)>
+		void AddUserDefinedAttribute(const FName& Name, const FName& GroupName);
+
+		void RemoveUserDefinedAttribute(const FName& Name, const FName& GroupName);
+
+		template<typename T, TEMPLATE_REQUIRES(TIsUserAttributeType<T>::Value)>
+		bool HasUserDefinedAttribute(const FName& Name, const FName& GroupName) const;
+
+		template<typename T, TEMPLATE_REQUIRES(TIsUserAttributeType<T>::Value)>
+		const TManagedArray<T>* GetUserDefinedAttribute(const FName& Name, const FName& GroupName) const;
+
+		template<typename T, TEMPLATE_REQUIRES(TIsUserAttributeType<T>::Value)>
+		TManagedArray<T>* GetUserDefinedAttribute(const FName& Name, const FName& GroupName);
+
+		//~ LODs Group (There should be only one LOD per ClothCollection)
+		const TManagedArray<FString>* GetPhysicsAssetPathName() const { return PhysicsAssetPathName; }
+		const TManagedArray<FString>* GetSkeletonAssetPathName() const { return SkeletonAssetPathName; }
+
+		//~ Seam Group
+		const TManagedArray<int32>* GetSeamStitchStart() const { return SeamStitchStart; }
+		const TManagedArray<int32>* GetSeamStitchEnd() const { return SeamStitchEnd; }
+
+		//~ Seam Stitches Group
+		const TManagedArray<FIntVector2>* GetSeamStitch2DEndIndices() const { return SeamStitch2DEndIndices; }
+		const TManagedArray<int32>* GetSeamStitch3DIndex() const { return SeamStitch3DIndex; }
+
+		//~ Sim Patterns Group
+		const TManagedArray<int32>* GetSimVertices2DStart() const { return SimVertices2DStart; }
+		const TManagedArray<int32>* GetSimVertices2DEnd() const { return SimVertices2DEnd; }
+		const TManagedArray<int32>* GetSimFacesStart() const { return SimFacesStart; }
+		const TManagedArray<int32>* GetSimFacesEnd() const { return SimFacesEnd; }
+
+		//~ Render Patterns Group
+		const TManagedArray<int32>* GetRenderVerticesStart() const { return RenderVerticesStart; }
+		const TManagedArray<int32>* GetRenderVerticesEnd() const { return RenderVerticesEnd; }
+		const TManagedArray<int32>* GetRenderFacesStart() const { return RenderFacesStart; }
+		const TManagedArray<int32>* GetRenderFacesEnd() const { return RenderFacesEnd; }
+		const TManagedArray<FString>* GetRenderMaterialPathName() const { return RenderMaterialPathName; }
+
+		//~ Sim Faces Group
+		const TManagedArray<FIntVector3>* GetSimIndices2D() const { return SimIndices2D; }
+		const TManagedArray<FIntVector3>* GetSimIndices3D() const { return SimIndices3D; }
+
+		//~ Sim Vertices 2D Group
+		const TManagedArray<FVector2f>* GetSimPosition2D() const { return SimPosition2D; }
+		const TManagedArray<int32>* GetSimVertex3DLookup() const { return SimVertex3DLookup; }
+
+		//~ Sim Vertices 3D Group
+		const TManagedArray<FVector3f>* GetSimPosition3D() const { return SimPosition3D; }
+		const TManagedArray<FVector3f>* GetSimNormal() const { return SimNormal; }
+		const TManagedArray<TArray<int32>>* GetSimBoneIndices() const { return SimBoneIndices; }
+		const TManagedArray<TArray<float>>* GetSimBoneWeights() const { return SimBoneWeights; }
+		const TManagedArray<TArray<int32>>* GetTetherKinematicIndex() const { return TetherKinematicIndex; }
+		const TManagedArray<TArray<float>>* GetTetherReferenceLength() const { return TetherReferenceLength; }
+		const TManagedArray<TArray<int32>>* GetSimVertex2DLookup() const { return SimVertex2DLookup; }
+		const TManagedArray<TArray<int32>>* GetSeamStitchLookup() const { return SeamStitchLookup; }
+
+		//~ Render Faces Group
+		const TManagedArray<FIntVector3>* GetRenderIndices() const { return RenderIndices; }
+
+		//~ Render Vertices Group
+		const TManagedArray<FVector3f>* GetRenderPosition() const { return RenderPosition; }
+		const TManagedArray<FVector3f>* GetRenderNormal() const { return RenderNormal; }
+		const TManagedArray<FVector3f>* GetRenderTangentU() const { return RenderTangentU; }
+		const TManagedArray<FVector3f>* GetRenderTangentV() const { return RenderTangentV; }
+		const TManagedArray<TArray<FVector2f>>* GetRenderUVs() const { return RenderUVs; }
+		const TManagedArray<FLinearColor>* GetRenderColor() const { return RenderColor; }
+		const TManagedArray<TArray<int32>>* GetRenderBoneIndices() const { return RenderBoneIndices; }
+		const TManagedArray<TArray<float>>* GetRenderBoneWeights() const { return RenderBoneWeights; }
+
+		//~ LODs Group (There should be only one LOD per ClothCollection)
+		TManagedArray<FString>* GetPhysicsAssetPathName(){ return PhysicsAssetPathName; }
+		TManagedArray<FString>* GetSkeletonAssetPathName(){ return SkeletonAssetPathName; }
+
+		//~ Seam Group
+		TManagedArray<int32>* GetSeamStitchStart() { return SeamStitchStart; }
+		TManagedArray<int32>* GetSeamStitchEnd() { return SeamStitchEnd; }
+
+		//~ Seam Stitches Group
+		TManagedArray<FIntVector2>* GetSeamStitch2DEndIndices() { return SeamStitch2DEndIndices; }
+		TManagedArray<int32>* GetSeamStitch3DIndex() { return SeamStitch3DIndex; }
+
+		//~ Sim Patterns Group
+		TManagedArray<int32>* GetSimVertices2DStart() { return SimVertices2DStart; }
+		TManagedArray<int32>* GetSimVertices2DEnd() { return SimVertices2DEnd; }
+		TManagedArray<int32>* GetSimFacesStart() { return SimFacesStart; }
+		TManagedArray<int32>* GetSimFacesEnd() { return SimFacesEnd; }
+
+		//~ Render Patterns Group
+		TManagedArray<int32>* GetRenderVerticesStart() { return RenderVerticesStart; }
+		TManagedArray<int32>* GetRenderVerticesEnd() { return RenderVerticesEnd; }
+		TManagedArray<int32>* GetRenderFacesStart() { return RenderFacesStart; }
+		TManagedArray<int32>* GetRenderFacesEnd() { return RenderFacesEnd; }
+		TManagedArray<FString>* GetRenderMaterialPathName() { return RenderMaterialPathName; }
+
+		//~ Sim Faces Group
+		TManagedArray<FIntVector3>* GetSimIndices2D() { return SimIndices2D; }
+		TManagedArray<FIntVector3>* GetSimIndices3D() { return SimIndices3D; }
+
+		//~ Sim Vertices 2D Group
+		TManagedArray<FVector2f>* GetSimPosition2D() { return SimPosition2D; }
+		TManagedArray<int32>* GetSimVertex3DLookup() { return SimVertex3DLookup; }
+
+		//~ Sim Vertices 3D Group
+		TManagedArray<FVector3f>* GetSimPosition3D() { return SimPosition3D; }
+		TManagedArray<FVector3f>* GetSimNormal() { return SimNormal; }
+		TManagedArray<TArray<int32>>* GetSimBoneIndices() { return SimBoneIndices; }
+		TManagedArray<TArray<float>>* GetSimBoneWeights() { return SimBoneWeights; }
+		TManagedArray<TArray<int32>>* GetTetherKinematicIndex() { return TetherKinematicIndex; }
+		TManagedArray<TArray<float>>* GetTetherReferenceLength() { return TetherReferenceLength; }
+		TManagedArray<TArray<int32>>* GetSimVertex2DLookup() { return SimVertex2DLookup; }
+		TManagedArray<TArray<int32>>* GetSeamStitchLookup() { return SeamStitchLookup; }
+
+		//~ Render Faces Group
+		TManagedArray<FIntVector3>* GetRenderIndices() { return RenderIndices; }
+
+		//~ Render Vertices Group
+		TManagedArray<FVector3f>* GetRenderPosition() { return RenderPosition; }
+		TManagedArray<FVector3f>* GetRenderNormal() { return RenderNormal; }
+		TManagedArray<FVector3f>* GetRenderTangentU() { return RenderTangentU; }
+		TManagedArray<FVector3f>* GetRenderTangentV() { return RenderTangentV; }
+		TManagedArray<TArray<FVector2f>>* GetRenderUVs() { return RenderUVs; }
+		TManagedArray<FLinearColor>* GetRenderColor() { return RenderColor; }
+		TManagedArray<TArray<int32>>* GetRenderBoneIndices() { return RenderBoneIndices; }
+		TManagedArray<TArray<float>>* GetRenderBoneWeights() { return RenderBoneWeights; }
+
+	private:
+		//~ Cloth collection
+		TSharedPtr<FManagedArrayCollection> ManagedArrayCollection;
+
+		//~ LODs Group
+		TManagedArray<FString>* PhysicsAssetPathName;
+		TManagedArray<FString>* SkeletonAssetPathName;
+
+		//~ Seam Group
+		TManagedArray<int32>* SeamStitchStart;
+		TManagedArray<int32>* SeamStitchEnd;
+
+		//~ Seam Stitches Group
+		TManagedArray<FIntVector2>* SeamStitch2DEndIndices;  // Stitched 2D vertex indices pair
+		TManagedArray<int32>* SeamStitch3DIndex;  // Corresponding stitched 3D vertex
+
+		//~ Sim Patterns Group
+		TManagedArray<int32>* SimVertices2DStart;
+		TManagedArray<int32>* SimVertices2DEnd;
+		TManagedArray<int32>* SimFacesStart;
+		TManagedArray<int32>* SimFacesEnd;
+
+		//~ Render Patterns Group
+		TManagedArray<int32>* RenderVerticesStart;
+		TManagedArray<int32>* RenderVerticesEnd;
+		TManagedArray<int32>* RenderFacesStart;
+		TManagedArray<int32>* RenderFacesEnd;
+		TManagedArray<FString>* RenderMaterialPathName;
+
+		//~ Sim Faces Group
+		TManagedArray<FIntVector3>* SimIndices2D;
+		TManagedArray<FIntVector3>* SimIndices3D;
+
+		//~ Sim Vertices 2D Group
+		TManagedArray<FVector2f>* SimPosition2D;
+		TManagedArray<int32>* SimVertex3DLookup; // Lookup into corresponding 3D vertices
+
+		//~ Sim Vertices 3D Group
+		TManagedArray<FVector3f>* SimPosition3D;
+		TManagedArray<FVector3f>* SimNormal;  // Used for capture, maxdistance, backstop authoring ...etc
+		TManagedArray<TArray<int32>>* SimBoneIndices;
+		TManagedArray<TArray<float>>* SimBoneWeights;
+		TManagedArray<TArray<int32>>* TetherKinematicIndex;
+		TManagedArray<TArray<float>>* TetherReferenceLength;
+		TManagedArray<TArray<int32>>* SimVertex2DLookup; // Lookup into corresponding 2D vertices
+		TManagedArray<TArray<int32>>* SeamStitchLookup; // Lookup into any seam stitches which weld this vertex
+
+		//~ Render Faces Group
+		TManagedArray<FIntVector3>* RenderIndices;
+
+		//~ Render Vertices Group
+		TManagedArray<FVector3f>* RenderPosition;
+		TManagedArray<FVector3f>* RenderNormal;
+		TManagedArray<FVector3f>* RenderTangentU;
+		TManagedArray<FVector3f>* RenderTangentV;
+		TManagedArray<TArray<FVector2f>>* RenderUVs;
+		TManagedArray<FLinearColor>* RenderColor;
+		TManagedArray<TArray<int32>>* RenderBoneIndices;
+		TManagedArray<TArray<float>>* RenderBoneWeights;
+	};
+
+	template<typename T>
+	inline TConstArrayView<T> FClothCollection::GetElements(const TManagedArray<T>* ElementArray, const TManagedArray<int32>* StartArray, const TManagedArray<int32>* EndArray, int32 ArrayIndex)
+	{
+		if (StartArray && EndArray && ElementArray)
+		{
+			const int32 Start = (*StartArray)[ArrayIndex];
+			const int32 End = (*EndArray)[ArrayIndex];
+			if (Start != INDEX_NONE && End != INDEX_NONE)
+			{
+				return TConstArrayView<T>(ElementArray->GetData() + Start, End - Start + 1);
+			}
+			checkf(Start == End, TEXT("Only one boundary of the range is set to INDEX_NONE, when both should."));
+		}
+		return TConstArrayView<T>();
+	}
+
+	template<typename T>
+	inline TArrayView<T> FClothCollection::GetElements(TManagedArray<T>* ElementArray, const TManagedArray<int32>* StartArray, const TManagedArray<int32>* EndArray, int32 ArrayIndex)
+	{
+		if (StartArray && EndArray && ElementArray)
+		{
+			const int32 Start = (*StartArray)[ArrayIndex];
+			const int32 End = (*EndArray)[ArrayIndex];
+			if (Start != INDEX_NONE && End != INDEX_NONE)
+			{
+				return TArrayView<T>(ElementArray->GetData() + Start, End - Start + 1);
+			}
+			checkf(Start == End, TEXT("Only one boundary of the range is set to INDEX_NONE, when both should."));
+		}
+		return TArrayView<T>();
+	}
+
+	template<typename T>
+	inline TConstArrayView<T> FClothCollection::GetElements(const TManagedArray<T>* ElementArray)
+	{
+		if (ElementArray)
+		{
+			return TConstArrayView<T>(ElementArray->GetData(), ElementArray->Num());
+		}
+		return TConstArrayView<T>();
+	}
+
+	template<typename T>
+	inline TArrayView<T> FClothCollection::GetElements(TManagedArray<T>* ElementArray)
+	{
+		if (ElementArray)
+		{
+			return TArrayView<T>(ElementArray->GetData(), ElementArray->Num());
+		}
+		return TArrayView<T>();
+	}
+
+	template<typename T>
+	inline TConstArrayView<T> FClothCollection::GetSubElements(
+		const TManagedArray<T>* SubElementArray,
+		const TManagedArray<int32>* StartArray,
+		const TManagedArray<int32>* EndArray,
+		const TManagedArray<int32>* StartSubArray,
+		const TManagedArray<int32>* EndSubArray,
+		int32 ArrayIndex)
+	{
+		if (StartArray && EndArray && SubElementArray)
+		{
+			const TTuple<int32, int32> StartEnd = GetSubElementsStartEnd(StartArray, EndArray, StartSubArray, EndSubArray, ArrayIndex);
+			const int32 Start = StartEnd.Get<0>();
+			const int32 End = StartEnd.Get<1>();
+			if (Start != INDEX_NONE && End != INDEX_NONE)
+			{
+				return TConstArrayView<T>(SubElementArray->GetData() + Start, End - Start + 1);
+			}
+			checkf(Start == End, TEXT("Only one boundary of the range is set to INDEX_NONE, when both should."));
+		}
+		return TConstArrayView<T>();
+	}
+
+	template<typename T>
+	inline TArrayView<T> FClothCollection::GetSubElements(
+		TManagedArray<T>* SubElementArray,
+		const TManagedArray<int32>* StartArray,
+		const TManagedArray<int32>* EndArray,
+		const TManagedArray<int32>* StartSubArray,
+		const TManagedArray<int32>* EndSubArray,
+		int32 ArrayIndex)
+	{
+		if (StartArray && EndArray && SubElementArray)
+		{
+			const TTuple<int32, int32> StartEnd = GetSubElementsStartEnd(StartArray, EndArray, StartSubArray, EndSubArray, ArrayIndex);
+			const int32 Start = StartEnd.Get<0>();
+			const int32 End = StartEnd.Get<1>();
+			if (Start != INDEX_NONE && End != INDEX_NONE)
+			{
+				return TArrayView<T>(SubElementArray->GetData() + Start, End - Start + 1);
+			}
+			checkf(Start == End, TEXT("Only one boundary of the range is set to INDEX_NONE, when both should."));
+		}
+		return TArrayView<T>();
+	}
+
+	template<typename T>
+	inline void FClothCollection::CopyArrayViewData(const TArrayView<T>& To, const TConstArrayView<T>& From)
+	{
+		check(To.Num() == From.Num());
+		for (int32 Index = 0; Index < To.Num(); ++Index)
+		{
+			To[Index] = From[Index];
+		}
+	}
+
+	template<typename T>
+	inline void FClothCollection::CopyArrayViewDataAndApplyOffset(const TArrayView<T>& To, const TConstArrayView<T>& From, const T Offset)
+	{
+		check(To.Num() == From.Num());
+		for (int32 Index = 0; Index < To.Num(); ++Index)
+		{
+			To[Index] = From[Index] + Offset;
+		}
+	}
+}  // End namespace UE::Chaos::ClothAsset
diff --git a/Engine/Plugins/Experimental/ChaosClothAsset/Source/ChaosClothAsset/Public/ChaosClothAsset/ClothCollection.h b/Engine/Plugins/Experimental/ChaosClothAsset/Source/ChaosClothAsset/Public/ChaosClothAsset/ClothCollection.h
deleted file mode 100644
diff --git a/Engine/Plugins/Experimental/ChaosClothAsset/Source/ChaosClothAsset/Public/ChaosClothAsset/CollectionClothFacade.h b/Engine/Plugins/Experimental/ChaosClothAsset/Source/ChaosClothAsset/Public/ChaosClothAsset/CollectionClothFacade.h
--- a/Engine/Plugins/Experimental/ChaosClothAsset/Source/ChaosClothAsset/Public/ChaosClothAsset/CollectionClothFacade.h
+++ b/Engine/Plugins/Experimental/ChaosClothAsset/Source/ChaosClothAsset/Public/ChaosClothAsset/CollectionClothFacade.h
@@ -13,8 +13,6 @@ class FChaosArchive;
 
 namespace UE::Chaos::ClothAsset
 {
-	class FClothCollection;
-
 	/**
 	 * Cloth Asset collection facade class focused on draping and pattern information.
 	 * Const access (read only) version.
@@ -124,11 +122,11 @@ namespace UE::Chaos::ClothAsset
 			TArray<uint32>& PatternToWeldedIndices, TArray<TArray<int32>>* OptionalWeldedToPatternIndices = nullptr) const;
 
 	protected:
-		TSharedPtr<const FClothCollection> ClothCollection;
+		TSharedPtr<const class FClothCollection> ClothCollection;
 
 		friend class FCollectionClothSeamFacade;
 		friend class FCollectionClothSeamConstFacade;
-		explicit FCollectionClothConstFacade(const TSharedPtr<const FClothCollection>& ClothCollection);
+		explicit FCollectionClothConstFacade(const TSharedPtr<const class FClothCollection>& ClothCollection);
 	};
 
 	/**
@@ -249,12 +247,12 @@ namespace UE::Chaos::ClothAsset
 
 		void SetDefaults();
 
-		TSharedPtr<FClothCollection> GetClothCollection() { return ConstCastSharedPtr<FClothCollection>(ClothCollection); }
+		TSharedPtr<class FClothCollection> GetClothCollection() { return ConstCastSharedPtr<class FClothCollection>(ClothCollection); }
 
 		friend class FCollectionClothSeamFacade;
 		friend class FCollectionClothSimPatternFacade;
 
-		explicit FCollectionClothFacade(const TSharedPtr<FClothCollection>& InClothCollection);
+		explicit FCollectionClothFacade(const TSharedPtr<class FClothCollection>& InClothCollection);
 
 		// These methods are private because they're managed by the FCollectionClothSeamFacade.
 		//~ Sim Vertices 2D Group
diff --git a/Engine/Plugins/Experimental/ChaosClothAsset/Source/ChaosClothAsset/Public/ChaosClothAsset/CollectionClothRenderPatternFacade.h b/Engine/Plugins/Experimental/ChaosClothAsset/Source/ChaosClothAsset/Public/ChaosClothAsset/CollectionClothRenderPatternFacade.h
--- a/Engine/Plugins/Experimental/ChaosClothAsset/Source/ChaosClothAsset/Public/ChaosClothAsset/CollectionClothRenderPatternFacade.h
+++ b/Engine/Plugins/Experimental/ChaosClothAsset/Source/ChaosClothAsset/Public/ChaosClothAsset/CollectionClothRenderPatternFacade.h
@@ -8,8 +8,6 @@ struct FManagedArrayCollection;
 
 namespace UE::Chaos::ClothAsset
 {
-	class FClothCollection;
-
 	/**
 	 * Cloth Asset collection render pattern facade class to access cloth render pattern data.
 	 * Constructed from FCollectionClothConstFacade.
@@ -59,12 +57,12 @@ namespace UE::Chaos::ClothAsset
 	protected:
 		friend class FCollectionClothRenderPatternFacade;  // For other instances access
 		friend class FCollectionClothConstFacade;
-		FCollectionClothRenderPatternConstFacade(const TSharedPtr<const FClothCollection>& InClothCollection, int32 InPatternIndex);
+		FCollectionClothRenderPatternConstFacade(const TSharedPtr<const class FClothCollection>& InClothCollection, int32 InPatternIndex);
 
 		static constexpr int32 GetBaseElementIndex() { return 0; }
 		int32 GetElementIndex() const { return GetBaseElementIndex() + PatternIndex; }
 
-		TSharedPtr<const FClothCollection> ClothCollection;
+		TSharedPtr<const class FClothCollection> ClothCollection;
 		int32 PatternIndex;
 	};
 
@@ -115,10 +113,10 @@ namespace UE::Chaos::ClothAsset
 
 	protected:
 		friend class FCollectionClothFacade;
-		FCollectionClothRenderPatternFacade(const TSharedPtr<FClothCollection>& InClothCollection, int32 InPatternIndex);
+		FCollectionClothRenderPatternFacade(const TSharedPtr<class FClothCollection>& InClothCollection, int32 InPatternIndex);
 
 		void SetDefaults();
 
-		TSharedPtr<FClothCollection> GetClothCollection() { return ConstCastSharedPtr<FClothCollection>(ClothCollection); }
+		TSharedPtr<class FClothCollection> GetClothCollection() { return ConstCastSharedPtr<class FClothCollection>(ClothCollection); }
 	};
 }  // End namespace UE::Chaos::ClothAsset
diff --git a/Engine/Plugins/Experimental/ChaosClothAsset/Source/ChaosClothAsset/Public/ChaosClothAsset/CollectionClothSeamFacade.h b/Engine/Plugins/Experimental/ChaosClothAsset/Source/ChaosClothAsset/Public/ChaosClothAsset/CollectionClothSeamFacade.h
--- a/Engine/Plugins/Experimental/ChaosClothAsset/Source/ChaosClothAsset/Public/ChaosClothAsset/CollectionClothSeamFacade.h
+++ b/Engine/Plugins/Experimental/ChaosClothAsset/Source/ChaosClothAsset/Public/ChaosClothAsset/CollectionClothSeamFacade.h
@@ -4,12 +4,8 @@
 
 #include "Templates/SharedPointer.h"
 
-struct FManagedArrayCollection;
-
 namespace UE::Chaos::ClothAsset
 {
-	class FClothCollection;
-
 	/**
 	 * Cloth Asset collection seam facade class to access cloth seam data.
 	 * Constructed from FCollectionClothConstFacade.
@@ -45,12 +41,12 @@ namespace UE::Chaos::ClothAsset
 	protected:
 		friend class FCollectionClothSeamFacade;  // For other instances access
 		friend class FCollectionClothConstFacade;
-		FCollectionClothSeamConstFacade(const TSharedPtr<const FClothCollection>& InClothCollection, int32 SeamIndex);
+		FCollectionClothSeamConstFacade(const TSharedPtr<const class FClothCollection>& InClothCollection, int32 SeamIndex);
 
 		static constexpr int32 GetBaseElementIndex() { return 0; }
 		int32 GetElementIndex() const { return GetBaseElementIndex() + SeamIndex; }
 
-		TSharedPtr<const FClothCollection> ClothCollection;
+		TSharedPtr<const class FClothCollection> ClothCollection;
 		int32 SeamIndex;
 	};
 
@@ -83,7 +79,7 @@ namespace UE::Chaos::ClothAsset
 
 private:
 		friend class FCollectionClothFacade;
-		FCollectionClothSeamFacade(const TSharedPtr<FClothCollection>& InClothCollection, int32 SeamIndex);
+		FCollectionClothSeamFacade(const TSharedPtr<class FClothCollection>& InClothCollection, int32 SeamIndex);
 
 		void SetDefaults();
 
@@ -93,6 +89,6 @@ private:
 		TArrayView<FIntVector2> GetSeamStitch2DEndIndices();
 		TArrayView<int32> GetSeamStitch3DIndex();
 
-		TSharedPtr<FClothCollection> GetClothCollection() { return ConstCastSharedPtr<FClothCollection>(ClothCollection); }
+		TSharedPtr<class FClothCollection> GetClothCollection() { return ConstCastSharedPtr<class FClothCollection>(ClothCollection); }
 	};
 }  // End namespace UE::Chaos::ClothAsset
diff --git a/Engine/Plugins/Experimental/ChaosClothAsset/Source/ChaosClothAsset/Public/ChaosClothAsset/CollectionClothSimPatternFacade.h b/Engine/Plugins/Experimental/ChaosClothAsset/Source/ChaosClothAsset/Public/ChaosClothAsset/CollectionClothSimPatternFacade.h
--- a/Engine/Plugins/Experimental/ChaosClothAsset/Source/ChaosClothAsset/Public/ChaosClothAsset/CollectionClothSimPatternFacade.h
+++ b/Engine/Plugins/Experimental/ChaosClothAsset/Source/ChaosClothAsset/Public/ChaosClothAsset/CollectionClothSimPatternFacade.h
@@ -4,12 +4,8 @@
 
 #include "Templates/SharedPointer.h"
 
-struct FManagedArrayCollection;
-
 namespace UE::Chaos::ClothAsset
 {
-	class FClothCollection;
-
 	/**
 	 * Cloth Asset collection sim pattern facade class to access cloth sim pattern data.
 	 * Constructed from FCollectionClothConstFacade.
@@ -51,12 +47,12 @@ namespace UE::Chaos::ClothAsset
 	protected:
 		friend class FCollectionClothSimPatternFacade;  // For other instances access
 		friend class FCollectionClothConstFacade;
-		FCollectionClothSimPatternConstFacade(const TSharedPtr<const FClothCollection>& InClothCollection, int32 InPatternIndex);
+		FCollectionClothSimPatternConstFacade(const TSharedPtr<const class FClothCollection>& InClothCollection, int32 InPatternIndex);
 
 		static constexpr int32 GetBaseElementIndex() { return 0; }
 		int32 GetElementIndex() const { return GetBaseElementIndex() + PatternIndex; }
 
-		TSharedPtr<const FClothCollection> ClothCollection;
+		TSharedPtr<const class FClothCollection> ClothCollection;
 		int32 PatternIndex;
 	};
 
@@ -105,7 +101,7 @@ namespace UE::Chaos::ClothAsset
 
 	private:
 		friend class FCollectionClothFacade;
-		FCollectionClothSimPatternFacade(const TSharedPtr<FClothCollection>& InClothCollection, int32 InPatternIndex);
+		FCollectionClothSimPatternFacade(const TSharedPtr<class FClothCollection>& InClothCollection, int32 InPatternIndex);
 
 		void SetDefaults();
 
@@ -115,6 +111,6 @@ namespace UE::Chaos::ClothAsset
 		void SetNumSimVertices2D(int32 NumSimVertices);
 		TArrayView<int32> GetSimVertex3DLookup();
 
-		TSharedPtr<FClothCollection> GetClothCollection() { return ConstCastSharedPtr<FClothCollection>(ClothCollection); }
+		TSharedPtr<class FClothCollection> GetClothCollection() { return ConstCastSharedPtr<class FClothCollection>(ClothCollection); }
 	};
 }  // End namespace UE::Chaos::ClothAsset
diff --git a/Engine/Plugins/Experimental/ChaosClothAsset/Source/ChaosClothAssetEngine/Private/ChaosClothAsset/ClothEngineTools.cpp b/Engine/Plugins/Experimental/ChaosClothAsset/Source/ChaosClothAssetEngine/Private/ChaosClothAsset/ClothEngineTools.cpp
--- a/Engine/Plugins/Experimental/ChaosClothAsset/Source/ChaosClothAssetEngine/Private/ChaosClothAsset/ClothEngineTools.cpp
+++ b/Engine/Plugins/Experimental/ChaosClothAsset/Source/ChaosClothAssetEngine/Private/ChaosClothAsset/ClothEngineTools.cpp
@@ -2,7 +2,6 @@
 
 #include "ChaosClothAsset/ClothEngineTools.h"
 #include "ChaosClothAsset/CollectionClothFacade.h"
-#include "ChaosClothAsset/ClothCollection.h"
 #include "ChaosClothAsset/ClothGeometryTools.h"
 #include "ClothTetherData.h"
 
diff --git a/Engine/Plugins/Experimental/ChaosClothAssetEditor/Source/ChaosClothAssetDataflowNodes/Private/ChaosClothAsset/ClothDataflowTools.cpp b/Engine/Plugins/Experimental/ChaosClothAssetEditor/Source/ChaosClothAssetDataflowNodes/Private/ChaosClothAsset/ClothDataflowTools.cpp
--- a/Engine/Plugins/Experimental/ChaosClothAssetEditor/Source/ChaosClothAssetDataflowNodes/Private/ChaosClothAsset/ClothDataflowTools.cpp
+++ b/Engine/Plugins/Experimental/ChaosClothAssetEditor/Source/ChaosClothAssetDataflowNodes/Private/ChaosClothAsset/ClothDataflowTools.cpp
@@ -3,7 +3,6 @@
 #include "ChaosClothAsset/ClothDataflowTools.h"
 #include "ChaosClothAsset/ClothGeometryTools.h"
 #include "ChaosClothAsset/CollectionClothFacade.h"
-#include "ChaosClothAsset/ClothCollection.h"
 #include "Animation/Skeleton.h"
 #include "BoneWeights.h"
 #include "DynamicMesh/NonManifoldMappingSupport.h"
diff --git a/Engine/Plugins/Experimental/ChaosClothAssetEditor/Source/ChaosClothAssetDataflowNodes/Private/ChaosClothAsset/StaticMeshImportNode.cpp b/Engine/Plugins/Experimental/ChaosClothAssetEditor/Source/ChaosClothAssetDataflowNodes/Private/ChaosClothAsset/StaticMeshImportNode.cpp
--- a/Engine/Plugins/Experimental/ChaosClothAssetEditor/Source/ChaosClothAssetDataflowNodes/Private/ChaosClothAsset/StaticMeshImportNode.cpp
+++ b/Engine/Plugins/Experimental/ChaosClothAssetEditor/Source/ChaosClothAssetDataflowNodes/Private/ChaosClothAsset/StaticMeshImportNode.cpp
@@ -4,7 +4,6 @@
 #include "ChaosClothAsset/DataflowNodes.h"
 #include "ChaosClothAsset/CollectionClothFacade.h"
 #include "ChaosClothAsset/ClothAsset.h"
-#include "ChaosClothAsset/ClothCollection.h"
 #include "ChaosClothAsset/ClothGeometryTools.h"
 #include "ChaosClothAsset/ClothDataflowTools.h"
 #include "Engine/StaticMesh.h"
diff --git a/Engine/Plugins/Experimental/ChaosClothAssetEditor/Source/ChaosClothAssetEditor/Private/ChaosClothAsset/ClothEditorToolkit.cpp b/Engine/Plugins/Experimental/ChaosClothAssetEditor/Source/ChaosClothAssetEditor/Private/ChaosClothAsset/ClothEditorToolkit.cpp
--- a/Engine/Plugins/Experimental/ChaosClothAssetEditor/Source/ChaosClothAssetEditor/Private/ChaosClothAsset/ClothEditorToolkit.cpp
+++ b/Engine/Plugins/Experimental/ChaosClothAssetEditor/Source/ChaosClothAssetEditor/Private/ChaosClothAsset/ClothEditorToolkit.cpp
@@ -13,7 +13,6 @@
 #include "ChaosClothAsset/ClothComponent.h"
 #include "ChaosClothAsset/ClothAsset.h"
 #include "ChaosClothAsset/ClothEditorModeToolkit.h"
-#include "ChaosClothAsset/ClothCollection.h"
 #include "ChaosClothAsset/CollectionClothFacade.h"
 #include "ChaosClothAsset/ClothEditorCommands.h"
 #include "ChaosClothAsset/AddWeightMapNode.h"
diff --git a/Engine/Plugins/Experimental/ChaosClothAssetEditor/Source/ChaosClothAssetEditorTools/Private/ChaosClothAsset/ClothTrainingTool.cpp b/Engine/Plugins/Experimental/ChaosClothAssetEditor/Source/ChaosClothAssetEditorTools/Private/ChaosClothAsset/ClothTrainingTool.cpp
--- a/Engine/Plugins/Experimental/ChaosClothAssetEditor/Source/ChaosClothAssetEditorTools/Private/ChaosClothAsset/ClothTrainingTool.cpp
+++ b/Engine/Plugins/Experimental/ChaosClothAssetEditor/Source/ChaosClothAssetEditorTools/Private/ChaosClothAsset/ClothTrainingTool.cpp
@@ -12,7 +12,6 @@
 #include "ChaosClothAsset/ClothDataGenerationComponent.h"
 #include "ChaosClothAsset/ClothComponent.h"
 #include "ChaosClothAsset/ClothComponentToolTarget.h"
-#include "ChaosClothAsset/ClothCollection.h"
 #include "ChaosClothAsset/ClothSimulationProxy.h"
 #include "ClothingSystemRuntimeTypes.h"
 #include "ClothTrainingToolProperties.h"
diff --git a/Engine/Plugins/Experimental/ChaosClothAssetEditor/Source/ChaosClothAssetTools/Private/ChaosClothAsset/ClothPatternToDynamicMesh.cpp b/Engine/Plugins/Experimental/ChaosClothAssetEditor/Source/ChaosClothAssetTools/Private/ChaosClothAsset/ClothPatternToDynamicMesh.cpp
--- a/Engine/Plugins/Experimental/ChaosClothAssetEditor/Source/ChaosClothAssetTools/Private/ChaosClothAsset/ClothPatternToDynamicMesh.cpp
+++ b/Engine/Plugins/Experimental/ChaosClothAssetEditor/Source/ChaosClothAssetTools/Private/ChaosClothAsset/ClothPatternToDynamicMesh.cpp
@@ -4,7 +4,6 @@
 
 #include "ChaosClothAsset/ClothPatternToDynamicMesh.h"
 #include "ChaosClothAsset/ClothAsset.h"
-#include "ChaosClothAsset/ClothCollection.h"
 #include "ChaosClothAsset/ClothComponent.h"
 #include "ChaosClothAsset/CollectionClothFacade.h"
 #include "ToDynamicMesh.h"
