ClothAsset: compute LOD transition up and  down information

diff --git a/Engine/Plugins/Experimental/ChaosClothAsset/Source/ChaosClothAssetEngine/Private/ChaosClothAsset/ClothSimulationMesh.cpp b/Engine/Plugins/Experimental/ChaosClothAsset/Source/ChaosClothAssetEngine/Private/ChaosClothAsset/ClothSimulationMesh.cpp
--- a/Engine/Plugins/Experimental/ChaosClothAsset/Source/ChaosClothAssetEngine/Private/ChaosClothAsset/ClothSimulationMesh.cpp
+++ b/Engine/Plugins/Experimental/ChaosClothAsset/Source/ChaosClothAssetEngine/Private/ChaosClothAsset/ClothSimulationMesh.cpp
@@ -173,13 +173,11 @@ PRAGMA_ENABLE_DEPRECATION_WARNINGS
 
 	TConstArrayView<FMeshToMeshVertData> FClothSimulationMesh::GetTransitionUpSkinData(int32 LODIndex) const
 	{
-		// TODO: UpSkinData
-		return TConstArrayView<FMeshToMeshVertData>();
+		return ClothSimulationModel.ClothSimulationLodModels[LODIndex].LODTransitionUpData;
 	}
 
 	TConstArrayView<FMeshToMeshVertData> FClothSimulationMesh::GetTransitionDownSkinData(int32 LODIndex) const
 	{
-		// TODO: DownSkinData
-		return TConstArrayView<FMeshToMeshVertData>();
+		return ClothSimulationModel.ClothSimulationLodModels[LODIndex].LODTransitionDownData;
 	}
 }
diff --git a/Engine/Plugins/Experimental/ChaosClothAsset/Source/ChaosClothAssetEngine/Private/ChaosClothAsset/ClothSimulationModel.cpp b/Engine/Plugins/Experimental/ChaosClothAsset/Source/ChaosClothAssetEngine/Private/ChaosClothAsset/ClothSimulationModel.cpp
--- a/Engine/Plugins/Experimental/ChaosClothAsset/Source/ChaosClothAssetEngine/Private/ChaosClothAsset/ClothSimulationModel.cpp
+++ b/Engine/Plugins/Experimental/ChaosClothAsset/Source/ChaosClothAssetEngine/Private/ChaosClothAsset/ClothSimulationModel.cpp
@@ -3,6 +3,7 @@
 #include "ChaosClothAsset/ClothSimulationModel.h"
 #include "ChaosClothAsset/CollectionClothFacade.h"
 #include "ChaosClothAsset/ClothGeometryTools.h"
+#include "Utils/ClothingMeshUtils.h"
 #include "ReferenceSkeleton.h"
 
 namespace UE::Chaos::ClothAsset
@@ -118,6 +119,9 @@ bool FChaosClothSimulationLodModel::Serialize(FArchive& Ar)
 	// Serialize weight maps (not a tagged property)
 	Ar << WeightMaps;
 
+	Ar << LODTransitionUpData;
+	Ar << LODTransitionDownData;
+
 	// Return true to confirm that serialization has already been taken care of
 	return true;
 }
@@ -193,6 +197,41 @@ FChaosClothSimulationModel::FChaosClothSimulationModel(const TArray<TSharedPtr<c
 
 	// Initialize Reference bone index
 	ReferenceBoneIndex = CalculateReferenceBoneIndex(ClothSimulationLodModels, ReferenceSkeleton, UsedBoneIndices);
+
+	// Compute LOD Transition data
+	
+	// Parameters for GenerateMeshToMeshVertData
+	const FPointWeightMap* const MaxDistances = nullptr;  // No need to update the vertex contribution on the transition maps
+	constexpr bool bUseSmoothTransitions = false;         // Smooth transitions are only used at rendering for now and not during LOD transitions
+	constexpr bool bUseMultipleInfluences = false;        // Multiple influences must not be used for LOD transitions
+	constexpr float SkinningKernelRadius = 0.f;           // KernelRadius is only required when using multiple influences
+
+	for (int32 LODIndex = 0; LODIndex < NumLods; ++LODIndex)
+	{
+		const int32 PrevLODIndex = LODIndex - 1;
+		const int32 NextLODIndex = LODIndex + 1;
+
+		FChaosClothSimulationLodModel& CurrModel = ClothSimulationLodModels[LODIndex];
+		const ClothingMeshUtils::ClothMeshDesc CurrentMeshDesc(CurrModel.Positions, CurrModel.Indices);
+
+		if (ClothSimulationLodModels.IsValidIndex(PrevLODIndex))
+		{
+			const FChaosClothSimulationLodModel& PrevModel = ClothSimulationLodModels[PrevLODIndex];
+			const ClothingMeshUtils::ClothMeshDesc PrevMeshDesc(PrevModel.Positions, PrevModel.Indices);
+
+			ClothingMeshUtils::GenerateMeshToMeshVertData(CurrModel.LODTransitionUpData, CurrentMeshDesc, PrevMeshDesc,
+				MaxDistances, bUseSmoothTransitions, bUseMultipleInfluences, SkinningKernelRadius);
+		}
+
+		if (ClothSimulationLodModels.IsValidIndex(NextLODIndex))
+		{
+			const FChaosClothSimulationLodModel& NextModel = ClothSimulationLodModels[NextLODIndex];
+			const ClothingMeshUtils::ClothMeshDesc NextMeshDesc(NextModel.Positions, NextModel.Indices);
+
+			ClothingMeshUtils::GenerateMeshToMeshVertData(CurrModel.LODTransitionDownData, CurrentMeshDesc, NextMeshDesc,
+				MaxDistances, bUseSmoothTransitions, bUseMultipleInfluences, SkinningKernelRadius);
+		}
+	}
 }
 
 TArray<TConstArrayView<TTuple<int32, int32, float>>> FChaosClothSimulationModel::GetTethers(int32 LODIndex) const
diff --git a/Engine/Plugins/Experimental/ChaosClothAsset/Source/ChaosClothAssetEngine/Private/ChaosClothAsset/ClothSimulationModel.h b/Engine/Plugins/Experimental/ChaosClothAsset/Source/ChaosClothAssetEngine/Private/ChaosClothAsset/ClothSimulationModel.h
--- a/Engine/Plugins/Experimental/ChaosClothAsset/Source/ChaosClothAssetEngine/Private/ChaosClothAsset/ClothSimulationModel.h
+++ b/Engine/Plugins/Experimental/ChaosClothAsset/Source/ChaosClothAssetEngine/Private/ChaosClothAsset/ClothSimulationModel.h
@@ -7,7 +7,7 @@
 #include "Templates/SharedPointer.h"
 #include "ClothVertBoneData.h"
 #include "ClothTetherData.h"
-
+#include "SkeletalMeshTypes.h"
 #include "ClothSimulationModel.generated.h"
 
 struct FReferenceSkeleton;
@@ -38,6 +38,10 @@ struct FChaosClothSimulationLodModel
 	UPROPERTY()
 	TArray<FClothVertBoneData> BoneData;
 
+	/** LOD Transition mesh to mesh skinning weights. */
+	TArray<FMeshToMeshVertData> LODTransitionUpData;
+	TArray<FMeshToMeshVertData> LODTransitionDownData;
+
 	/** 2d pattern positions. */
 	UPROPERTY()
 	TArray<FVector2f> PatternPositions;
