Chaos Cloth Asset - seam splitting when removing seams.

Duplicate 3D points as necessary when removing seams. Final point count and mesh topology will be preserved if you create and remove a seam, but vertex indexing may change. Also there will be data loss as maps are averaged on seaming, for example.

diff --git a/Engine/Plugins/Experimental/ChaosClothAsset/Source/ChaosClothAsset/Private/ChaosClothAsset/CollectionClothSeamFacade.cpp b/Engine/Plugins/Experimental/ChaosClothAsset/Source/ChaosClothAsset/Private/ChaosClothAsset/CollectionClothSeamFacade.cpp
--- a/Engine/Plugins/Experimental/ChaosClothAsset/Source/ChaosClothAsset/Private/ChaosClothAsset/CollectionClothSeamFacade.cpp
+++ b/Engine/Plugins/Experimental/ChaosClothAsset/Source/ChaosClothAsset/Private/ChaosClothAsset/CollectionClothSeamFacade.cpp
@@ -10,6 +10,9 @@ namespace UE::Chaos::ClothAsset
 {
 	namespace Private
 	{
+		//
+		// Welding
+		//
 		typedef TMap<int32, int32> FWeldingGroup; // Key = Index, Value = Weight;
 
 		static int32 WeldingMappedValue(const TMap<int32, int32>& WeldingMap, int32 Index)
@@ -261,6 +264,203 @@ namespace UE::Chaos::ClothAsset
 				[](const TPair<float, int32>& A, const TPair<float, int32>& B) { return A < B; });
 		}
 
+		//
+		// Splitting
+		//
+		struct FSplittingGroup
+		{
+			int32 NewVertex3DIndex = INDEX_NONE;
+			TSet<int32> Vertex2Ds;
+			TSet<int32> SeamStitches;
+		};
+		typedef TMap<int32, TArray<FSplittingGroup>> FSplittingGroups; // Key = Original 3D vertex
+
+		static void UpdateSplittingGroup(TArray<FSplittingGroup>& SplittingGroup, const int32 VertexIndex3D, const TArray<int32>& RemainingStitches, const TConstArrayView<FIntVector2>& SeamStitch2DEndIndices,
+			const TArray<int32>& SimVertex2DLookup)
+		{				
+			// Form groups of 2D vertices that should stay welded by remaining stitches.
+			for (const int32 RemainingStitch : RemainingStitches)
+			{
+				if (RemainingStitch != INDEX_NONE)
+				{
+					const FIntVector2& Global2DEndPoints = SeamStitch2DEndIndices[RemainingStitch];
+					int32 FirstExistingGroup = INDEX_NONE;
+					for(int32 GrpIdx = 0; GrpIdx < SplittingGroup.Num(); ++GrpIdx)
+					{
+						FSplittingGroup& Group = SplittingGroup[GrpIdx];
+						if (Group.Vertex2Ds.Contains(Global2DEndPoints[0]) || Group.Vertex2Ds.Contains(Global2DEndPoints[1]))
+						{
+							// Make sure both ends are in this group
+							Group.Vertex2Ds.Add(Global2DEndPoints[0]);
+							Group.Vertex2Ds.Add(Global2DEndPoints[1]);
+							Group.SeamStitches.Add(RemainingStitch);
+
+							if (FirstExistingGroup != INDEX_NONE)
+							{
+								// This stitch connects the two groups. Merge them.
+								FSplittingGroup& OtherGroup = SplittingGroup[FirstExistingGroup];
+								Group.SeamStitches.Append(OtherGroup.SeamStitches);
+								Group.Vertex2Ds.Append(OtherGroup.Vertex2Ds);
+								SplittingGroup.RemoveAtSwap(FirstExistingGroup);
+								break;
+							}
+							else
+							{
+								FirstExistingGroup = GrpIdx;
+							}
+						}
+					}
+					if (FirstExistingGroup == INDEX_NONE)
+					{
+						// Add new group with this seam.
+						FSplittingGroup& NewGroup = SplittingGroup.AddDefaulted_GetRef();
+						NewGroup.Vertex2Ds.Add(Global2DEndPoints[0]);
+						NewGroup.Vertex2Ds.Add(Global2DEndPoints[1]);
+						NewGroup.SeamStitches.Add(RemainingStitch);
+					}
+				}
+			}
+
+			// Create new singleton groups for any un-stitched 2D vertices that correspond with this 3D vertex
+			for (const int32 Vertex2D : SimVertex2DLookup)
+			{
+				if (Vertex2D != INDEX_NONE)
+				{
+					bool bFoundExistingGroup = false;
+					for (FSplittingGroup& Group : SplittingGroup)
+					{
+						if (Group.Vertex2Ds.Contains(Vertex2D))
+						{
+							bFoundExistingGroup = true;
+							break;
+						}
+					}
+					if (!bFoundExistingGroup)
+					{
+						SplittingGroup.AddDefaulted_GetRef().Vertex2Ds.Add(Vertex2D);
+					}
+				}
+			}
+		}
+
+		static void UpdateVertexLookupsAfterSplitting(const FSplittingGroups& SplittingGroups, TArrayView<int32> SimVertex3DLookup, TArrayView<TArray<int32>> SimVertex2DLookup)
+		{
+			for (FSplittingGroups::TConstIterator GroupIter = SplittingGroups.CreateConstIterator(); GroupIter; ++GroupIter)
+			{
+				for (const FSplittingGroup& Group : GroupIter.Value())
+				{
+					SimVertex2DLookup[Group.NewVertex3DIndex] = Group.Vertex2Ds.Array();
+					for (const int32 Vertex2D : SimVertex2DLookup[Group.NewVertex3DIndex])
+					{
+						SimVertex3DLookup[Vertex2D] = Group.NewVertex3DIndex;
+					}
+				}
+			}
+		}
+
+		static void UpdateSeamStitchLookupsAfterSplitting(const FSplittingGroups& SplittingGroups, TArrayView<int32> SeamStitch3DIndex, TArrayView<TArray<int32>> SeamStitchLookup)
+		{
+			for (FSplittingGroups::TConstIterator GroupIter = SplittingGroups.CreateConstIterator(); GroupIter; ++GroupIter)
+			{
+				for (const FSplittingGroup& Group : GroupIter.Value())
+				{
+					SeamStitchLookup[Group.NewVertex3DIndex] = Group.SeamStitches.Array();
+					for (const int32 SeamStitch : SeamStitchLookup[Group.NewVertex3DIndex])
+					{
+						SeamStitch3DIndex[SeamStitch] = Group.NewVertex3DIndex;
+					}
+				}
+			}
+		}
+
+		// Simple splitting by just duplicating data to newly split vertices
+		template<typename T>
+		static void SplitCopyVertexData(const FSplittingGroups& SplittingGroups, TArrayView<T> Values)
+		{
+			for (FSplittingGroups::TConstIterator GroupIter = SplittingGroups.CreateConstIterator(); GroupIter; ++GroupIter)
+			{
+				const T& OrigValue = Values[GroupIter.Key()];
+				for (int32 GroupIdx = 1; GroupIdx < GroupIter.Value().Num(); ++GroupIdx)
+				{
+					Values[GroupIter.Value()[GroupIdx].NewVertex3DIndex] = OrigValue;
+				}
+			}
+		}
+
+		static void SplitTethers(const FSplittingGroups& SplittingGroups, TArrayView<TArray<int32>> TetherKinematicIndices, TArrayView<TArray<float>> TetherReferenceLengths)
+		{
+			// Copy kinematic indices. 
+			check(TetherKinematicIndices.Num() == TetherReferenceLengths.Num());
+			const int32 NumVertices = TetherKinematicIndices.Num();
+			for (int32 VertexIndex = 0; VertexIndex < NumVertices; ++VertexIndex)
+			{
+				TArray<int32>& Indices = TetherKinematicIndices[VertexIndex];
+				TArray<float>& Lengths = TetherReferenceLengths[VertexIndex];
+
+				check(Indices.Num() == Lengths.Num());
+				const int32 NumTethers = Indices.Num();
+				// Go in reverse because we're going to remove any invalid tethers while we're here.
+				for (int32 TetherIndex = NumTethers - 1; TetherIndex >= 0; --TetherIndex)
+				{
+					if (Indices[TetherIndex] == INDEX_NONE)
+					{
+						Indices.RemoveAtSwap(TetherIndex);
+						Lengths.RemoveAtSwap(TetherIndex);
+						continue;
+					}
+					
+					if (const TArray<FSplittingGroup>* SplittingGroup = SplittingGroups.Find(Indices[TetherIndex]))
+					{
+						check(SplittingGroup->Num() > 1);
+						// Need to duplicate this tether.
+						Indices.Reserve(Indices.Num() + SplittingGroup->Num() - 1);
+						Lengths.Reserve(Indices.Num() + SplittingGroup->Num() - 1);
+						const float Length = Lengths[TetherIndex];
+						for (int32 GroupIdx = 1; GroupIdx < SplittingGroup->Num(); ++GroupIdx)
+						{
+							Indices.Add((*SplittingGroup)[GroupIdx].NewVertex3DIndex);
+							Lengths.Add(Length);
+						}
+					}
+				}
+
+				if (Indices.Num() > FClothCollection::MaxNumTetherAttachments)
+				{
+					// Remove farthest tethers
+					TArray<TPair<float, int32>> SortableData;
+					SortableData.Reserve(Indices.Num());
+					for (int32 Idx = 0; Idx < Indices.Num(); ++Idx)
+					{
+						SortableData.Emplace(Lengths[Idx], Indices[Idx]);
+					}
+					SortableData.Sort([](const TPair<float, int32>& A, const TPair<float, int32>& B) { return A < B; });
+					Indices.SetNum(FClothCollection::MaxNumTetherAttachments);
+					Lengths.SetNum(FClothCollection::MaxNumTetherAttachments);
+					for (int32 Idx = 0; Idx < FClothCollection::MaxNumTetherAttachments; ++Idx)
+					{
+						Indices[Idx] = SortableData[Idx].Get<1>();
+						Lengths[Idx] = SortableData[Idx].Get<0>();
+					}
+				}
+			}
+
+			// Now split dynamic indices
+			SplitCopyVertexData(SplittingGroups, TetherKinematicIndices);
+			SplitCopyVertexData(SplittingGroups, TetherReferenceLengths);
+		}
+
+		static void UpdateSimIndicesAfterSplitting(const TConstArrayView<FIntVector3>& SimIndices2D, const TConstArrayView<int32>& SimVertex3DLookup, TArrayView<FIntVector3> SimIndices3D)
+		{
+			// Since we don't have a vertex -> face lookup, it's faster to just regenerate all of the sim indices
+			check(SimIndices2D.Num() == SimIndices3D.Num());
+			for (int32 FaceIndex = 0; FaceIndex < SimIndices3D.Num(); ++FaceIndex)
+			{
+				SimIndices3D[FaceIndex][0] = SimVertex3DLookup[SimIndices2D[FaceIndex][0]];
+				SimIndices3D[FaceIndex][1] = SimVertex3DLookup[SimIndices2D[FaceIndex][1]];
+				SimIndices3D[FaceIndex][2] = SimVertex3DLookup[SimIndices2D[FaceIndex][2]];
+			}
+		}
+
 	} // namespace Private
 
 	int32 FCollectionClothSeamConstFacade::GetNumSeamStitches() const
@@ -308,16 +508,106 @@ namespace UE::Chaos::ClothAsset
 
 	void FCollectionClothSeamFacade::Reset()
 	{
-		FCollectionClothFacade Cloth(GetClothCollection());
+		using namespace Private;
+
+		const int32 OrigNumSeamStitches = GetNumSeamStitches();
+
 		// Split all seams by duplicating points.
-		// TODO: there is no way to remove seams without removing the associated vertices right now.
-#if DO_CHECK // TODO: switch to GUARD_SLOW 
-		const TConstArrayView<int32> SeamStitch3DIndex = GetSeamStitch3DIndex();
-		for (int32 Idx : SeamStitch3DIndex)
+		if (OrigNumSeamStitches > 0)
 		{
-			check(Idx == INDEX_NONE);
+			FCollectionClothFacade Cloth(GetClothCollection());
+
+			const int32 GlobalSeamStitchesOffset = GetSeamStitchesOffset();
+			const TArrayView<int32> SeamStitch3DIndex = GetSeamStitch3DIndex();
+
+			const TArrayView<TArray<int32>> SeamStitchLookup = Cloth.GetSeamStitchLookupPrivate();
+			const TConstArrayView<FIntVector2> GlobalSeamStitch2DEndIndices = ClothCollection->GetElements(ClothCollection->GetSeamStitch2DEndIndices());
+			const TConstArrayView<TArray<int32>> SimVertex2DLookup = Cloth.GetSimVertex2DLookup();
+
+			const int32 OrigNumSimVertices3D = Cloth.GetNumSimVertices3D();
+			int32 NumNewSimVertices3D = 0;
+
+			FSplittingGroups SplittingGroups;
+
+			// Gather all 3D vertices that might need to be split
+			for (int32 StitchIndex = 0; StitchIndex < OrigNumSeamStitches; ++StitchIndex)
+			{
+				const int32 VertexIndex3D = SeamStitch3DIndex[StitchIndex];
+				if (VertexIndex3D != INDEX_NONE)
+				{
+					// Break this stitch
+					verify(SeamStitchLookup[VertexIndex3D].RemoveSwap(StitchIndex + GlobalSeamStitchesOffset) == 1);
+
+					// Add this to the splitting groups to be processed
+					SplittingGroups.FindOrAdd(VertexIndex3D);
+				}
+			}
+
+			// Process splitting groups
+			for (FSplittingGroups::TIterator GroupIter = SplittingGroups.CreateIterator(); GroupIter; ++GroupIter)
+			{
+				const int32 VertexIndex3D = GroupIter.Key();
+
+				UpdateSplittingGroup(GroupIter.Value(), GroupIter.Key(), SeamStitchLookup[VertexIndex3D], GlobalSeamStitch2DEndIndices, SimVertex2DLookup[VertexIndex3D]);
+
+				check(GroupIter.Value().Num() > 0);
+
+				if (GroupIter.Value().Num() == 1)
+				{
+					// This vertex actually doesn't need to split.
+					GroupIter.RemoveCurrent();
+					continue;
+				}
+
+				GroupIter.Value()[0].NewVertex3DIndex = VertexIndex3D; // First group will stay at the original 3D index.
+				for (int32 GroupIdx = 1; GroupIdx < GroupIter.Value().Num(); ++GroupIdx)
+				{
+					GroupIter.Value()[GroupIdx].NewVertex3DIndex = OrigNumSimVertices3D + NumNewSimVertices3D++;
+				}
+			}
+
+			if (NumNewSimVertices3D > 0)
+			{	
+				//
+				// Add the 3D vertices and split all of their data.
+				//
+
+				// This resize will invalidate any existing SimVertex3DGroup ArrayViews!
+				GetClothCollection()->SetNumElements(OrigNumSimVertices3D + NumNewSimVertices3D, FClothCollection::SimVertices3DGroup);
+
+				// Update 2D <--> 3D lookups
+				UpdateVertexLookupsAfterSplitting(SplittingGroups, Cloth.GetSimVertex3DLookupPrivate(), Cloth.GetSimVertex2DLookupPrivate());
+
+				// Update Stitch <-> 3D vertex lookups for stitches in other seams.
+				UpdateSeamStitchLookupsAfterSplitting(SplittingGroups, GetClothCollection()->GetElements(GetClothCollection()->GetSeamStitch3DIndex()), Cloth.GetSeamStitchLookupPrivate());
+
+				// Split 3D positions
+				SplitCopyVertexData(SplittingGroups, Cloth.GetSimPosition3D());
+
+				// Split normals
+				SplitCopyVertexData(SplittingGroups, Cloth.GetSimNormal());
+
+				// Split BoneIndices
+				SplitCopyVertexData(SplittingGroups, Cloth.GetSimBoneIndices());
+
+				// Split BoneWeights
+				SplitCopyVertexData(SplittingGroups, Cloth.GetSimBoneWeights());
+
+				// Split Tethers
+				SplitTethers(SplittingGroups, Cloth.GetTetherKinematicIndex(), Cloth.GetTetherReferenceLength());
+
+				// Split Faces
+				// The easiest thing to do here is actually to regenerate the 3D faces from the 2D faces.
+				UpdateSimIndicesAfterSplitting(Cloth.GetSimIndices2D(), Cloth.GetSimVertex3DLookup(), Cloth.GetSimIndices3D());
+
+				// Split Maps
+				const TArray<FName> WeightMapNames = Cloth.GetWeightMapNames();
+				for (const FName& WeightMapName : WeightMapNames)
+				{
+					SplitCopyVertexData(SplittingGroups, Cloth.GetWeightMap(WeightMapName));
+				}
+			}
 		}
-#endif
 		
 		SetNumSeamStitches(0);
 		SetDefaults();
