Cloth Asset: Compactify/clean up sim meshes

diff --git a/Engine/Plugins/Experimental/ChaosClothAsset/Source/ChaosClothAsset/Private/ChaosClothAsset/ClothCollection.cpp b/Engine/Plugins/Experimental/ChaosClothAsset/Source/ChaosClothAsset/Private/ChaosClothAsset/ClothCollection.cpp
--- a/Engine/Plugins/Experimental/ChaosClothAsset/Source/ChaosClothAsset/Private/ChaosClothAsset/ClothCollection.cpp
+++ b/Engine/Plugins/Experimental/ChaosClothAsset/Source/ChaosClothAsset/Private/ChaosClothAsset/ClothCollection.cpp
@@ -472,6 +472,46 @@ namespace UE::Chaos::ClothAsset
 		ManagedArrayCollection->RemoveElements(Group, SortedDeletionList);
 	}
 
+	void FClothCollection::RemoveElements(const FName& GroupName, const TArray<int32>& SortedDeletionList, TManagedArray<int32>* StartArray, TManagedArray<int32>* EndArray, int32 ArrayIndex)
+	{
+		if (SortedDeletionList.IsEmpty())
+		{
+			return;
+		}
+
+		check(IsValid());
+
+		check(StartArray && StartArray->GetConstArray().IsValidIndex(ArrayIndex));
+		check(EndArray && EndArray->GetConstArray().IsValidIndex(ArrayIndex));
+
+		int32& Start = (*StartArray)[ArrayIndex];
+		int32& End = (*EndArray)[ArrayIndex];
+		check(Start != INDEX_NONE && End != INDEX_NONE);
+
+		const int32 OrigStart = Start;
+		const int32 OrigNumElements = End - Start + 1;
+
+		check(SortedDeletionList[0] >= Start);
+		check(SortedDeletionList.Last() <= End);
+		check(OrigNumElements >= SortedDeletionList.Num());
+
+		ManagedArrayCollection->RemoveElements(GroupName, SortedDeletionList);
+
+		if (SortedDeletionList.Num() == OrigNumElements)
+		{
+			Start = End = INDEX_NONE;
+		}
+		else
+		{
+			const int32 NewNumElements = OrigNumElements - SortedDeletionList.Num();
+			const int32 NewEnd = OrigStart + NewNumElements - 1;
+			check(Start == OrigStart || Start == INDEX_NONE);
+			check(End == NewEnd || End == INDEX_NONE);
+			Start = OrigStart;
+			End = NewEnd;
+		}
+	}
+
 	int32 FClothCollection::GetElementsOffset(const TManagedArray<int32>* StartArray, int32 BaseElementIndex, int32 ElementIndex)
 	{
 		while ((*StartArray)[BaseElementIndex] == INDEX_NONE && BaseElementIndex < ElementIndex)
diff --git a/Engine/Plugins/Experimental/ChaosClothAsset/Source/ChaosClothAsset/Private/ChaosClothAsset/ClothCollection.h b/Engine/Plugins/Experimental/ChaosClothAsset/Source/ChaosClothAsset/Private/ChaosClothAsset/ClothCollection.h
--- a/Engine/Plugins/Experimental/ChaosClothAsset/Source/ChaosClothAsset/Private/ChaosClothAsset/ClothCollection.h
+++ b/Engine/Plugins/Experimental/ChaosClothAsset/Source/ChaosClothAsset/Private/ChaosClothAsset/ClothCollection.h
@@ -64,6 +64,9 @@ namespace UE::Chaos::ClothAsset
 		/** Remove Elements */
 		void RemoveElements(const FName& Group, const TArray<int32>& SortedDeletionList);
 
+		/** Remove Elements. SortedDeletionList should be global indices. */
+		void RemoveElements(const FName& GroupName, const TArray<int32>& SortedDeletionList, TManagedArray<int32>* StartArray, TManagedArray<int32>* EndArray, int32 Index);
+
 		template<typename T>
 		static inline TConstArrayView<T> GetElements(const TManagedArray<T>* ElementArray, const TManagedArray<int32>* StartArray, const TManagedArray<int32>* EndArray, int32 ArrayIndex);
 
diff --git a/Engine/Plugins/Experimental/ChaosClothAsset/Source/ChaosClothAsset/Private/ChaosClothAsset/ClothGeometryTools.cpp b/Engine/Plugins/Experimental/ChaosClothAsset/Source/ChaosClothAsset/Private/ChaosClothAsset/ClothGeometryTools.cpp
--- a/Engine/Plugins/Experimental/ChaosClothAsset/Source/ChaosClothAsset/Private/ChaosClothAsset/ClothGeometryTools.cpp
+++ b/Engine/Plugins/Experimental/ChaosClothAsset/Source/ChaosClothAsset/Private/ChaosClothAsset/ClothGeometryTools.cpp
@@ -771,4 +771,190 @@ namespace UE::Chaos::ClothAsset
 			AddSeam(Cloth, Seam);
 		}
 	}
+
+	void FClothGeometryTools::CleanupAndCompactMesh(const TSharedPtr<FManagedArrayCollection>& ClothCollection)
+	{
+		FCollectionClothFacade Cloth(ClothCollection);
+
+		TArray<int32> SimPatternsToRemove;
+		for (int32 PatternIndex = 0; PatternIndex < Cloth.GetNumSimPatterns(); ++PatternIndex)
+		{
+			FCollectionClothSimPatternFacade Pattern = Cloth.GetSimPattern(PatternIndex);
+			{
+				// Remove any triangles that are topologically degenerate
+				TArray<int32> FacesToRemove;
+				TConstArrayView<FIntVector3> SimIndices3D = Pattern.GetSimIndices3D();
+				TConstArrayView<FIntVector3> SimIndices2D = Pattern.GetSimIndices2D();
+				for (int32 FaceIndex = 0; FaceIndex < SimIndices3D.Num(); ++FaceIndex)
+				{
+					if (SimIndices3D[FaceIndex][0] == INDEX_NONE ||
+						SimIndices3D[FaceIndex][1] == INDEX_NONE ||
+						SimIndices3D[FaceIndex][2] == INDEX_NONE ||
+
+						SimIndices2D[FaceIndex][0] == INDEX_NONE ||
+						SimIndices2D[FaceIndex][1] == INDEX_NONE ||
+						SimIndices2D[FaceIndex][2] == INDEX_NONE ||
+
+						SimIndices3D[FaceIndex][0] == SimIndices3D[FaceIndex][1] ||
+						SimIndices3D[FaceIndex][0] == SimIndices3D[FaceIndex][2] ||
+						SimIndices3D[FaceIndex][1] == SimIndices3D[FaceIndex][2] ||
+
+						SimIndices2D[FaceIndex][0] == SimIndices2D[FaceIndex][1] ||
+						SimIndices2D[FaceIndex][0] == SimIndices2D[FaceIndex][2] ||
+						SimIndices2D[FaceIndex][1] == SimIndices2D[FaceIndex][2])
+					{
+						FacesToRemove.Add(FaceIndex);
+					}
+				}
+
+				if (FacesToRemove.Num())
+				{
+					Pattern.RemoveSimFaces(FacesToRemove);
+				}
+			}
+			{
+				// Remove any 2D vertices that are not used in a face.
+				TConstArrayView<FIntVector3> SimIndices2D = Pattern.GetSimIndices2D();
+				const int32 SimVertex2DOffset = Pattern.GetSimVertices2DOffset();
+				TBitArray SimVertex2DToRemove;
+				SimVertex2DToRemove.Init(true, Pattern.GetNumSimVertices2D());
+				for (const FIntVector3& Face : SimIndices2D)
+				{
+					SimVertex2DToRemove[Face[0] - SimVertex2DOffset] = false;
+					SimVertex2DToRemove[Face[1] - SimVertex2DOffset] = false;
+					SimVertex2DToRemove[Face[2] - SimVertex2DOffset] = false;
+				}
+
+				TArray<int32> SimVertex2DToRemoveList;
+				for (TConstSetBitIterator ToRemoveIter(SimVertex2DToRemove); ToRemoveIter; ++ToRemoveIter)
+				{
+					SimVertex2DToRemoveList.Add(ToRemoveIter.GetIndex());
+				}
+
+				if (SimVertex2DToRemoveList.Num())
+				{
+					Pattern.RemoveSimVertices2D(SimVertex2DToRemoveList);
+				}
+			}
+
+			if (Pattern.IsEmpty())
+			{
+				SimPatternsToRemove.Add(PatternIndex);
+			}
+		}
+		if (!SimPatternsToRemove.IsEmpty())
+		{
+			Cloth.RemoveSimPatterns(SimPatternsToRemove);
+		}
+
+		// Remove any unused 3D vertices
+		{
+			TConstArrayView<FIntVector3> SimIndices3D = Cloth.GetSimIndices3D();
+			TBitArray SimVertex3DToRemove;
+			SimVertex3DToRemove.Init(true, Cloth.GetNumSimVertices3D());
+			for (const FIntVector3& Face : SimIndices3D)
+			{
+				SimVertex3DToRemove[Face[0]] = false;
+				SimVertex3DToRemove[Face[1]] = false;
+				SimVertex3DToRemove[Face[2]] = false;
+			}
+
+			TArray<int32> SimVertex3DToRemoveList;
+			for (TConstSetBitIterator ToRemoveIter(SimVertex3DToRemove); ToRemoveIter; ++ToRemoveIter)
+			{
+				SimVertex3DToRemoveList.Add(ToRemoveIter.GetIndex());
+			}
+
+			if (SimVertex3DToRemoveList.Num())
+			{
+				Cloth.RemoveSimVertices3D(SimVertex3DToRemoveList);
+			}
+		}
+		{
+			// Clean up any references to vertices that no longer exist.
+			// NOTE: should not need to clean up 2D vertices pointing to INDEX_NONE 3D vertices since this should have
+			// meant the 2D vertex either was unused in the faces, or was associated with an invalid face (it should already be cleaned up).
+			// NOTE: not cleaning up seam references for now as they aren't exported as part of the final mesh/sim data.
+			Cloth.CompactSimVertex2DLookup();
+
+			TArrayView<TArray<int32>> TetherKinematicIndex = Cloth.GetTetherKinematicIndex();
+			TArrayView<TArray<float>> TetherReferenceLength = Cloth.GetTetherReferenceLength();
+			const int32 NumVertices = TetherKinematicIndex.Num();
+			for (int32 VertexIdx = 0; VertexIdx < NumVertices; ++VertexIdx)
+			{
+				const int32 NumTethers = TetherKinematicIndex[VertexIdx].Num();
+				for (int32 TetherIdx = NumTethers - 1; TetherIdx >= 0; --TetherIdx)
+				{
+					if (TetherKinematicIndex[VertexIdx][TetherIdx] == INDEX_NONE)
+					{
+						TetherKinematicIndex[VertexIdx].RemoveAtSwap(TetherIdx);
+						TetherReferenceLength[VertexIdx].RemoveAtSwap(TetherIdx);
+					}
+				}
+			}
+		}
+
+		TArray<int32> RenderPatternsToRemove;
+		for (int32 PatternIndex = 0; PatternIndex < Cloth.GetNumRenderPatterns(); ++PatternIndex)
+		{
+			FCollectionClothRenderPatternFacade Pattern = Cloth.GetRenderPattern(PatternIndex);
+			{
+				// Remove any triangles that are topologically degenerate
+				TArray<int32> FacesToRemove;
+				TConstArrayView<FIntVector3> RenderIndices = Pattern.GetRenderIndices();
+				for (int32 FaceIndex = 0; FaceIndex < RenderIndices.Num(); ++FaceIndex)
+				{
+					if (RenderIndices[FaceIndex][0] == INDEX_NONE ||
+						RenderIndices[FaceIndex][1] == INDEX_NONE ||
+						RenderIndices[FaceIndex][2] == INDEX_NONE ||
+
+						RenderIndices[FaceIndex][0] == RenderIndices[FaceIndex][1] ||
+						RenderIndices[FaceIndex][0] == RenderIndices[FaceIndex][2] ||
+						RenderIndices[FaceIndex][1] == RenderIndices[FaceIndex][2])
+					{
+						FacesToRemove.Add(FaceIndex);
+					}
+				}
+
+				if (FacesToRemove.Num())
+				{
+					Pattern.RemoveRenderFaces(FacesToRemove);
+				}
+			}
+
+			{
+				// Remove any vertices that are not used in a face.
+				TConstArrayView<FIntVector3> RenderIndices = Pattern.GetRenderIndices();
+				const int32 RenderVertexOffset = Pattern.GetRenderVerticesOffset();
+				TBitArray RenderVertexToRemove;
+				RenderVertexToRemove.Init(true, Pattern.GetNumRenderVertices());
+				for (const FIntVector3& Face : RenderIndices)
+				{
+					RenderVertexToRemove[Face[0] - RenderVertexOffset] = false;
+					RenderVertexToRemove[Face[1] - RenderVertexOffset] = false;
+					RenderVertexToRemove[Face[2] - RenderVertexOffset] = false;
+				}
+
+				TArray<int32> RenderVertexToRemoveList;
+				for (TConstSetBitIterator ToRemoveIter(RenderVertexToRemove); ToRemoveIter; ++ToRemoveIter)
+				{
+					RenderVertexToRemoveList.Add(ToRemoveIter.GetIndex());
+				}
+
+				if (RenderVertexToRemoveList.Num())
+				{
+					Pattern.RemoveRenderVertices(RenderVertexToRemoveList);
+				}
+			}
+
+			if (Pattern.IsEmpty())
+			{
+				RenderPatternsToRemove.Add(PatternIndex);
+			}
+		}
+		if (!RenderPatternsToRemove.IsEmpty())
+		{
+			Cloth.RemoveRenderPatterns(RenderPatternsToRemove);
+		}
+	}
 }  // End namespace UE::Chaos::ClothAsset
diff --git a/Engine/Plugins/Experimental/ChaosClothAsset/Source/ChaosClothAsset/Private/ChaosClothAsset/CollectionClothFacade.cpp b/Engine/Plugins/Experimental/ChaosClothAsset/Source/ChaosClothAsset/Private/ChaosClothAsset/CollectionClothFacade.cpp
--- a/Engine/Plugins/Experimental/ChaosClothAsset/Source/ChaosClothAsset/Private/ChaosClothAsset/CollectionClothFacade.cpp
+++ b/Engine/Plugins/Experimental/ChaosClothAsset/Source/ChaosClothAsset/Private/ChaosClothAsset/CollectionClothFacade.cpp
@@ -453,6 +453,26 @@ namespace UE::Chaos::ClothAsset
 		GetClothCollection()->SetNumElements(NumSimVertices - InNumSimVertices, FClothCollection::SimVertices3DGroup);
 	}
 
+	void FCollectionClothFacade::RemoveSimVertices3D(const TArray<int32>& SortedDeletionList)
+	{
+		GetClothCollection()->RemoveElements(FClothCollection::SimVertices3DGroup, SortedDeletionList);
+	}
+
+	void FCollectionClothFacade::CompactSimVertex2DLookup()
+	{
+		TArrayView<TArray<int32>> SimVertex2DLookup = GetSimVertex2DLookupPrivate();
+		for (TArray<int32>& VertexLookup : SimVertex2DLookup)
+		{
+			for (int32 Idx = VertexLookup.Num() - 1; Idx >= 0; --Idx)
+			{
+				if (VertexLookup[Idx] == INDEX_NONE)
+				{
+					VertexLookup.RemoveAtSwap(Idx);
+				}
+			}
+		}
+	}
+
 	TArrayView<FIntVector3> FCollectionClothFacade::GetSimIndices2D()
 	{
 		return GetClothCollection()->GetElements(GetClothCollection()->GetSimIndices2D());
@@ -492,6 +512,16 @@ namespace UE::Chaos::ClothAsset
 		return FCollectionClothSimPatternFacade(GetClothCollection(), PatternIndex);
 	}
 
+	void FCollectionClothFacade::RemoveSimPatterns(const TArray<int32>& SortedDeletionList)
+	{
+		for (const int32 PatternToRemove : SortedDeletionList)
+		{
+			GetSimPattern(PatternToRemove).Reset();
+		}
+
+		GetClothCollection()->RemoveElements(FClothCollection::SimPatternsGroup, SortedDeletionList);
+	}
+
 	void FCollectionClothFacade::SetNumRenderPatterns(int32 InNumPatterns)
 	{
 		const int32 NumPatterns = GetNumRenderPatterns();
@@ -521,6 +551,16 @@ namespace UE::Chaos::ClothAsset
 		return FCollectionClothRenderPatternFacade(GetClothCollection(), PatternIndex);
 	}
 
+	void FCollectionClothFacade::RemoveRenderPatterns(const TArray<int32>& SortedDeletionList)
+	{
+		for (const int32 PatternToRemove : SortedDeletionList)
+		{
+			GetRenderPattern(PatternToRemove).Reset();
+		}
+
+		GetClothCollection()->RemoveElements(FClothCollection::RenderPatternsGroup, SortedDeletionList);
+	}
+
 	TArrayView<FString> FCollectionClothFacade::GetRenderMaterialPathName()
 	{
 		return GetClothCollection()->GetElements(GetClothCollection()->GetRenderMaterialPathName());
@@ -555,6 +595,15 @@ namespace UE::Chaos::ClothAsset
 		return FCollectionClothSeamFacade(GetClothCollection(), SeamIndex);
 	}
 
+	void FCollectionClothFacade::RemoveSeams(const TArray<int32>& SortedDeletionList)
+	{
+		for (const int32 SeamToRemove : SortedDeletionList)
+		{
+			GetSeam(SeamToRemove).Reset();
+		}
+		GetClothCollection()->RemoveElements(FClothCollection::SeamsGroup, SortedDeletionList);
+	}
+
 	//~ Render Vertices Group
 	TArrayView<FVector3f> FCollectionClothFacade::GetRenderPosition()
 	{
diff --git a/Engine/Plugins/Experimental/ChaosClothAsset/Source/ChaosClothAsset/Private/ChaosClothAsset/CollectionClothRenderPatternFacade.cpp b/Engine/Plugins/Experimental/ChaosClothAsset/Source/ChaosClothAsset/Private/ChaosClothAsset/CollectionClothRenderPatternFacade.cpp
--- a/Engine/Plugins/Experimental/ChaosClothAsset/Source/ChaosClothAsset/Private/ChaosClothAsset/CollectionClothRenderPatternFacade.cpp
+++ b/Engine/Plugins/Experimental/ChaosClothAsset/Source/ChaosClothAsset/Private/ChaosClothAsset/CollectionClothRenderPatternFacade.cpp
@@ -126,6 +126,11 @@ namespace UE::Chaos::ClothAsset
 			GetElementIndex());
 	}
 
+	bool FCollectionClothRenderPatternConstFacade::IsEmpty() const
+	{
+		return GetNumRenderVertices() == 0 && GetNumRenderFaces() == 0;
+	}
+
 	FCollectionClothRenderPatternConstFacade::FCollectionClothRenderPatternConstFacade(const TSharedPtr<const FClothCollection>& ClothCollection, int32 InPatternIndex)
 		: ClothCollection(ClothCollection)
 		, PatternIndex(InPatternIndex)
@@ -179,6 +184,25 @@ namespace UE::Chaos::ClothAsset
 			GetClothCollection()->GetRenderVerticesEnd(),
 			GetElementIndex());
 	}
+
+	void FCollectionClothRenderPatternFacade::RemoveRenderVertices(const TArray<int32>& SortedDeletionList)
+	{
+		TArray<int32> GlobalIndexSortedDeletionList;
+		const int32 Offset = GetRenderVerticesOffset();
+		GlobalIndexSortedDeletionList.SetNumUninitialized(SortedDeletionList.Num());
+		for (int32 Idx = 0; Idx < SortedDeletionList.Num(); ++Idx)
+		{
+			GlobalIndexSortedDeletionList[Idx] = SortedDeletionList[Idx] + Offset;
+		}
+
+		GetClothCollection()->RemoveElements(
+			FClothCollection::RenderVerticesGroup,
+			GlobalIndexSortedDeletionList,
+			GetClothCollection()->GetRenderVerticesStart(),
+			GetClothCollection()->GetRenderVerticesEnd(),
+			GetElementIndex());
+	}
+
 	TArrayView<FVector3f> FCollectionClothRenderPatternFacade::GetRenderPosition()
 	{
 		return GetClothCollection()->GetElements(
@@ -261,6 +285,24 @@ namespace UE::Chaos::ClothAsset
 			GetElementIndex());
 	}
 
+	void FCollectionClothRenderPatternFacade::RemoveRenderFaces(const TArray<int32>& SortedDeletionList)
+	{
+		TArray<int32> GlobalIndexSortedDeletionList;
+		const int32 Offset = GetRenderFacesOffset();
+		GlobalIndexSortedDeletionList.SetNumUninitialized(SortedDeletionList.Num());
+		for (int32 Idx = 0; Idx < SortedDeletionList.Num(); ++Idx)
+		{
+			GlobalIndexSortedDeletionList[Idx] = SortedDeletionList[Idx] + Offset;
+		}
+
+		GetClothCollection()->RemoveElements(
+			FClothCollection::RenderFacesGroup,
+			GlobalIndexSortedDeletionList,
+			GetClothCollection()->GetRenderFacesStart(),
+			GetClothCollection()->GetRenderFacesEnd(),
+			GetElementIndex());
+	}
+
 	TArrayView<FIntVector3> FCollectionClothRenderPatternFacade::GetRenderIndices()
 	{
 		return GetClothCollection()->GetElements(
diff --git a/Engine/Plugins/Experimental/ChaosClothAsset/Source/ChaosClothAsset/Private/ChaosClothAsset/CollectionClothSeamFacade.cpp b/Engine/Plugins/Experimental/ChaosClothAsset/Source/ChaosClothAsset/Private/ChaosClothAsset/CollectionClothSeamFacade.cpp
--- a/Engine/Plugins/Experimental/ChaosClothAsset/Source/ChaosClothAsset/Private/ChaosClothAsset/CollectionClothSeamFacade.cpp
+++ b/Engine/Plugins/Experimental/ChaosClothAsset/Source/ChaosClothAsset/Private/ChaosClothAsset/CollectionClothSeamFacade.cpp
@@ -455,9 +455,9 @@ namespace UE::Chaos::ClothAsset
 			check(SimIndices2D.Num() == SimIndices3D.Num());
 			for (int32 FaceIndex = 0; FaceIndex < SimIndices3D.Num(); ++FaceIndex)
 			{
-				SimIndices3D[FaceIndex][0] = SimVertex3DLookup[SimIndices2D[FaceIndex][0]];
-				SimIndices3D[FaceIndex][1] = SimVertex3DLookup[SimIndices2D[FaceIndex][1]];
-				SimIndices3D[FaceIndex][2] = SimVertex3DLookup[SimIndices2D[FaceIndex][2]];
+				SimIndices3D[FaceIndex][0] = SimIndices2D[FaceIndex][0] == INDEX_NONE ? INDEX_NONE : SimVertex3DLookup[SimIndices2D[FaceIndex][0]];
+				SimIndices3D[FaceIndex][1] = SimIndices2D[FaceIndex][1] == INDEX_NONE ? INDEX_NONE : SimVertex3DLookup[SimIndices2D[FaceIndex][1]];
+				SimIndices3D[FaceIndex][2] = SimIndices2D[FaceIndex][2] == INDEX_NONE ? INDEX_NONE : SimVertex3DLookup[SimIndices2D[FaceIndex][2]];
 			}
 		}
 
diff --git a/Engine/Plugins/Experimental/ChaosClothAsset/Source/ChaosClothAsset/Private/ChaosClothAsset/CollectionClothSimPatternFacade.cpp b/Engine/Plugins/Experimental/ChaosClothAsset/Source/ChaosClothAsset/Private/ChaosClothAsset/CollectionClothSimPatternFacade.cpp
--- a/Engine/Plugins/Experimental/ChaosClothAsset/Source/ChaosClothAsset/Private/ChaosClothAsset/CollectionClothSimPatternFacade.cpp
+++ b/Engine/Plugins/Experimental/ChaosClothAsset/Source/ChaosClothAsset/Private/ChaosClothAsset/CollectionClothSimPatternFacade.cpp
@@ -75,6 +75,11 @@ namespace UE::Chaos::ClothAsset
 			GetElementIndex());
 	}
 
+	bool FCollectionClothSimPatternConstFacade::IsEmpty() const
+	{
+		return GetNumSimFaces() == 0 && GetNumSimVertices2D() == 0;
+	}
+
 	FCollectionClothSimPatternConstFacade::FCollectionClothSimPatternConstFacade(const TSharedPtr<const FClothCollection>& ClothCollection, int32 InPatternIndex)
 		: ClothCollection(ClothCollection)
 		, PatternIndex(InPatternIndex)
@@ -219,6 +224,24 @@ namespace UE::Chaos::ClothAsset
 		SetNumSimVertices2D(NumSimVertices - InNumSimVertices);
 	}
 
+	void FCollectionClothSimPatternFacade::RemoveSimVertices2D(const TArray<int32>& SortedDeletionList)
+	{
+		TArray<int32> GlobalIndexSortedDeletionList;
+		const int32 Offset = GetSimVertices2DOffset();
+		GlobalIndexSortedDeletionList.SetNumUninitialized(SortedDeletionList.Num());
+		for (int32 Idx = 0; Idx < SortedDeletionList.Num(); ++Idx)
+		{
+			GlobalIndexSortedDeletionList[Idx] = SortedDeletionList[Idx] + Offset;
+		}
+
+		GetClothCollection()->RemoveElements(
+			FClothCollection::SimVertices2DGroup,
+			GlobalIndexSortedDeletionList,
+			GetClothCollection()->GetSimVertices2DStart(),
+			GetClothCollection()->GetSimVertices2DEnd(),
+			GetElementIndex());
+	}
+
 	void FCollectionClothSimPatternFacade::SetNumSimFaces(int32 NumSimFaces)
 	{
 		GetClothCollection()->SetNumElements(
@@ -247,6 +270,24 @@ namespace UE::Chaos::ClothAsset
 			GetElementIndex());
 	}
 
+	void FCollectionClothSimPatternFacade::RemoveSimFaces(const TArray<int32>& SortedDeletionList)
+	{
+		TArray<int32> GlobalIndexSortedDeletionList;
+		const int32 Offset = GetSimFacesOffset();
+		GlobalIndexSortedDeletionList.SetNumUninitialized(SortedDeletionList.Num());
+		for (int32 Idx = 0; Idx < SortedDeletionList.Num(); ++Idx)
+		{
+			GlobalIndexSortedDeletionList[Idx] = SortedDeletionList[Idx] + Offset;
+		}
+
+		GetClothCollection()->RemoveElements(
+			FClothCollection::SimFacesGroup,
+			GlobalIndexSortedDeletionList,
+			GetClothCollection()->GetSimFacesStart(),
+			GetClothCollection()->GetSimFacesEnd(),
+			GetElementIndex());
+	}
+
 	FCollectionClothSimPatternFacade::FCollectionClothSimPatternFacade(const TSharedPtr<FClothCollection>& ClothCollection, int32 InPatternIndex)
 		: FCollectionClothSimPatternConstFacade(ClothCollection, InPatternIndex)
 	{
diff --git a/Engine/Plugins/Experimental/ChaosClothAsset/Source/ChaosClothAsset/Public/ChaosClothAsset/ClothGeometryTools.h b/Engine/Plugins/Experimental/ChaosClothAsset/Source/ChaosClothAsset/Public/ChaosClothAsset/ClothGeometryTools.h
--- a/Engine/Plugins/Experimental/ChaosClothAsset/Source/ChaosClothAsset/Public/ChaosClothAsset/ClothGeometryTools.h
+++ b/Engine/Plugins/Experimental/ChaosClothAsset/Source/ChaosClothAsset/Public/ChaosClothAsset/ClothGeometryTools.h
@@ -59,5 +59,11 @@ namespace UE::Chaos::ClothAsset
 		static void BuildSimMeshFromDynamicMesh(
 			const TSharedPtr<FManagedArrayCollection>& ClothCollection,
 			const UE::Geometry::FDynamicMesh3& DynamicMesh, int32 UVChannelIndex, const FVector2f& UVScale, bool bAppend);
+
+		/**
+		* Remove (topologically) degenerate triangles. Remove any vertices that aren't in a triangle. Compact any lookup arrays that contain INDEX_NONEs.
+		* Remove any empty patterns.
+		*/
+		static void CleanupAndCompactMesh(const TSharedPtr<FManagedArrayCollection>& ClothCollection);
 	};
 }  // End namespace UE::Chaos::ClothAsset
\ No newline at end of file
diff --git a/Engine/Plugins/Experimental/ChaosClothAsset/Source/ChaosClothAsset/Public/ChaosClothAsset/CollectionClothFacade.h b/Engine/Plugins/Experimental/ChaosClothAsset/Source/ChaosClothAsset/Public/ChaosClothAsset/CollectionClothFacade.h
--- a/Engine/Plugins/Experimental/ChaosClothAsset/Source/ChaosClothAsset/Public/ChaosClothAsset/CollectionClothFacade.h
+++ b/Engine/Plugins/Experimental/ChaosClothAsset/Source/ChaosClothAsset/Public/ChaosClothAsset/CollectionClothFacade.h
@@ -183,6 +183,9 @@ namespace UE::Chaos::ClothAsset
 		/** This will remove the 3D vertices, but the associated seams and 2D vertices will still exist, and point to INDEX_NONE */
 		void RemoveSimVertices3D(int32 NumSimVertices);
 		void RemoveAllSimVertices3D() { RemoveSimVertices3D(GetNumSimVertices3D()); }
+		void RemoveSimVertices3D(const TArray<int32>& SortedDeletionList);
+		/** Compact SimVertex2DLookup to remove any references to INDEX_NONE that may have been created by deleting 2D vertices. */
+		void CompactSimVertex2DLookup();
 
 		//~ Pattern Sim Faces Group
 		/** SetNumSimFaces per pattern within pattern facade. */
@@ -198,6 +201,8 @@ namespace UE::Chaos::ClothAsset
 		FCollectionClothSimPatternFacade GetSimPattern(int32 PatternIndex);
 		/** Add a new pattern to this cloth LOD, and return the cloth pattern facade set to its index. */
 		FCollectionClothSimPatternFacade AddGetSimPattern() { return GetSimPattern(AddSimPattern()); }
+		/** Remove a sorted list of sim patterns. */
+		void RemoveSimPatterns(const TArray<int32>& SortedDeletionList);
 
 		//~ Render Patterns Group
 		/** Set the new number of patterns to this cloth LOD. */
@@ -208,6 +213,8 @@ namespace UE::Chaos::ClothAsset
 		FCollectionClothRenderPatternFacade GetRenderPattern(int32 PatternIndex);
 		/** Add a new pattern to this cloth LOD, and return the cloth pattern facade set to its index. */
 		FCollectionClothRenderPatternFacade AddGetRenderPattern() { return GetRenderPattern(AddRenderPattern()); }
+		/** Remove a sorted list of render patterns. */
+		void RemoveRenderPatterns(const TArray<int32>& SortedDeletionList);
 		/** Return a view of all the render materials used on this collection across all patterns. */
 		TArrayView<FString> GetRenderMaterialPathName();
 
@@ -220,6 +227,8 @@ namespace UE::Chaos::ClothAsset
 		FCollectionClothSeamFacade GetSeam(int32 SeamIndex);
 		/** Add a new seam to this cloth and return the seam facade set to its index. */
 		FCollectionClothSeamFacade AddGetSeam() { return GetSeam(AddSeam()); }
+		/** Remove a sorted list of seams. */
+		void RemoveSeams(const TArray<int32>& SortedDeletionList);
 
 		//~ Render Vertices Group
 		/** SetNumRenderVertices per pattern within pattern facade. */
diff --git a/Engine/Plugins/Experimental/ChaosClothAsset/Source/ChaosClothAsset/Public/ChaosClothAsset/CollectionClothRenderPatternFacade.h b/Engine/Plugins/Experimental/ChaosClothAsset/Source/ChaosClothAsset/Public/ChaosClothAsset/CollectionClothRenderPatternFacade.h
--- a/Engine/Plugins/Experimental/ChaosClothAsset/Source/ChaosClothAsset/Public/ChaosClothAsset/CollectionClothRenderPatternFacade.h
+++ b/Engine/Plugins/Experimental/ChaosClothAsset/Source/ChaosClothAsset/Public/ChaosClothAsset/CollectionClothRenderPatternFacade.h
@@ -51,6 +51,8 @@ namespace UE::Chaos::ClothAsset
 		int32 GetRenderFacesOffset() const;
 		TConstArrayView<FIntVector3> GetRenderIndices() const;
 
+		bool IsEmpty() const;
+
 		/** Return the Pattern index this facade has been created with. */
 		int32 GetPatternIndex() const { return PatternIndex; }
 
@@ -97,6 +99,7 @@ namespace UE::Chaos::ClothAsset
 		// Note: Use the FCollectionClothFacade accessors instead of these for the array indices to match the RenderIndices values
 		/** Grow or shrink the space reserved for render vertices for this pattern within the cloth collection and return its start index. */
 		void SetNumRenderVertices(int32 NumRenderVertices);
+		void RemoveRenderVertices(const TArray<int32>& SortedDeletionList);
 		TArrayView<FVector3f> GetRenderPosition();
 		TArrayView<FVector3f> GetRenderNormal();
 		TArrayView<FVector3f> GetRenderTangentU();
@@ -109,6 +112,7 @@ namespace UE::Chaos::ClothAsset
 		//~ Render Faces Group
 		/** Grow or shrink the space reserved for render faces for this pattern within the cloth collection and return its start index. */
 		void SetNumRenderFaces(int32 NumRenderFaces);
+		void RemoveRenderFaces(const TArray<int32>& SortedDeletionList);
 		TArrayView<FIntVector3> GetRenderIndices();
 
 	protected:
diff --git a/Engine/Plugins/Experimental/ChaosClothAsset/Source/ChaosClothAsset/Public/ChaosClothAsset/CollectionClothSimPatternFacade.h b/Engine/Plugins/Experimental/ChaosClothAsset/Source/ChaosClothAsset/Public/ChaosClothAsset/CollectionClothSimPatternFacade.h
--- a/Engine/Plugins/Experimental/ChaosClothAsset/Source/ChaosClothAsset/Public/ChaosClothAsset/CollectionClothSimPatternFacade.h
+++ b/Engine/Plugins/Experimental/ChaosClothAsset/Source/ChaosClothAsset/Public/ChaosClothAsset/CollectionClothSimPatternFacade.h
@@ -41,6 +41,9 @@ namespace UE::Chaos::ClothAsset
 		TConstArrayView<FIntVector3> GetSimIndices2D() const;
 		TConstArrayView<FIntVector3> GetSimIndices3D() const;
 
+		/** Whether or not this pattern is empty. */
+		bool IsEmpty() const;
+
 		/** Return the Pattern index this facade has been created with. */
 		int32 GetPatternIndex() const { return PatternIndex; }
 
@@ -91,6 +94,7 @@ namespace UE::Chaos::ClothAsset
 		/** This will remove the 2D vertices, but the associated seams and 3D vertices will still exist, and point to INDEX_NONE */
 		void RemoveSimVertices2D(int32 NumSimVertices);
 		void RemoveAllSimVertices2D() { RemoveSimVertices2D(GetNumSimVertices2D()); }
+		void RemoveSimVertices2D(const TArray<int32>& SortedDeletionList);
 
 		//~ Sim Faces Group
 		// Note: SimIndices points to the collection arrays, not the pattern arrays
@@ -98,6 +102,7 @@ namespace UE::Chaos::ClothAsset
 		void SetNumSimFaces(int32 NumSimFaces);
 		TArrayView<FIntVector3> GetSimIndices2D();
 		TArrayView<FIntVector3> GetSimIndices3D();
+		void RemoveSimFaces(const TArray<int32>& SortedDeletionList);
 
 	private:
 		friend class FCollectionClothFacade;
diff --git a/Engine/Plugins/Experimental/ChaosClothAssetEditor/Source/ChaosClothAssetDataflowNodes/Private/ChaosClothAsset/TerminalNode.cpp b/Engine/Plugins/Experimental/ChaosClothAssetEditor/Source/ChaosClothAssetDataflowNodes/Private/ChaosClothAsset/TerminalNode.cpp
--- a/Engine/Plugins/Experimental/ChaosClothAssetEditor/Source/ChaosClothAssetDataflowNodes/Private/ChaosClothAsset/TerminalNode.cpp
+++ b/Engine/Plugins/Experimental/ChaosClothAssetEditor/Source/ChaosClothAssetDataflowNodes/Private/ChaosClothAsset/TerminalNode.cpp
@@ -4,6 +4,7 @@
 #include "ChaosClothAsset/DataflowNodes.h"
 #include "ChaosClothAsset/ClothAsset.h"
 #include "ChaosClothAsset/CollectionClothFacade.h"
+#include "ChaosClothAsset/ClothGeometryTools.h"
 #include "Chaos/CollectionPropertyFacade.h"
 #include "Animation/Skeleton.h"
 #include "Engine/SkinnedAssetCommon.h"
@@ -63,6 +64,7 @@ void FChaosClothAssetTerminalNode::SetAssetValue(TObjectPtr<UObject> Asset, Data
 
 				// Copy input LOD to current output LOD
 				ClothFacade.Initialize(InClothFacade);
+				FClothGeometryTools::CleanupAndCompactMesh(ClothCollection);
 
 				// Add this LOD's materials to the asset
 				const int32 NumLodMaterials = ClothFacade.GetNumRenderPatterns();
