Chaos Cloth - Move to shared references.

#rnx

diff --git a/Engine/Plugins/Enterprise/DatasmithImporter/Source/DatasmithImporter/Private/DatasmithImporter.cpp b/Engine/Plugins/Enterprise/DatasmithImporter/Source/DatasmithImporter/Private/DatasmithImporter.cpp
--- a/Engine/Plugins/Enterprise/DatasmithImporter/Source/DatasmithImporter/Private/DatasmithImporter.cpp
+++ b/Engine/Plugins/Enterprise/DatasmithImporter/Source/DatasmithImporter/Private/DatasmithImporter.cpp
@@ -84,6 +84,7 @@
 #include "Animation/Skeleton.h"
 #include "ChaosClothAsset/ClothAsset.h"
 #include "ChaosClothAsset/CollectionClothFacade.h"
+#include "GeometryCollection/ManagedArrayCollection.h"
 
 extern UNREALED_API UEditorEngine* GEditor;
 
@@ -495,10 +496,10 @@ void FDatasmithImporter::ImportClothes(FDatasmithImportContext& ImportContext)
 
 		using namespace UE::Chaos::ClothAsset;
 
-		TArray<TSharedPtr<FManagedArrayCollection>>& Collections = ClothAsset->GetClothCollections();
-		Collections.SetNum(1);
-		FCollectionClothFacade Cloth(Collections[0]);
-		Cloth.Reset(); // on reimport, asset is non-empty
+		TArray<TSharedRef<FManagedArrayCollection>>& Collections = ClothAsset->GetClothCollections();
+		Collections.Reset(1);
+		FCollectionClothFacade Cloth(Collections.Emplace_GetRef(MakeShared<FManagedArrayCollection>()));
+		Cloth.DefineSchema();
 
 		for (FDatasmithClothPattern& Pattern : DsCloth.Patterns)
 		{
diff --git a/Engine/Plugins/Experimental/ChaosClothAsset/Source/ChaosClothAsset/Private/ChaosClothAsset/ClothCollection.cpp b/Engine/Plugins/Experimental/ChaosClothAsset/Source/ChaosClothAsset/Private/ChaosClothAsset/ClothCollection.cpp
--- a/Engine/Plugins/Experimental/ChaosClothAsset/Source/ChaosClothAsset/Private/ChaosClothAsset/ClothCollection.cpp
+++ b/Engine/Plugins/Experimental/ChaosClothAsset/Source/ChaosClothAsset/Private/ChaosClothAsset/ClothCollection.cpp
@@ -169,13 +169,11 @@ namespace UE::Chaos::ClothAsset
 	const FName FClothCollection::RenderFacesGroup = Private::RenderFacesGroup;
 	const FName FClothCollection::RenderVerticesGroup = Private::RenderVerticesGroup;
 
-	FClothCollection::FClothCollection(const TSharedPtr<FManagedArrayCollection>& InManagedArrayCollection)
+	FClothCollection::FClothCollection(const TSharedRef<FManagedArrayCollection>& InManagedArrayCollection)
 		: ManagedArrayCollection(InManagedArrayCollection)
 	{
 		using namespace UE::Chaos::ClothAsset::Private;
 
-		check(ManagedArrayCollection.IsValid());
-
 		// LODs Group
 		PhysicsAssetPathName = ManagedArrayCollection->FindAttribute<FString>(PhysicsAssetPathNameAttribute, LodsGroup);
 		SkeletonAssetPathName = ManagedArrayCollection->FindAttribute<FString>(SkeletonAssetPathNameAttribute, LodsGroup);
diff --git a/Engine/Plugins/Experimental/ChaosClothAsset/Source/ChaosClothAsset/Private/ChaosClothAsset/ClothCollection.h b/Engine/Plugins/Experimental/ChaosClothAsset/Source/ChaosClothAsset/Private/ChaosClothAsset/ClothCollection.h
--- a/Engine/Plugins/Experimental/ChaosClothAsset/Source/ChaosClothAsset/Private/ChaosClothAsset/ClothCollection.h
+++ b/Engine/Plugins/Experimental/ChaosClothAsset/Source/ChaosClothAsset/Private/ChaosClothAsset/ClothCollection.h
@@ -35,7 +35,7 @@ namespace UE::Chaos::ClothAsset
 		static constexpr int8 MaxNumBoneInfluences = 12; // This should be <= MAX_TOTAL_INFLUENCES defined in GPUSkinPublicDefs.h 
 		static constexpr int8 MaxNumTetherAttachments = 4; // This should be <= FClothTetherDataPrivate::MaxNumAttachments
 
-		explicit FClothCollection(const TSharedPtr<FManagedArrayCollection>& InManagedArrayCollection);
+		explicit FClothCollection(const TSharedRef<FManagedArrayCollection>& InManagedArrayCollection);
 		~FClothCollection() = default;
 
 		FClothCollection(const FClothCollection&) = delete;
@@ -271,7 +271,7 @@ namespace UE::Chaos::ClothAsset
 
 	private:
 		//~ Cloth collection
-		TSharedPtr<FManagedArrayCollection> ManagedArrayCollection;
+		TSharedRef<FManagedArrayCollection> ManagedArrayCollection;
 
 		//~ LODs Group
 		TManagedArray<FString>* PhysicsAssetPathName;
diff --git a/Engine/Plugins/Experimental/ChaosClothAsset/Source/ChaosClothAsset/Private/ChaosClothAsset/ClothGeometryTools.cpp b/Engine/Plugins/Experimental/ChaosClothAsset/Source/ChaosClothAsset/Private/ChaosClothAsset/ClothGeometryTools.cpp
--- a/Engine/Plugins/Experimental/ChaosClothAsset/Source/ChaosClothAsset/Private/ChaosClothAsset/ClothGeometryTools.cpp
+++ b/Engine/Plugins/Experimental/ChaosClothAsset/Source/ChaosClothAsset/Private/ChaosClothAsset/ClothGeometryTools.cpp
@@ -429,25 +429,25 @@ namespace UE::Chaos::ClothAsset
 	}// namespace Private::SimMeshBuilder
 
 	// TODO: Move these functions to the cloth facade?
-	bool FClothGeometryTools::HasSimMesh(const TSharedPtr<const FManagedArrayCollection>& ClothCollection)
+	bool FClothGeometryTools::HasSimMesh(const TSharedRef<const FManagedArrayCollection>& ClothCollection)
 	{
 		FCollectionClothConstFacade ClothFacade(ClothCollection);
 		return ClothFacade.GetNumSimVertices2D() > 0 && ClothFacade.GetNumSimVertices3D() && ClothFacade.GetNumSimFaces() > 0;
 	}
 
-	bool FClothGeometryTools::HasRenderMesh(const TSharedPtr<const FManagedArrayCollection>& ClothCollection)
+	bool FClothGeometryTools::HasRenderMesh(const TSharedRef<const FManagedArrayCollection>& ClothCollection)
 	{
 		FCollectionClothConstFacade ClothFacade(ClothCollection);
 		return ClothFacade.GetNumRenderVertices() > 0 && ClothFacade.GetNumRenderFaces() > 0;
 	}
 
-	void FClothGeometryTools::DeleteRenderMesh(const TSharedPtr<FManagedArrayCollection>& ClothCollection)
+	void FClothGeometryTools::DeleteRenderMesh(const TSharedRef<FManagedArrayCollection>& ClothCollection)
 	{
 		FCollectionClothFacade ClothFacade(ClothCollection);
 		ClothFacade.SetNumRenderPatterns(0);
 	}
 
-	void FClothGeometryTools::DeleteSimMesh(const TSharedPtr<FManagedArrayCollection>& ClothCollection)
+	void FClothGeometryTools::DeleteSimMesh(const TSharedRef<FManagedArrayCollection>& ClothCollection)
 	{
 		FCollectionClothFacade ClothFacade(ClothCollection);
 		ClothFacade.SetNumSimPatterns(0);
@@ -455,7 +455,7 @@ namespace UE::Chaos::ClothAsset
 		ClothFacade.SetNumSeams(0);
 	}
 
-	void FClothGeometryTools::DeleteTethers(const TSharedPtr<FManagedArrayCollection>& ClothCollection)
+	void FClothGeometryTools::DeleteTethers(const TSharedRef<FManagedArrayCollection>& ClothCollection)
 	{
 		FCollectionClothFacade ClothFacade(ClothCollection);
 		for (TArray<int32>& KinematicIndex : ClothFacade.GetTetherKinematicIndex())
@@ -468,7 +468,7 @@ namespace UE::Chaos::ClothAsset
 		}
 	}
 
-	void FClothGeometryTools::CopySimMeshToRenderMesh(const TSharedPtr<FManagedArrayCollection>& ClothCollection, const FString& RenderMaterialPathName, bool bSingleRenderPattern)
+	void FClothGeometryTools::CopySimMeshToRenderMesh(const TSharedRef<FManagedArrayCollection>& ClothCollection, const FString& RenderMaterialPathName, bool bSingleRenderPattern)
 	{
 		FCollectionClothFacade ClothFacade(ClothCollection);
 
@@ -565,7 +565,7 @@ namespace UE::Chaos::ClothAsset
 	}
 
 	void FClothGeometryTools::ReverseMesh(
-		const TSharedPtr<FManagedArrayCollection>& ClothCollection,
+		const TSharedRef<FManagedArrayCollection>& ClothCollection,
 		bool bReverseSimMeshNormals,
 		bool bReverseSimMeshWindingOrder,
 		bool bReverseRenderMeshNormals,
@@ -681,7 +681,7 @@ namespace UE::Chaos::ClothAsset
 	}
 
 	void FClothGeometryTools::BindMeshToRootBone(
-		const TSharedPtr<FManagedArrayCollection>& ClothCollection,
+		const TSharedRef<FManagedArrayCollection>& ClothCollection,
 		bool bBindSimMesh,
 		bool bBindRenderMesh)
 	{
@@ -719,7 +719,7 @@ namespace UE::Chaos::ClothAsset
 	}
 
 	void FClothGeometryTools::BuildSimMeshFromDynamicMesh(
-		const TSharedPtr<FManagedArrayCollection>& ClothCollection,
+		const TSharedRef<FManagedArrayCollection>& ClothCollection,
 		const UE::Geometry::FDynamicMesh3& DynamicMesh, int32 UVChannelIndex, const FVector2f& UVScale, bool bAppend)
 	{
 		using namespace Private::SimMeshBuilder;
@@ -772,7 +772,7 @@ namespace UE::Chaos::ClothAsset
 		}
 	}
 
-	void FClothGeometryTools::CleanupAndCompactMesh(const TSharedPtr<FManagedArrayCollection>& ClothCollection)
+	void FClothGeometryTools::CleanupAndCompactMesh(const TSharedRef<FManagedArrayCollection>& ClothCollection)
 	{
 		FCollectionClothFacade Cloth(ClothCollection);
 
diff --git a/Engine/Plugins/Experimental/ChaosClothAsset/Source/ChaosClothAsset/Private/ChaosClothAsset/CollectionClothFacade.cpp b/Engine/Plugins/Experimental/ChaosClothAsset/Source/ChaosClothAsset/Private/ChaosClothAsset/CollectionClothFacade.cpp
--- a/Engine/Plugins/Experimental/ChaosClothAsset/Source/ChaosClothAsset/Private/ChaosClothAsset/CollectionClothFacade.cpp
+++ b/Engine/Plugins/Experimental/ChaosClothAsset/Source/ChaosClothAsset/Private/ChaosClothAsset/CollectionClothFacade.cpp
@@ -8,12 +8,12 @@
 
 namespace UE::Chaos::ClothAsset
 {
-	FCollectionClothConstFacade::FCollectionClothConstFacade(const TSharedPtr<const FManagedArrayCollection>& ManagedArrayCollection)
-		: ClothCollection(MakeShared<const FClothCollection>(ConstCastSharedPtr<FManagedArrayCollection>(ManagedArrayCollection)))
+	FCollectionClothConstFacade::FCollectionClothConstFacade(const TSharedRef<const FManagedArrayCollection>& ManagedArrayCollection)
+		: ClothCollection(MakeShared<const FClothCollection>(ConstCastSharedRef<FManagedArrayCollection>(ManagedArrayCollection)))
 	{
 	}
 
-	FCollectionClothConstFacade::FCollectionClothConstFacade(const TSharedPtr<const FClothCollection>& ClothCollection)
+	FCollectionClothConstFacade::FCollectionClothConstFacade(const TSharedRef<const FClothCollection>& ClothCollection)
 		:ClothCollection(ClothCollection)
 	{
 	}
@@ -275,12 +275,12 @@ namespace UE::Chaos::ClothAsset
 		}
 	}
 
-	FCollectionClothFacade::FCollectionClothFacade(const TSharedPtr<FManagedArrayCollection>& ManagedArrayCollection)
+	FCollectionClothFacade::FCollectionClothFacade(const TSharedRef<FManagedArrayCollection>& ManagedArrayCollection)
 		: FCollectionClothConstFacade(ManagedArrayCollection)
 	{
 	}
 
-	FCollectionClothFacade::FCollectionClothFacade(const TSharedPtr<FClothCollection>& InClothCollection)
+	FCollectionClothFacade::FCollectionClothFacade(const TSharedRef<FClothCollection>& InClothCollection)
 		: FCollectionClothConstFacade(InClothCollection)
 	{
 	}
diff --git a/Engine/Plugins/Experimental/ChaosClothAsset/Source/ChaosClothAsset/Private/ChaosClothAsset/CollectionClothRenderPatternFacade.cpp b/Engine/Plugins/Experimental/ChaosClothAsset/Source/ChaosClothAsset/Private/ChaosClothAsset/CollectionClothRenderPatternFacade.cpp
--- a/Engine/Plugins/Experimental/ChaosClothAsset/Source/ChaosClothAsset/Private/ChaosClothAsset/CollectionClothRenderPatternFacade.cpp
+++ b/Engine/Plugins/Experimental/ChaosClothAsset/Source/ChaosClothAsset/Private/ChaosClothAsset/CollectionClothRenderPatternFacade.cpp
@@ -131,11 +131,10 @@ namespace UE::Chaos::ClothAsset
 		return GetNumRenderVertices() == 0 && GetNumRenderFaces() == 0;
 	}
 
-	FCollectionClothRenderPatternConstFacade::FCollectionClothRenderPatternConstFacade(const TSharedPtr<const FClothCollection>& ClothCollection, int32 InPatternIndex)
+	FCollectionClothRenderPatternConstFacade::FCollectionClothRenderPatternConstFacade(const TSharedRef<const FClothCollection>& ClothCollection, int32 InPatternIndex)
 		: ClothCollection(ClothCollection)
 		, PatternIndex(InPatternIndex)
 	{
-		check(ClothCollection.IsValid());
 		check(ClothCollection->IsValid());
 		check(PatternIndex >= 0 && PatternIndex < ClothCollection->GetNumElements(FClothCollection::RenderPatternsGroup));
 	}
@@ -312,7 +311,7 @@ namespace UE::Chaos::ClothAsset
 			GetElementIndex());
 	}
 
-	FCollectionClothRenderPatternFacade::FCollectionClothRenderPatternFacade(const TSharedPtr<FClothCollection>& ClothCollection, int32 InPatternIndex)
+	FCollectionClothRenderPatternFacade::FCollectionClothRenderPatternFacade(const TSharedRef<FClothCollection>& ClothCollection, int32 InPatternIndex)
 		: FCollectionClothRenderPatternConstFacade(ClothCollection, InPatternIndex)
 	{
 	}
diff --git a/Engine/Plugins/Experimental/ChaosClothAsset/Source/ChaosClothAsset/Private/ChaosClothAsset/CollectionClothSeamFacade.cpp b/Engine/Plugins/Experimental/ChaosClothAsset/Source/ChaosClothAsset/Private/ChaosClothAsset/CollectionClothSeamFacade.cpp
--- a/Engine/Plugins/Experimental/ChaosClothAsset/Source/ChaosClothAsset/Private/ChaosClothAsset/CollectionClothSeamFacade.cpp
+++ b/Engine/Plugins/Experimental/ChaosClothAsset/Source/ChaosClothAsset/Private/ChaosClothAsset/CollectionClothSeamFacade.cpp
@@ -497,11 +497,10 @@ namespace UE::Chaos::ClothAsset
 			GetElementIndex());
 	}
 
-	FCollectionClothSeamConstFacade::FCollectionClothSeamConstFacade(const TSharedPtr<const FClothCollection>& ClothCollection, int32 SeamIndex)
+	FCollectionClothSeamConstFacade::FCollectionClothSeamConstFacade(const TSharedRef<const FClothCollection>& ClothCollection, int32 SeamIndex)
 		: ClothCollection(ClothCollection)
 		, SeamIndex(SeamIndex)
 	{
-		check(ClothCollection.IsValid());
 		check(ClothCollection->IsValid());
 		check(SeamIndex >= 0 && SeamIndex < ClothCollection->GetNumElements(FClothCollection::SeamsGroup));
 	}
@@ -766,7 +765,7 @@ namespace UE::Chaos::ClothAsset
 			GetElementIndex());
 	}
 
-	FCollectionClothSeamFacade::FCollectionClothSeamFacade(const TSharedPtr<FClothCollection>& ClothCollection, int32 InSeamIndex)
+	FCollectionClothSeamFacade::FCollectionClothSeamFacade(const TSharedRef<FClothCollection>& ClothCollection, int32 InSeamIndex)
 		: FCollectionClothSeamConstFacade(ClothCollection, InSeamIndex)
 	{
 	}
diff --git a/Engine/Plugins/Experimental/ChaosClothAsset/Source/ChaosClothAsset/Private/ChaosClothAsset/CollectionClothSimPatternFacade.cpp b/Engine/Plugins/Experimental/ChaosClothAsset/Source/ChaosClothAsset/Private/ChaosClothAsset/CollectionClothSimPatternFacade.cpp
--- a/Engine/Plugins/Experimental/ChaosClothAsset/Source/ChaosClothAsset/Private/ChaosClothAsset/CollectionClothSimPatternFacade.cpp
+++ b/Engine/Plugins/Experimental/ChaosClothAsset/Source/ChaosClothAsset/Private/ChaosClothAsset/CollectionClothSimPatternFacade.cpp
@@ -80,11 +80,10 @@ namespace UE::Chaos::ClothAsset
 		return GetNumSimFaces() == 0 && GetNumSimVertices2D() == 0;
 	}
 
-	FCollectionClothSimPatternConstFacade::FCollectionClothSimPatternConstFacade(const TSharedPtr<const FClothCollection>& ClothCollection, int32 InPatternIndex)
+	FCollectionClothSimPatternConstFacade::FCollectionClothSimPatternConstFacade(const TSharedRef<const FClothCollection>& ClothCollection, int32 InPatternIndex)
 		: ClothCollection(ClothCollection)
 		, PatternIndex(InPatternIndex)
 	{
-		check(ClothCollection.IsValid());
 		check(ClothCollection->IsValid());
 		check(PatternIndex >= 0 && PatternIndex < ClothCollection->GetNumElements(FClothCollection::SimPatternsGroup));
 	}
@@ -288,7 +287,7 @@ namespace UE::Chaos::ClothAsset
 			GetElementIndex());
 	}
 
-	FCollectionClothSimPatternFacade::FCollectionClothSimPatternFacade(const TSharedPtr<FClothCollection>& ClothCollection, int32 InPatternIndex)
+	FCollectionClothSimPatternFacade::FCollectionClothSimPatternFacade(const TSharedRef<FClothCollection>& ClothCollection, int32 InPatternIndex)
 		: FCollectionClothSimPatternConstFacade(ClothCollection, InPatternIndex)
 	{
 	}
diff --git a/Engine/Plugins/Experimental/ChaosClothAsset/Source/ChaosClothAsset/Public/ChaosClothAsset/ClothGeometryTools.h b/Engine/Plugins/Experimental/ChaosClothAsset/Source/ChaosClothAsset/Public/ChaosClothAsset/ClothGeometryTools.h
--- a/Engine/Plugins/Experimental/ChaosClothAsset/Source/ChaosClothAsset/Public/ChaosClothAsset/ClothGeometryTools.h
+++ b/Engine/Plugins/Experimental/ChaosClothAsset/Source/ChaosClothAsset/Public/ChaosClothAsset/ClothGeometryTools.h
@@ -16,26 +16,26 @@ namespace UE::Chaos::ClothAsset
 	{
 
 		/** Return whether at least one pattern of this collection has any faces to simulate. */
-		static bool HasSimMesh(const TSharedPtr<const FManagedArrayCollection>& ClothCollection);
+		static bool HasSimMesh(const TSharedRef<const FManagedArrayCollection>& ClothCollection);
 
 		/** Return whether at least one pattern of this collection has any faces to render. */
-		static bool HasRenderMesh(const TSharedPtr<const FManagedArrayCollection>& ClothCollection);
+		static bool HasRenderMesh(const TSharedRef<const FManagedArrayCollection>& ClothCollection);
 
 		/** Delete the render mesh data. */
-		static void DeleteRenderMesh(const TSharedPtr<FManagedArrayCollection>& ClothCollection);
+		static void DeleteRenderMesh(const TSharedRef<FManagedArrayCollection>& ClothCollection);
 
 		/** Delete the sim mesh data. */
-		static void DeleteSimMesh(const TSharedPtr<FManagedArrayCollection>& ClothCollection);
+		static void DeleteSimMesh(const TSharedRef<FManagedArrayCollection>& ClothCollection);
 
 		/** Remove all tethers. */
-		static void DeleteTethers(const TSharedPtr<FManagedArrayCollection>& ClothCollection);
+		static void DeleteTethers(const TSharedRef<FManagedArrayCollection>& ClothCollection);
 
 		/** Turn the sim mesh portion of this ClothCollection into a render mesh. */
-		static void CopySimMeshToRenderMesh(const TSharedPtr<FManagedArrayCollection>& ClothCollection, const FString& RenderMaterialPathName, bool bSingleRenderPattern);
+		static void CopySimMeshToRenderMesh(const TSharedRef<FManagedArrayCollection>& ClothCollection, const FString& RenderMaterialPathName, bool bSingleRenderPattern);
 
 		/** Reverse the mesh normals. Will reverse all normals if pattern selection is empty. */
 		static void ReverseMesh(
-			const TSharedPtr<FManagedArrayCollection>& ClothCollection,
+			const TSharedRef<FManagedArrayCollection>& ClothCollection,
 			bool bReverseSimMeshNormals,
 			bool bReverseSimMeshWindingOrder,
 			bool bReverseRenderMeshNormals,
@@ -49,7 +49,7 @@ namespace UE::Chaos::ClothAsset
 		 * @param Lods if empty will apply the change to all LODs. Otherwise only LODs specified in the array (if exist) are affected.
 		 */
 		static void BindMeshToRootBone(
-			const TSharedPtr<FManagedArrayCollection>& ClothCollection,
+			const TSharedRef<FManagedArrayCollection>& ClothCollection,
 			bool bBindSimMesh,
 			bool bBindRenderMesh);
 
@@ -57,13 +57,13 @@ namespace UE::Chaos::ClothAsset
 		* Unwrap and build SimMesh data from a DynamicMesh
 		*/
 		static void BuildSimMeshFromDynamicMesh(
-			const TSharedPtr<FManagedArrayCollection>& ClothCollection,
+			const TSharedRef<FManagedArrayCollection>& ClothCollection,
 			const UE::Geometry::FDynamicMesh3& DynamicMesh, int32 UVChannelIndex, const FVector2f& UVScale, bool bAppend);
 
 		/**
 		* Remove (topologically) degenerate triangles. Remove any vertices that aren't in a triangle. Compact any lookup arrays that contain INDEX_NONEs.
 		* Remove any empty patterns.
 		*/
-		static void CleanupAndCompactMesh(const TSharedPtr<FManagedArrayCollection>& ClothCollection);
+		static void CleanupAndCompactMesh(const TSharedRef<FManagedArrayCollection>& ClothCollection);
 	};
 }  // End namespace UE::Chaos::ClothAsset
\ No newline at end of file
diff --git a/Engine/Plugins/Experimental/ChaosClothAsset/Source/ChaosClothAsset/Public/ChaosClothAsset/CollectionClothFacade.h b/Engine/Plugins/Experimental/ChaosClothAsset/Source/ChaosClothAsset/Public/ChaosClothAsset/CollectionClothFacade.h
--- a/Engine/Plugins/Experimental/ChaosClothAsset/Source/ChaosClothAsset/Public/ChaosClothAsset/CollectionClothFacade.h
+++ b/Engine/Plugins/Experimental/ChaosClothAsset/Source/ChaosClothAsset/Public/ChaosClothAsset/CollectionClothFacade.h
@@ -20,7 +20,7 @@ namespace UE::Chaos::ClothAsset
 	class CHAOSCLOTHASSET_API FCollectionClothConstFacade
 	{
 	public:
-		explicit FCollectionClothConstFacade(const TSharedPtr<const FManagedArrayCollection>& ManagedArrayCollection);
+		explicit FCollectionClothConstFacade(const TSharedRef<const FManagedArrayCollection>& ManagedArrayCollection);
 
 		FCollectionClothConstFacade() = delete;
 
@@ -122,11 +122,11 @@ namespace UE::Chaos::ClothAsset
 			TArray<uint32>& PatternToWeldedIndices, TArray<TArray<int32>>* OptionalWeldedToPatternIndices = nullptr) const;
 
 	protected:
-		TSharedPtr<const class FClothCollection> ClothCollection;
+		TSharedRef<const class FClothCollection> ClothCollection;
 
 		friend class FCollectionClothSeamFacade;
 		friend class FCollectionClothSeamConstFacade;
-		explicit FCollectionClothConstFacade(const TSharedPtr<const class FClothCollection>& ClothCollection);
+		explicit FCollectionClothConstFacade(const TSharedRef<const class FClothCollection>& ClothCollection);
 	};
 
 	/**
@@ -136,7 +136,7 @@ namespace UE::Chaos::ClothAsset
 	class CHAOSCLOTHASSET_API FCollectionClothFacade final : public FCollectionClothConstFacade
 	{
 	public:
-		explicit FCollectionClothFacade(const TSharedPtr<FManagedArrayCollection>& ManagedArrayCollection);
+		explicit FCollectionClothFacade(const TSharedRef<FManagedArrayCollection>& ManagedArrayCollection);
 
 		FCollectionClothFacade() = delete;
 
@@ -252,16 +252,16 @@ namespace UE::Chaos::ClothAsset
 		void RemoveWeightMap(const FName& Name);
 
 		TArrayView<float> GetWeightMap(const FName& Name);
-	private:
 
+	private:
 		void SetDefaults();
 
-		TSharedPtr<class FClothCollection> GetClothCollection() { return ConstCastSharedPtr<class FClothCollection>(ClothCollection); }
+		TSharedRef<class FClothCollection> GetClothCollection() { return ConstCastSharedRef<class FClothCollection>(ClothCollection); }
 
 		friend class FCollectionClothSeamFacade;
 		friend class FCollectionClothSimPatternFacade;
 
-		explicit FCollectionClothFacade(const TSharedPtr<class FClothCollection>& InClothCollection);
+		explicit FCollectionClothFacade(const TSharedRef<class FClothCollection>& InClothCollection);
 
 		// These methods are private because they're managed by the FCollectionClothSeamFacade.
 		//~ Sim Vertices 2D Group
diff --git a/Engine/Plugins/Experimental/ChaosClothAsset/Source/ChaosClothAsset/Public/ChaosClothAsset/CollectionClothRenderPatternFacade.h b/Engine/Plugins/Experimental/ChaosClothAsset/Source/ChaosClothAsset/Public/ChaosClothAsset/CollectionClothRenderPatternFacade.h
--- a/Engine/Plugins/Experimental/ChaosClothAsset/Source/ChaosClothAsset/Public/ChaosClothAsset/CollectionClothRenderPatternFacade.h
+++ b/Engine/Plugins/Experimental/ChaosClothAsset/Source/ChaosClothAsset/Public/ChaosClothAsset/CollectionClothRenderPatternFacade.h
@@ -59,12 +59,12 @@ namespace UE::Chaos::ClothAsset
 	protected:
 		friend class FCollectionClothRenderPatternFacade;  // For other instances access
 		friend class FCollectionClothConstFacade;
-		FCollectionClothRenderPatternConstFacade(const TSharedPtr<const class FClothCollection>& InClothCollection, int32 InPatternIndex);
+		FCollectionClothRenderPatternConstFacade(const TSharedRef<const class FClothCollection>& InClothCollection, int32 InPatternIndex);
 
 		static constexpr int32 GetBaseElementIndex() { return 0; }
 		int32 GetElementIndex() const { return GetBaseElementIndex() + PatternIndex; }
 
-		TSharedPtr<const class FClothCollection> ClothCollection;
+		TSharedRef<const class FClothCollection> ClothCollection;
 		int32 PatternIndex;
 	};
 
@@ -117,10 +117,10 @@ namespace UE::Chaos::ClothAsset
 
 	protected:
 		friend class FCollectionClothFacade;
-		FCollectionClothRenderPatternFacade(const TSharedPtr<class FClothCollection>& InClothCollection, int32 InPatternIndex);
+		FCollectionClothRenderPatternFacade(const TSharedRef<class FClothCollection>& InClothCollection, int32 InPatternIndex);
 
 		void SetDefaults();
 
-		TSharedPtr<class FClothCollection> GetClothCollection() { return ConstCastSharedPtr<class FClothCollection>(ClothCollection); }
+		TSharedRef<class FClothCollection> GetClothCollection() { return ConstCastSharedRef<class FClothCollection>(ClothCollection); }
 	};
 }  // End namespace UE::Chaos::ClothAsset
diff --git a/Engine/Plugins/Experimental/ChaosClothAsset/Source/ChaosClothAsset/Public/ChaosClothAsset/CollectionClothSeamFacade.h b/Engine/Plugins/Experimental/ChaosClothAsset/Source/ChaosClothAsset/Public/ChaosClothAsset/CollectionClothSeamFacade.h
--- a/Engine/Plugins/Experimental/ChaosClothAsset/Source/ChaosClothAsset/Public/ChaosClothAsset/CollectionClothSeamFacade.h
+++ b/Engine/Plugins/Experimental/ChaosClothAsset/Source/ChaosClothAsset/Public/ChaosClothAsset/CollectionClothSeamFacade.h
@@ -41,12 +41,12 @@ namespace UE::Chaos::ClothAsset
 	protected:
 		friend class FCollectionClothSeamFacade;  // For other instances access
 		friend class FCollectionClothConstFacade;
-		FCollectionClothSeamConstFacade(const TSharedPtr<const class FClothCollection>& InClothCollection, int32 SeamIndex);
+		FCollectionClothSeamConstFacade(const TSharedRef<const class FClothCollection>& InClothCollection, int32 SeamIndex);
 
 		static constexpr int32 GetBaseElementIndex() { return 0; }
 		int32 GetElementIndex() const { return GetBaseElementIndex() + SeamIndex; }
 
-		TSharedPtr<const class FClothCollection> ClothCollection;
+		TSharedRef<const class FClothCollection> ClothCollection;
 		int32 SeamIndex;
 	};
 
@@ -79,7 +79,7 @@ namespace UE::Chaos::ClothAsset
 
 private:
 		friend class FCollectionClothFacade;
-		FCollectionClothSeamFacade(const TSharedPtr<class FClothCollection>& InClothCollection, int32 SeamIndex);
+		FCollectionClothSeamFacade(const TSharedRef<class FClothCollection>& InClothCollection, int32 SeamIndex);
 
 		void SetDefaults();
 
@@ -89,6 +89,6 @@ private:
 		TArrayView<FIntVector2> GetSeamStitch2DEndIndices();
 		TArrayView<int32> GetSeamStitch3DIndex();
 
-		TSharedPtr<class FClothCollection> GetClothCollection() { return ConstCastSharedPtr<class FClothCollection>(ClothCollection); }
+		TSharedRef<class FClothCollection> GetClothCollection() { return ConstCastSharedRef<class FClothCollection>(ClothCollection); }
 	};
 }  // End namespace UE::Chaos::ClothAsset
diff --git a/Engine/Plugins/Experimental/ChaosClothAsset/Source/ChaosClothAsset/Public/ChaosClothAsset/CollectionClothSimPatternFacade.h b/Engine/Plugins/Experimental/ChaosClothAsset/Source/ChaosClothAsset/Public/ChaosClothAsset/CollectionClothSimPatternFacade.h
--- a/Engine/Plugins/Experimental/ChaosClothAsset/Source/ChaosClothAsset/Public/ChaosClothAsset/CollectionClothSimPatternFacade.h
+++ b/Engine/Plugins/Experimental/ChaosClothAsset/Source/ChaosClothAsset/Public/ChaosClothAsset/CollectionClothSimPatternFacade.h
@@ -50,12 +50,12 @@ namespace UE::Chaos::ClothAsset
 	protected:
 		friend class FCollectionClothSimPatternFacade;  // For other instances access
 		friend class FCollectionClothConstFacade;
-		FCollectionClothSimPatternConstFacade(const TSharedPtr<const class FClothCollection>& InClothCollection, int32 InPatternIndex);
+		FCollectionClothSimPatternConstFacade(const TSharedRef<const class FClothCollection>& InClothCollection, int32 InPatternIndex);
 
 		static constexpr int32 GetBaseElementIndex() { return 0; }
 		int32 GetElementIndex() const { return GetBaseElementIndex() + PatternIndex; }
 
-		TSharedPtr<const class FClothCollection> ClothCollection;
+		TSharedRef<const class FClothCollection> ClothCollection;
 		int32 PatternIndex;
 	};
 
@@ -106,7 +106,7 @@ namespace UE::Chaos::ClothAsset
 
 	private:
 		friend class FCollectionClothFacade;
-		FCollectionClothSimPatternFacade(const TSharedPtr<class FClothCollection>& InClothCollection, int32 InPatternIndex);
+		FCollectionClothSimPatternFacade(const TSharedRef<class FClothCollection>& InClothCollection, int32 InPatternIndex);
 
 		void SetDefaults();
 
@@ -116,6 +116,6 @@ namespace UE::Chaos::ClothAsset
 		void SetNumSimVertices2D(int32 NumSimVertices);
 		TArrayView<int32> GetSimVertex3DLookup();
 
-		TSharedPtr<class FClothCollection> GetClothCollection() { return ConstCastSharedPtr<class FClothCollection>(ClothCollection); }
+		TSharedRef<class FClothCollection> GetClothCollection() { return ConstCastSharedRef<class FClothCollection>(ClothCollection); }
 	};
 }  // End namespace UE::Chaos::ClothAsset
diff --git a/Engine/Plugins/Experimental/ChaosClothAsset/Source/ChaosClothAssetEngine/Private/ChaosClothAsset/ClothAsset.cpp b/Engine/Plugins/Experimental/ChaosClothAsset/Source/ChaosClothAssetEngine/Private/ChaosClothAsset/ClothAsset.cpp
--- a/Engine/Plugins/Experimental/ChaosClothAsset/Source/ChaosClothAssetEngine/Private/ChaosClothAsset/ClothAsset.cpp
+++ b/Engine/Plugins/Experimental/ChaosClothAsset/Source/ChaosClothAssetEngine/Private/ChaosClothAsset/ClothAsset.cpp
@@ -30,7 +30,7 @@
 
 namespace UE::Chaos::ClothAsset::Private
 {
-::Chaos::FChaosArchive& Serialize(::Chaos::FChaosArchive& Ar, TArray<TSharedPtr<FManagedArrayCollection>>& ClothCollections)
+::Chaos::FChaosArchive& Serialize(::Chaos::FChaosArchive& Ar, TArray<TSharedRef<FManagedArrayCollection>>& ClothCollections)
 {
 	Ar.UsingCustomVersion(FUE5MainStreamObjectVersion::GUID);
 
@@ -38,7 +38,7 @@ namespace UE::Chaos::ClothAsset::Private
 	{
 		// Cloth assets before this version had a single ClothCollection with a completely different schema.
 		ClothCollections.Empty(1);
-		TSharedPtr<FManagedArrayCollection>& ClothCollection = ClothCollections.Emplace_GetRef(MakeShared<FManagedArrayCollection>());
+		TSharedRef<FManagedArrayCollection>& ClothCollection = ClothCollections.Emplace_GetRef(MakeShared<FManagedArrayCollection>());
 		ClothCollection->Serialize(Ar);
 
 		// Now we're just going to hard reset and define a new schema.
@@ -77,7 +77,7 @@ namespace UE::Chaos::ClothAsset::Private
 
 			for (int32 i = 0; i < SerializeNum; i++)
 			{
-				TSharedPtr<FManagedArrayCollection>& ClothCollection = ClothCollections.Emplace_GetRef(MakeShared<FManagedArrayCollection>());
+				TSharedRef<FManagedArrayCollection>& ClothCollection = ClothCollections.Emplace_GetRef(MakeShared<FManagedArrayCollection>());
 				ClothCollection->Serialize(Ar);
 			}
 		}
@@ -104,7 +104,7 @@ UChaosClothAsset::UChaosClothAsset(const FObjectInitializer& ObjectInitializer)
 #endif
 {
 	// Setup a single LOD's Cloth Collection
-	TSharedPtr<FManagedArrayCollection>& ClothCollection = ClothCollections.Emplace_GetRef(MakeShared<FManagedArrayCollection>());
+	TSharedRef<FManagedArrayCollection>& ClothCollection = ClothCollections.Emplace_GetRef(MakeShared<FManagedArrayCollection>());
 	UE::Chaos::ClothAsset::FCollectionClothFacade ClothFacade(ClothCollection);
 	ClothFacade.DefineSchema();
 
@@ -209,14 +209,14 @@ void UChaosClothAsset::BeginPostLoadInternal(FSkinnedAssetPostLoadContext& Conte
 	if (ClothCollections.IsEmpty())
 	{
 		UE_LOG(LogChaosClothAsset, Warning, TEXT("Invalid Cloth Collection (no LODs) found while loading Cloth Asset %s."), *GetFullName());
-		TSharedPtr<FManagedArrayCollection>& ClothCollection = ClothCollections.Emplace_GetRef(MakeShared<FManagedArrayCollection>());
+		TSharedRef<FManagedArrayCollection>& ClothCollection = ClothCollections.Emplace_GetRef(MakeShared<FManagedArrayCollection>());
 		UE::Chaos::ClothAsset::FCollectionClothFacade ClothFacade(ClothCollection);
 		ClothFacade.DefineSchema();
 		bAnyInvalidLods = true;
 	}
 	for(int32 LODIndex = 0; LODIndex < ClothCollections.Num(); ++LODIndex)
 	{
-		TSharedPtr<FManagedArrayCollection>& ClothCollection = ClothCollections[LODIndex];
+		TSharedRef<FManagedArrayCollection>& ClothCollection = ClothCollections[LODIndex];
 
 		UE::Chaos::ClothAsset::FCollectionClothFacade ClothFacade(ClothCollection);
 		if (!ClothFacade.IsValid())
@@ -394,7 +394,7 @@ void UChaosClothAsset::CalculateBounds()
 
 	FBox BoundingBox(ForceInit);
 
-	for (const TSharedPtr<FManagedArrayCollection>& ClothCollection : ClothCollections)
+	for (const TSharedRef<FManagedArrayCollection>& ClothCollection : ClothCollections)
 	{
 		const FCollectionClothConstFacade Cloth(ClothCollection);
 		const TConstArrayView<FVector3f> RenderPositionArray = Cloth.GetRenderPosition();
@@ -681,7 +681,7 @@ void UChaosClothAsset::SetPhysicsAsset(UPhysicsAsset* InPhysicsAsset)
 	using namespace UE::Chaos::ClothAsset;
 
 	PhysicsAsset = InPhysicsAsset;
-	for (TSharedPtr<FManagedArrayCollection>& ClothCollection : ClothCollections)
+	for (TSharedRef<FManagedArrayCollection>& ClothCollection : ClothCollections)
 	{
 		FCollectionClothFacade Cloth(ClothCollection);
 		Cloth.SetPhysicsAssetPathName(PhysicsAsset ? PhysicsAsset->GetPathName() : FString());
@@ -700,9 +700,8 @@ void UChaosClothAsset::SetSkeleton(USkeleton* InSkeleton, bool bRebuildModels)
 
 	RefSkeleton = Skeleton->GetReferenceSkeleton();
 
-	for (TSharedPtr<FManagedArrayCollection>& ClothCollection : ClothCollections)
+	for (TSharedRef<FManagedArrayCollection>& ClothCollection : ClothCollections)
 	{
-		check(ClothCollection.IsValid());
 		constexpr bool bBindSimMesh = true;
 		constexpr bool bBindRenderMesh = true;
 		FClothGeometryTools::BindMeshToRootBone(ClothCollection, bBindSimMesh, bBindRenderMesh);
@@ -755,9 +754,8 @@ void UChaosClothAsset::CopySimMeshToRenderMesh(UMaterialInterface* Material)
 		FString(TEXT("/Engine/EditorMaterials/Cloth/CameraLitDoubleSided.CameraLitDoubleSided"));
 
 	bool bAnyLodHasRenderMesh = false;
-	for (TSharedPtr<FManagedArrayCollection>& ClothCollection : ClothCollections)
+	for (TSharedRef<FManagedArrayCollection>& ClothCollection : ClothCollections)
 	{
-		check(ClothCollection.IsValid());
 		constexpr bool bSingleRenderPattern = true;
 		FClothGeometryTools::CopySimMeshToRenderMesh(ClothCollection, RenderMaterialPathName, bSingleRenderPattern);
 		bAnyLodHasRenderMesh = bAnyLodHasRenderMesh || FClothGeometryTools::HasRenderMesh(ClothCollection);
@@ -794,9 +792,8 @@ void UChaosClothAsset::BindSimMeshToRootBone()
 	using namespace UE::Chaos::ClothAsset;
 	check(ClothCollections.Num());
 
-	for (TSharedPtr<FManagedArrayCollection>& ClothCollection : ClothCollections)
+	for (TSharedRef<FManagedArrayCollection>& ClothCollection : ClothCollections)
 	{
-		check(ClothCollection.IsValid());
 		FClothGeometryTools::BindMeshToRootBone(ClothCollection, true, false);
 	}
 }
\ No newline at end of file
diff --git a/Engine/Plugins/Experimental/ChaosClothAsset/Source/ChaosClothAssetEngine/Private/ChaosClothAsset/ClothComponent.cpp b/Engine/Plugins/Experimental/ChaosClothAsset/Source/ChaosClothAssetEngine/Private/ChaosClothAsset/ClothComponent.cpp
--- a/Engine/Plugins/Experimental/ChaosClothAsset/Source/ChaosClothAssetEngine/Private/ChaosClothAsset/ClothComponent.cpp
+++ b/Engine/Plugins/Experimental/ChaosClothAsset/Source/ChaosClothAssetEngine/Private/ChaosClothAsset/ClothComponent.cpp
@@ -73,7 +73,7 @@ void UChaosClothComponent::ResetConfigProperties()
 		::Chaos::Softs::FCollectionPropertyMutableFacade CollectionPropertyMutableFacade(PropertyCollection);
 
 		// TODO: Make PropertyCollection also one per LOD.
-		if (GetClothAsset() && GetClothAsset()->GetClothCollections().Num() > 0 && GetClothAsset()->GetClothCollections()[0].IsValid())
+		if (GetClothAsset() && GetClothAsset()->GetClothCollections().Num() > 0)
 		{
 			CollectionPropertyMutableFacade.Copy(*GetClothAsset()->GetClothCollections()[0]);
 		}
diff --git a/Engine/Plugins/Experimental/ChaosClothAsset/Source/ChaosClothAssetEngine/Private/ChaosClothAsset/ClothEngineTools.cpp b/Engine/Plugins/Experimental/ChaosClothAsset/Source/ChaosClothAssetEngine/Private/ChaosClothAsset/ClothEngineTools.cpp
--- a/Engine/Plugins/Experimental/ChaosClothAsset/Source/ChaosClothAssetEngine/Private/ChaosClothAsset/ClothEngineTools.cpp
+++ b/Engine/Plugins/Experimental/ChaosClothAsset/Source/ChaosClothAssetEngine/Private/ChaosClothAsset/ClothEngineTools.cpp
@@ -8,7 +8,7 @@
 namespace UE::Chaos::ClothAsset
 {
 
-void FClothEngineTools::GenerateTethers(const TSharedPtr<FManagedArrayCollection>& ClothCollection, const FName& WeightMapName, const bool bGenerateGeodesicTethers)
+void FClothEngineTools::GenerateTethers(const TSharedRef<FManagedArrayCollection>& ClothCollection, const FName& WeightMapName, const bool bGenerateGeodesicTethers)
 {
 	FCollectionClothFacade ClothFacade(ClothCollection);
 	FClothGeometryTools::DeleteTethers(ClothCollection);
diff --git a/Engine/Plugins/Experimental/ChaosClothAsset/Source/ChaosClothAssetEngine/Private/ChaosClothAsset/ClothSimulationModel.cpp b/Engine/Plugins/Experimental/ChaosClothAsset/Source/ChaosClothAssetEngine/Private/ChaosClothAsset/ClothSimulationModel.cpp
--- a/Engine/Plugins/Experimental/ChaosClothAsset/Source/ChaosClothAssetEngine/Private/ChaosClothAsset/ClothSimulationModel.cpp
+++ b/Engine/Plugins/Experimental/ChaosClothAsset/Source/ChaosClothAssetEngine/Private/ChaosClothAsset/ClothSimulationModel.cpp
@@ -126,7 +126,7 @@ bool FChaosClothSimulationLodModel::Serialize(FArchive& Ar)
 	return true;
 }
 
-FChaosClothSimulationModel::FChaosClothSimulationModel(const TArray<TSharedPtr<const FManagedArrayCollection>>& ClothCollections, const FReferenceSkeleton& ReferenceSkeleton)
+FChaosClothSimulationModel::FChaosClothSimulationModel(const TArray<TSharedRef<const FManagedArrayCollection>>& ClothCollections, const FReferenceSkeleton& ReferenceSkeleton)
 {
 	using namespace UE::Chaos::ClothAsset;
 
diff --git a/Engine/Plugins/Experimental/ChaosClothAsset/Source/ChaosClothAssetEngine/Private/ChaosClothAsset/ClothSimulationModel.h b/Engine/Plugins/Experimental/ChaosClothAsset/Source/ChaosClothAssetEngine/Private/ChaosClothAsset/ClothSimulationModel.h
--- a/Engine/Plugins/Experimental/ChaosClothAsset/Source/ChaosClothAssetEngine/Private/ChaosClothAsset/ClothSimulationModel.h
+++ b/Engine/Plugins/Experimental/ChaosClothAsset/Source/ChaosClothAssetEngine/Private/ChaosClothAsset/ClothSimulationModel.h
@@ -100,7 +100,7 @@ struct FChaosClothSimulationModel
 	int32 ReferenceBoneIndex = INDEX_NONE;
 
 	FChaosClothSimulationModel() = default;
-	FChaosClothSimulationModel(const TArray<TSharedPtr<const FManagedArrayCollection>>& ClothCollections, const FReferenceSkeleton& ReferenceSkeleton);
+	FChaosClothSimulationModel(const TArray<TSharedRef<const FManagedArrayCollection>>& ClothCollections, const FReferenceSkeleton& ReferenceSkeleton);
 
 	int32 GetNumLods() const { return ClothSimulationLodModels.Num(); }
 
diff --git a/Engine/Plugins/Experimental/ChaosClothAsset/Source/ChaosClothAssetEngine/Public/ChaosClothAsset/ClothAsset.h b/Engine/Plugins/Experimental/ChaosClothAsset/Source/ChaosClothAssetEngine/Public/ChaosClothAsset/ClothAsset.h
--- a/Engine/Plugins/Experimental/ChaosClothAsset/Source/ChaosClothAssetEngine/Public/ChaosClothAsset/ClothAsset.h
+++ b/Engine/Plugins/Experimental/ChaosClothAsset/Source/ChaosClothAssetEngine/Public/ChaosClothAsset/ClothAsset.h
@@ -98,10 +98,10 @@ public:
 	//~ End USkinnedAsset interface
 
 	/** Return the enclosed Cloth Collection object. */
-	TArray<TSharedPtr<FManagedArrayCollection>>& GetClothCollections() { return ClothCollections; }
+	TArray<TSharedRef<FManagedArrayCollection>>& GetClothCollections() { return ClothCollections; }
 
 	/** Return the enclosed Cloth Collection object, const version. */
-	const TArray<TSharedPtr<const FManagedArrayCollection>>& GetClothCollections() const { return reinterpret_cast<const TArray<TSharedPtr<const FManagedArrayCollection>>&>(ClothCollections); }
+	const TArray<TSharedRef<const FManagedArrayCollection>>& GetClothCollections() const { return reinterpret_cast<const TArray<TSharedRef<const FManagedArrayCollection>>&>(ClothCollections); }
 
 	/** Return the cloth simulation ready LOD model data. */
 	TSharedPtr<const FChaosClothSimulationModel> GetClothSimulationModel() const { return ClothSimulationModel; }
@@ -275,7 +275,7 @@ private:
 	TArray<FMatrix> CachedComposedRefPoseMatrices;
 
 	/** Cloth Collection containing this asset data. One per LOD. */
-	TArray<TSharedPtr<FManagedArrayCollection>> ClothCollections;
+	TArray<TSharedRef<FManagedArrayCollection>> ClothCollections;
 
 	/** Reference skeleton created from the provided skeleton asset. */
 	FReferenceSkeleton RefSkeleton;
diff --git a/Engine/Plugins/Experimental/ChaosClothAsset/Source/ChaosClothAssetEngine/Public/ChaosClothAsset/ClothEngineTools.h b/Engine/Plugins/Experimental/ChaosClothAsset/Source/ChaosClothAssetEngine/Public/ChaosClothAsset/ClothEngineTools.h
--- a/Engine/Plugins/Experimental/ChaosClothAsset/Source/ChaosClothAssetEngine/Public/ChaosClothAsset/ClothEngineTools.h
+++ b/Engine/Plugins/Experimental/ChaosClothAsset/Source/ChaosClothAssetEngine/Public/ChaosClothAsset/ClothEngineTools.h
@@ -16,6 +16,6 @@ namespace UE::Chaos::ClothAsset
 	struct CHAOSCLOTHASSETENGINE_API FClothEngineTools
 	{
 		/** Generate tether data. */
-		static void GenerateTethers(const TSharedPtr<FManagedArrayCollection>& ClothCollection, const FName& WeightMap, const bool bGeodesicTethers);
+		static void GenerateTethers(const TSharedRef<FManagedArrayCollection>& ClothCollection, const FName& WeightMap, const bool bGeodesicTethers);
 	};
 }  // End namespace UE::Chaos::ClothAsset
\ No newline at end of file
diff --git a/Engine/Plugins/Experimental/ChaosClothAssetEditor/Source/ChaosClothAssetDataflowNodes/Private/ChaosClothAsset/ClothDataflowTools.cpp b/Engine/Plugins/Experimental/ChaosClothAssetEditor/Source/ChaosClothAssetDataflowNodes/Private/ChaosClothAsset/ClothDataflowTools.cpp
--- a/Engine/Plugins/Experimental/ChaosClothAssetEditor/Source/ChaosClothAssetDataflowNodes/Private/ChaosClothAsset/ClothDataflowTools.cpp
+++ b/Engine/Plugins/Experimental/ChaosClothAssetEditor/Source/ChaosClothAssetDataflowNodes/Private/ChaosClothAsset/ClothDataflowTools.cpp
@@ -353,10 +353,9 @@ namespace UE::Chaos::ClothAsset
 		};
 	} // Private
 
-	void FClothDataflowTools::AddRenderPatternFromSkeletalMeshSection(const TSharedPtr<FManagedArrayCollection>& ClothCollection, const FSkeletalMeshLODModel& SkeletalMeshModel, const int32 SectionIndex, const FString& RenderMaterialPathName)
+	void FClothDataflowTools::AddRenderPatternFromSkeletalMeshSection(const TSharedRef<FManagedArrayCollection>& ClothCollection, const FSkeletalMeshLODModel& SkeletalMeshModel, const int32 SectionIndex, const FString& RenderMaterialPathName)
 	{
 		check(SectionIndex < SkeletalMeshModel.Sections.Num());
-		check(ClothCollection.IsValid());
 
 		FCollectionClothFacade Cloth(ClothCollection);
 		check(Cloth.IsValid());
@@ -415,7 +414,7 @@ namespace UE::Chaos::ClothAsset
 		ClothPatternFacade.SetRenderMaterialPathName(RenderMaterialPathName);
 	}
 	
-	void FClothDataflowTools::AddSimPatternsFromSkeletalMeshSection(const TSharedPtr<FManagedArrayCollection>& ClothCollection, const FSkeletalMeshLODModel& SkeletalMeshModel, const int32 SectionIndex, const int32 UVChannelIndex)
+	void FClothDataflowTools::AddSimPatternsFromSkeletalMeshSection(const TSharedRef<FManagedArrayCollection>& ClothCollection, const FSkeletalMeshLODModel& SkeletalMeshModel, const int32 SectionIndex, const int32 UVChannelIndex)
 	{
 		check(SectionIndex < SkeletalMeshModel.Sections.Num());
 
diff --git a/Engine/Plugins/Experimental/ChaosClothAssetEditor/Source/ChaosClothAssetDataflowNodes/Private/ChaosClothAsset/ClothDataflowTools.h b/Engine/Plugins/Experimental/ChaosClothAssetEditor/Source/ChaosClothAssetDataflowNodes/Private/ChaosClothAsset/ClothDataflowTools.h
--- a/Engine/Plugins/Experimental/ChaosClothAssetEditor/Source/ChaosClothAssetDataflowNodes/Private/ChaosClothAsset/ClothDataflowTools.h
+++ b/Engine/Plugins/Experimental/ChaosClothAssetEditor/Source/ChaosClothAssetDataflowNodes/Private/ChaosClothAsset/ClothDataflowTools.h
@@ -15,8 +15,8 @@ namespace UE::Chaos::ClothAsset
 	*/
 	struct FClothDataflowTools
 	{
-		static void AddRenderPatternFromSkeletalMeshSection(const TSharedPtr<FManagedArrayCollection>& ClothCollection, const FSkeletalMeshLODModel& SkeletalMeshModel, const int32 SectionIndex, const FString& RenderMaterialPathName);
+		static void AddRenderPatternFromSkeletalMeshSection(const TSharedRef<FManagedArrayCollection>& ClothCollection, const FSkeletalMeshLODModel& SkeletalMeshModel, const int32 SectionIndex, const FString& RenderMaterialPathName);
 
-		static void AddSimPatternsFromSkeletalMeshSection(const TSharedPtr<FManagedArrayCollection>& ClothCollection, const FSkeletalMeshLODModel& SkeletalMeshModel, const int32 SectionIndex, const int32 UVChannelIndex);
+		static void AddSimPatternsFromSkeletalMeshSection(const TSharedRef<FManagedArrayCollection>& ClothCollection, const FSkeletalMeshLODModel& SkeletalMeshModel, const int32 SectionIndex, const int32 UVChannelIndex);
 	};
 }  // End namespace UE::Chaos::ClothAsset
diff --git a/Engine/Plugins/Experimental/ChaosClothAssetEditor/Source/ChaosClothAssetDataflowNodes/Private/ChaosClothAsset/DatasmithImportNode.cpp b/Engine/Plugins/Experimental/ChaosClothAssetEditor/Source/ChaosClothAssetDataflowNodes/Private/ChaosClothAsset/DatasmithImportNode.cpp
--- a/Engine/Plugins/Experimental/ChaosClothAssetEditor/Source/ChaosClothAssetDataflowNodes/Private/ChaosClothAsset/DatasmithImportNode.cpp
+++ b/Engine/Plugins/Experimental/ChaosClothAssetEditor/Source/ChaosClothAssetDataflowNodes/Private/ChaosClothAsset/DatasmithImportNode.cpp
@@ -82,7 +82,7 @@ bool FChaosClothAssetDatasmithImportNode::EvaluateImpl(Dataflow::FContext& Conte
 		UChaosClothAsset* const DatasmithClothAsset = Cast<UChaosClothAsset>(ClothObject);
 		if (ensure(DatasmithClothAsset))
 		{
-			if (DatasmithClothAsset->GetClothCollections().Num() > 0 && DatasmithClothAsset->GetClothCollections()[0].IsValid())
+			if (DatasmithClothAsset->GetClothCollections().Num() > 0)
 			{
 				DatasmithClothAsset->GetClothCollections()[0]->CopyTo(&OutCollection);
 				return true;
diff --git a/Engine/Plugins/Experimental/ChaosClothAssetEditor/Source/ChaosClothAssetDataflowNodes/Private/ChaosClothAsset/ImportNode.cpp b/Engine/Plugins/Experimental/ChaosClothAssetEditor/Source/ChaosClothAssetDataflowNodes/Private/ChaosClothAsset/ImportNode.cpp
--- a/Engine/Plugins/Experimental/ChaosClothAssetEditor/Source/ChaosClothAssetDataflowNodes/Private/ChaosClothAsset/ImportNode.cpp
+++ b/Engine/Plugins/Experimental/ChaosClothAssetEditor/Source/ChaosClothAssetDataflowNodes/Private/ChaosClothAsset/ImportNode.cpp
@@ -64,8 +64,8 @@ void FChaosClothAssetImportNode::Evaluate(Dataflow::FContext& Context, const FDa
 		{
 			const TObjectPtr<const UChaosClothAsset>& InClothAsset = GetValue<TObjectPtr<const UChaosClothAsset>>(Context, &ClothAsset);
 			const int32& InImportLod = GetValue<int32>(Context, &ImportLod);
-			const TArray<TSharedPtr<const FManagedArrayCollection>>& InClothCollections = InClothAsset->GetClothCollections();
-			if (InImportLod >= 0 && InClothCollections.Num() > InImportLod && InClothCollections[InImportLod].IsValid())
+			const TArray<TSharedRef<const FManagedArrayCollection>>& InClothCollections = InClothAsset->GetClothCollections();
+			if (InImportLod >= 0 && InClothCollections.Num() > InImportLod)
 			{
 				const FCollectionClothConstFacade InClothFacade(InClothCollections[InImportLod]);
 				ClothFacade.Initialize(InClothFacade);
diff --git a/Engine/Plugins/Experimental/ChaosClothAssetEditor/Source/ChaosClothAssetDataflowNodes/Private/ChaosClothAsset/StaticMeshImportNode.cpp b/Engine/Plugins/Experimental/ChaosClothAssetEditor/Source/ChaosClothAssetDataflowNodes/Private/ChaosClothAsset/StaticMeshImportNode.cpp
--- a/Engine/Plugins/Experimental/ChaosClothAssetEditor/Source/ChaosClothAssetDataflowNodes/Private/ChaosClothAsset/StaticMeshImportNode.cpp
+++ b/Engine/Plugins/Experimental/ChaosClothAssetEditor/Source/ChaosClothAssetDataflowNodes/Private/ChaosClothAsset/StaticMeshImportNode.cpp
@@ -93,7 +93,11 @@ bool BuildSkeletalMeshModelFromMeshDescription(const FMeshDescription* const InM
 	return true;
 }
 
-void InitializeDataFromMeshDescription(const FMeshDescription* const InMeshDescription, const FMeshBuildSettings& InBuildSettings, const TArray<FStaticMaterial>& StaticMaterials, const TSharedPtr<FManagedArrayCollection>& ClothCollection)
+void InitializeDataFromMeshDescription(
+	const FMeshDescription* const InMeshDescription,
+	const FMeshBuildSettings& InBuildSettings,
+	const TArray<FStaticMaterial>& StaticMaterials,
+	const TSharedRef<FManagedArrayCollection>& ClothCollection)
 {
 	FSkeletalMeshLODModel SkeletalMeshModel;
 	if (BuildSkeletalMeshModelFromMeshDescription(InMeshDescription, InBuildSettings, SkeletalMeshModel))
diff --git a/Engine/Plugins/Experimental/ChaosClothAssetEditor/Source/ChaosClothAssetDataflowNodes/Private/ChaosClothAsset/TerminalNode.cpp b/Engine/Plugins/Experimental/ChaosClothAssetEditor/Source/ChaosClothAssetDataflowNodes/Private/ChaosClothAsset/TerminalNode.cpp
--- a/Engine/Plugins/Experimental/ChaosClothAssetEditor/Source/ChaosClothAssetDataflowNodes/Private/ChaosClothAsset/TerminalNode.cpp
+++ b/Engine/Plugins/Experimental/ChaosClothAssetEditor/Source/ChaosClothAssetDataflowNodes/Private/ChaosClothAsset/TerminalNode.cpp
@@ -29,7 +29,7 @@ void FChaosClothAssetTerminalNode::SetAssetValue(TObjectPtr<UObject> Asset, Data
 		const TArray<const FManagedArrayCollection*> CollectionLods = GetCollectionLods();
 
 		// Reset the asset's collection
-		TArray<TSharedPtr<FManagedArrayCollection>>& ClothCollections = ClothAsset->GetClothCollections();
+		TArray<TSharedRef<FManagedArrayCollection>>& ClothCollections = ClothAsset->GetClothCollections();
 		ClothCollections.Reset(CollectionLods.Num());
 
 		// Reset the asset's material list
@@ -42,7 +42,7 @@ void FChaosClothAssetTerminalNode::SetAssetValue(TObjectPtr<UObject> Asset, Data
 		for (int32 LodIndex = 0; LodIndex < CollectionLods.Num(); ++LodIndex)
 		{
 			// New LOD
-			TSharedPtr<FManagedArrayCollection>& ClothCollection = ClothCollections.Emplace_GetRef(MakeShared<FManagedArrayCollection>());
+			TSharedRef<FManagedArrayCollection>& ClothCollection = ClothCollections.Emplace_GetRef(MakeShared<FManagedArrayCollection>());
 			UE::Chaos::ClothAsset::FCollectionClothFacade ClothFacade(ClothCollection);
 			ClothFacade.DefineSchema();
 
@@ -131,8 +131,8 @@ void FChaosClothAssetTerminalNode::SetAssetValue(TObjectPtr<UObject> Asset, Data
 		// Make sure that whatever happens there is always at least one empty LOD to avoid crashing the render data
 		if (ClothCollections.Num() < 1)
 		{
-			TSharedPtr<FManagedArrayCollection>& ClothCollection = ClothCollections.Emplace_GetRef(MakeShared<FManagedArrayCollection>());
-			UE::Chaos::ClothAsset::FCollectionClothFacade ClothFacade(ClothCollection);
+			TSharedRef<FManagedArrayCollection>& ClothCollection = ClothCollections.Emplace_GetRef(MakeShared<FManagedArrayCollection>());
+			FCollectionClothFacade ClothFacade(ClothCollection);
 			ClothFacade.DefineSchema();
 		}
 
diff --git a/Engine/Plugins/Experimental/ChaosClothAssetEditor/Source/ChaosClothAssetEditor/Private/ChaosClothAsset/ClothEditorMode.cpp b/Engine/Plugins/Experimental/ChaosClothAssetEditor/Source/ChaosClothAssetEditor/Private/ChaosClothAsset/ClothEditorMode.cpp
--- a/Engine/Plugins/Experimental/ChaosClothAssetEditor/Source/ChaosClothAssetEditor/Private/ChaosClothAsset/ClothEditorMode.cpp
+++ b/Engine/Plugins/Experimental/ChaosClothAssetEditor/Source/ChaosClothAssetEditor/Private/ChaosClothAsset/ClothEditorMode.cpp
@@ -495,12 +495,12 @@ void UChaosClothAssetEditorMode::ReinitializeDynamicMeshComponents()
 		return;
 	}
 
-	const UE::Chaos::ClothAsset::FCollectionClothConstFacade ClothFacade(Collection);
+	const UE::Chaos::ClothAsset::FCollectionClothConstFacade ClothFacade(Collection.ToSharedRef());
 
 	UE::Geometry::FDynamicMesh3 LodMesh;
 	LodMesh.EnableAttributes();
 	FClothPatternToDynamicMesh Converter;
-	Converter.Convert(Collection, INDEX_NONE, ConstructionViewMode, LodMesh);
+	Converter.Convert(Collection.ToSharedRef(), INDEX_NONE, ConstructionViewMode, LodMesh);
 
 	if (ConstructionViewMode == EClothPatternVertexType::Sim2D)
 	{
diff --git a/Engine/Plugins/Experimental/ChaosClothAssetEditor/Source/ChaosClothAssetEditor/Private/ChaosClothAsset/ClothEditorToolkit.cpp b/Engine/Plugins/Experimental/ChaosClothAssetEditor/Source/ChaosClothAssetEditor/Private/ChaosClothAsset/ClothEditorToolkit.cpp
--- a/Engine/Plugins/Experimental/ChaosClothAssetEditor/Source/ChaosClothAssetEditor/Private/ChaosClothAsset/ClothEditorToolkit.cpp
+++ b/Engine/Plugins/Experimental/ChaosClothAssetEditor/Source/ChaosClothAssetEditor/Private/ChaosClothAsset/ClothEditorToolkit.cpp
@@ -898,7 +898,7 @@ void FChaosClothAssetEditorToolkit::OnNodeSelectionChanged(const TSet<UObject*>&
 				if (Output->GetType() == FName("FManagedArrayCollection"))
 				{
 					const FManagedArrayCollection DefaultValue;
-					TSharedPtr<FManagedArrayCollection> Collection = MakeShared<FManagedArrayCollection>(Output->GetValue<FManagedArrayCollection>(*Context, DefaultValue));
+					TSharedRef<FManagedArrayCollection> Collection = MakeShared<FManagedArrayCollection>(Output->GetValue<FManagedArrayCollection>(*Context, DefaultValue));
 
 					// see if the output collection is a ClothCollection
 					const UE::Chaos::ClothAsset::FCollectionClothConstFacade ClothFacade(Collection);
diff --git a/Engine/Plugins/Experimental/ChaosClothAssetEditor/Source/ChaosClothAssetEditorTools/Private/ChaosClothAsset/ClothWeightMapPaintTool.cpp b/Engine/Plugins/Experimental/ChaosClothAssetEditor/Source/ChaosClothAssetEditorTools/Private/ChaosClothAsset/ClothWeightMapPaintTool.cpp
--- a/Engine/Plugins/Experimental/ChaosClothAssetEditor/Source/ChaosClothAssetEditorTools/Private/ChaosClothAsset/ClothWeightMapPaintTool.cpp
+++ b/Engine/Plugins/Experimental/ChaosClothAssetEditor/Source/ChaosClothAssetEditorTools/Private/ChaosClothAsset/ClothWeightMapPaintTool.cpp
@@ -310,7 +310,7 @@ void UClothEditorWeightMapPaintTool::Setup()
 
 			bHaveDynamicMeshToWeightConversion = bHasNonManifoldMapping || bHas2D3DConversion;
 
-			FCollectionClothConstFacade Cloth(ClothCollection);
+			FCollectionClothConstFacade Cloth(ClothCollection.ToSharedRef());
 			check(Cloth.IsValid());
 			if (bHasNonManifoldMapping)
 			{
diff --git a/Engine/Plugins/Experimental/ChaosClothAssetEditor/Source/ChaosClothAssetTools/Private/ChaosClothAsset/ClothAssetBuilderEditor.cpp b/Engine/Plugins/Experimental/ChaosClothAssetEditor/Source/ChaosClothAssetTools/Private/ChaosClothAsset/ClothAssetBuilderEditor.cpp
--- a/Engine/Plugins/Experimental/ChaosClothAssetEditor/Source/ChaosClothAssetTools/Private/ChaosClothAsset/ClothAssetBuilderEditor.cpp
+++ b/Engine/Plugins/Experimental/ChaosClothAssetEditor/Source/ChaosClothAssetTools/Private/ChaosClothAsset/ClothAssetBuilderEditor.cpp
@@ -34,7 +34,7 @@ void UClothAssetBuilderEditor::BuildLod(FSkeletalMeshLODModel& LODModel, const U
 	// Create a table to remap the LOD materials to the asset materials
 	const TArray<FSkeletalMaterial>& Materials = ClothAsset.GetMaterials();
 
-	const TSharedPtr<const FManagedArrayCollection> ClothCollection = ClothAsset.GetClothCollections()[LodIndex];
+	const TSharedRef<const FManagedArrayCollection> ClothCollection = ClothAsset.GetClothCollections()[LodIndex];
 
 	const FCollectionClothConstFacade ClothFacade(ClothCollection);
 
diff --git a/Engine/Plugins/Experimental/ChaosClothAssetEditor/Source/ChaosClothAssetTools/Private/ChaosClothAsset/ClothAssetFactory.cpp b/Engine/Plugins/Experimental/ChaosClothAssetEditor/Source/ChaosClothAssetTools/Private/ChaosClothAsset/ClothAssetFactory.cpp
--- a/Engine/Plugins/Experimental/ChaosClothAssetEditor/Source/ChaosClothAssetTools/Private/ChaosClothAsset/ClothAssetFactory.cpp
+++ b/Engine/Plugins/Experimental/ChaosClothAssetEditor/Source/ChaosClothAssetTools/Private/ChaosClothAsset/ClothAssetFactory.cpp
@@ -24,8 +24,8 @@ UObject* UChaosClothAssetFactory::FactoryCreateNew(UClass* Class, UObject* Paren
 	ClothAsset->MarkPackageDirty();
 
 	// Add an empty default LOD, to avoid LOD mismatch with render data
-	TArray<TSharedPtr<FManagedArrayCollection>>& ClothCollections = ClothAsset->GetClothCollections();
-	TSharedPtr<FManagedArrayCollection>& ClothCollection = ClothCollections.Emplace_GetRef(MakeShared<FManagedArrayCollection>());
+	TArray<TSharedRef<FManagedArrayCollection>>& ClothCollections = ClothAsset->GetClothCollections();
+	TSharedRef<FManagedArrayCollection>& ClothCollection = ClothCollections.Emplace_GetRef(MakeShared<FManagedArrayCollection>());
 
 	FCollectionClothFacade ClothFacade(ClothCollection);
 	ClothFacade.DefineSchema();
diff --git a/Engine/Plugins/Experimental/ChaosClothAssetEditor/Source/ChaosClothAssetTools/Private/ChaosClothAsset/ClothPatternToDynamicMesh.cpp b/Engine/Plugins/Experimental/ChaosClothAssetEditor/Source/ChaosClothAssetTools/Private/ChaosClothAsset/ClothPatternToDynamicMesh.cpp
--- a/Engine/Plugins/Experimental/ChaosClothAssetEditor/Source/ChaosClothAssetTools/Private/ChaosClothAsset/ClothPatternToDynamicMesh.cpp
+++ b/Engine/Plugins/Experimental/ChaosClothAssetEditor/Source/ChaosClothAssetTools/Private/ChaosClothAsset/ClothPatternToDynamicMesh.cpp
@@ -586,7 +586,7 @@ private:
 };
 
 
-void FClothPatternToDynamicMesh::Convert(const TSharedPtr<const FManagedArrayCollection> ClothCollection, int32 PatternIndex, EClothPatternVertexType VertexDataType, UE::Geometry::FDynamicMesh3& MeshOut)
+void FClothPatternToDynamicMesh::Convert(const TSharedRef<const FManagedArrayCollection> ClothCollection, int32 PatternIndex, EClothPatternVertexType VertexDataType, UE::Geometry::FDynamicMesh3& MeshOut)
 {
 	const FCollectionClothConstFacade ClothFacade(ClothCollection);
 
@@ -641,9 +641,8 @@ void FClothPatternToDynamicMesh::Convert(const TSharedPtr<const FManagedArrayCol
 
 void FClothPatternToDynamicMesh::Convert(const UChaosClothAsset* ClothAssetMeshIn, int32 LODIndex, int32 PatternIndex, EClothPatternVertexType VertexDataType, UE::Geometry::FDynamicMesh3& MeshOut)
 {
-	const TArray<TSharedPtr<const FManagedArrayCollection>>& ClothCollections = ClothAssetMeshIn->GetClothCollections();
+	const TArray<TSharedRef<const FManagedArrayCollection>>& ClothCollections = ClothAssetMeshIn->GetClothCollections();
 	check(ClothCollections.IsValidIndex(LODIndex));
-	check(ClothCollections[LODIndex].IsValid());
 
 	Convert(ClothCollections[LODIndex], PatternIndex, VertexDataType, MeshOut);
 }
diff --git a/Engine/Plugins/Experimental/ChaosClothAssetEditor/Source/ChaosClothAssetTools/Private/ChaosClothAsset/ClothingAssetToClothAssetExporter.cpp b/Engine/Plugins/Experimental/ChaosClothAssetEditor/Source/ChaosClothAssetTools/Private/ChaosClothAsset/ClothingAssetToClothAssetExporter.cpp
--- a/Engine/Plugins/Experimental/ChaosClothAssetEditor/Source/ChaosClothAssetTools/Private/ChaosClothAsset/ClothingAssetToClothAssetExporter.cpp
+++ b/Engine/Plugins/Experimental/ChaosClothAssetEditor/Source/ChaosClothAssetTools/Private/ChaosClothAsset/ClothingAssetToClothAssetExporter.cpp
@@ -95,10 +95,10 @@ void UClothingAssetToChaosClothAssetExporter::Export(const UClothingAssetBase* C
 
 	UChaosClothAsset* const ClothAsset = CastChecked<UChaosClothAsset>(ExportedAsset);
 	check(ClothAsset);
-	TArray<TSharedPtr<FManagedArrayCollection>>& ClothCollections = ClothAsset->GetClothCollections();
+	TArray<TSharedRef<FManagedArrayCollection>>& ClothCollections = ClothAsset->GetClothCollections();
 	
 	const int32 NumLods = ClothingAssetCommon->LodData.Num();
-	ClothCollections.SetNum(NumLods);
+	ClothCollections.Reset(NumLods);
 
 	// Create the LODs
 	for (int32 LodIndex = 0; LodIndex < NumLods; ++LodIndex)
@@ -106,13 +106,9 @@ void UClothingAssetToChaosClothAssetExporter::Export(const UClothingAssetBase* C
 		const FClothLODDataCommon& ClothLODData = ClothingAssetCommon->LodData[LodIndex];
 		const FClothPhysicalMeshData& PhysicalMeshData = ClothLODData.PhysicalMeshData;
 
-		TSharedPtr<FManagedArrayCollection>& ClothCollection = ClothCollections[LodIndex];
-		if (!ClothCollection.IsValid())
-		{
-			ClothCollection = MakeShared<FManagedArrayCollection>();
-			FCollectionClothFacade Cloth(ClothCollection);
-			Cloth.DefineSchema();
-		}
+		TSharedRef<FManagedArrayCollection>& ClothCollection = ClothCollections.Emplace_GetRef(MakeShared<FManagedArrayCollection>());
+		FCollectionClothFacade ClothFacade(ClothCollection);
+		ClothFacade.DefineSchema();
 
 		using namespace UE::Chaos::ClothAsset::Private;
 
@@ -137,14 +133,9 @@ void UClothingAssetToChaosClothAssetExporter::Export(const UClothingAssetBase* C
 	else
 	{
 		// Make sure that at least one empty LOD is always created
-		ClothCollections.SetNum(1);
-		TSharedPtr<FManagedArrayCollection>& ClothCollection = ClothCollections[1];
-		if (!ClothCollection.IsValid())
-		{
-			ClothCollection = MakeShared<FManagedArrayCollection>();
-			FCollectionClothFacade Cloth(ClothCollection);
-			Cloth.DefineSchema();
-		}
+		TSharedRef<FManagedArrayCollection>& ClothCollection = ClothCollections.Emplace_GetRef(MakeShared<FManagedArrayCollection>());
+		FCollectionClothFacade Cloth(ClothCollection);
+		Cloth.DefineSchema();
 	}
 
 	// Assign the physics asset if any (must be done after having added the LODs)
diff --git a/Engine/Plugins/Experimental/ChaosClothAssetEditor/Source/ChaosClothAssetTools/Public/ChaosClothAsset/ClothPatternToDynamicMesh.h b/Engine/Plugins/Experimental/ChaosClothAssetEditor/Source/ChaosClothAssetTools/Public/ChaosClothAsset/ClothPatternToDynamicMesh.h
--- a/Engine/Plugins/Experimental/ChaosClothAssetEditor/Source/ChaosClothAssetTools/Public/ChaosClothAsset/ClothPatternToDynamicMesh.h
+++ b/Engine/Plugins/Experimental/ChaosClothAssetEditor/Source/ChaosClothAssetTools/Public/ChaosClothAsset/ClothPatternToDynamicMesh.h
@@ -23,7 +23,7 @@ class CHAOSCLOTHASSETTOOLS_API FClothPatternToDynamicMesh
 {
 public:
 
-	void Convert(const TSharedPtr<const FManagedArrayCollection> ClothCollection, int32 PatternIndex, EClothPatternVertexType VertexDataType, UE::Geometry::FDynamicMesh3& MeshOut);
+	void Convert(const TSharedRef<const FManagedArrayCollection> ClothCollection, int32 PatternIndex, EClothPatternVertexType VertexDataType, UE::Geometry::FDynamicMesh3& MeshOut);
 
 	void Convert(const UChaosClothAsset* ClothAssetMeshIn, int32 LODIndex, int32 PatternIndex, EClothPatternVertexType VertexDataType, UE::Geometry::FDynamicMesh3& MeshOut);
 };
