Fixed CIS linking issues

diff --git a/Engine/Plugins/Experimental/ChaosClothAsset/Source/ChaosClothAssetEngine/Public/ChaosClothAsset/ClothAsset.h b/Engine/Plugins/Experimental/ChaosClothAsset/Source/ChaosClothAssetEngine/Public/ChaosClothAsset/ClothAsset.h
--- a/Engine/Plugins/Experimental/ChaosClothAsset/Source/ChaosClothAssetEngine/Public/ChaosClothAsset/ClothAsset.h
+++ b/Engine/Plugins/Experimental/ChaosClothAsset/Source/ChaosClothAssetEngine/Public/ChaosClothAsset/ClothAsset.h
@@ -78,6 +78,8 @@ public:
 	virtual const USkeleton* GetSkeleton() const override						{ return Skeleton; }
 	virtual void SetSkeleton(USkeleton* InSkeleton) override					{ constexpr bool bRebuildModels = true; SetSkeleton(InSkeleton, bRebuildModels); }
 	virtual UMeshDeformer* GetDefaultMeshDeformer() const override				{ return nullptr; }
+	virtual class UMaterialInterface* GetOverlayMaterial() const override		{ return nullptr; }
+	virtual float GetOverlayMaterialMaxDrawDistance() const override			{ return 0.f; }
 	virtual bool IsValidLODIndex(int32 Index) const override					{ return LODInfo.IsValidIndex(Index); }
 	virtual int32 GetMinLodIdx(bool bForceLowestLODIdx = false) const override	{ return 0; }
 	virtual bool NeedCPUData(int32 LODIndex) const override						{ return false; }
diff --git a/Engine/Source/Runtime/Engine/Classes/Engine/SkinnedAsset.h b/Engine/Source/Runtime/Engine/Classes/Engine/SkinnedAsset.h
--- a/Engine/Source/Runtime/Engine/Classes/Engine/SkinnedAsset.h
+++ b/Engine/Source/Runtime/Engine/Classes/Engine/SkinnedAsset.h
@@ -176,9 +176,9 @@ public:
 	ENGINE_API virtual UMeshDeformer* GetDefaultMeshDeformer() const
 	PURE_VIRTUAL(USkinnedAsset::GetDefaultMeshDeformer, return nullptr;);
 
-	virtual class UMaterialInterface* GetOverlayMaterial() const
+	ENGINE_API virtual class UMaterialInterface* GetOverlayMaterial() const
 	PURE_VIRTUAL(USkinnedAsset::GetOverlayMaterial, return nullptr;);
-	virtual float GetOverlayMaterialMaxDrawDistance() const
+	ENGINE_API virtual float GetOverlayMaterialMaxDrawDistance() const
 	PURE_VIRTUAL(USkinnedAsset::GetOverlayMaterialMaxDrawDistance, return 0.f;);
 	
 	/** Return true if given index's LOD is valid */
