Chaos Cloth: Caching Up/Down transitions for LODs on the terminal node.

Calculating the Up/Down transitions can be expensive, especially for heavier meshes. However, since it happens at the terminal node, they are recalculated anytime anything changes on the asset, including simulation parameters, which do not affect the transitions.

This adds the ability to cache the transitions on the terminal node and hash the inputs. The transitions are reused if the hash didn't change.

diff --git a/Engine/Plugins/Experimental/ChaosClothAsset/Source/ChaosClothAssetEngine/Private/ChaosClothAsset/ClothAsset.cpp b/Engine/Plugins/Experimental/ChaosClothAsset/Source/ChaosClothAssetEngine/Private/ChaosClothAsset/ClothAsset.cpp
--- a/Engine/Plugins/Experimental/ChaosClothAsset/Source/ChaosClothAssetEngine/Private/ChaosClothAsset/ClothAsset.cpp
+++ b/Engine/Plugins/Experimental/ChaosClothAsset/Source/ChaosClothAssetEngine/Private/ChaosClothAsset/ClothAsset.cpp
@@ -419,7 +419,7 @@ void UChaosClothAsset::UpdateSkeleton(bool bRebuildClothSimulationModel)
 	}
 }
 
-void UChaosClothAsset::Build()
+void UChaosClothAsset::Build(TArray<FChaosClothAssetLodTransitionDataCache>* InOutTransitionCache)
 {
 	using namespace UE::Chaos::ClothAsset;
 
@@ -449,7 +449,7 @@ void UChaosClothAsset::Build()
 	LODInfo.AddDefaulted(NumLods);  // TODO: Expose some properties to fill up the LOD infos
 
 	// Build simulation model
-	BuildClothSimulationModel();
+	BuildClothSimulationModel(InOutTransitionCache);
 
 	// Rebuild LOD Model
 #if WITH_EDITORONLY_DATA
@@ -543,9 +543,10 @@ void UChaosClothAsset::BuildMeshModel()
 }
 #endif  // #if WITH_EDITORONLY_DATA
 
-void UChaosClothAsset::BuildClothSimulationModel()
+void UChaosClothAsset::BuildClothSimulationModel(TArray<FChaosClothAssetLodTransitionDataCache>* InOutTransitionCache)
 {
-	ClothSimulationModel = MakeShared<FChaosClothSimulationModel>(const_cast<const UChaosClothAsset*>(this)->GetClothCollections(), GetRefSkeleton());
+	ClothSimulationModel = MakeShared<FChaosClothSimulationModel>(const_cast<const UChaosClothAsset*>(this)->GetClothCollections(), 
+		GetRefSkeleton(), InOutTransitionCache);
 }
 
 const FMeshUVChannelInfo* UChaosClothAsset::GetUVChannelData(int32 MaterialIndex) const
@@ -729,7 +730,7 @@ void UChaosClothAsset::UpdateSkeletonFromCollection(bool bRebuildModels)
 
 	if (!InSkeleton)
 	{
-		SetSkeleton(nullptr, true);
+		SetSkeleton(nullptr, bRebuildModels);
 	}
 	else
 	{
diff --git a/Engine/Plugins/Experimental/ChaosClothAsset/Source/ChaosClothAssetEngine/Private/ChaosClothAsset/ClothLodTransitionDataCache.cpp b/Engine/Plugins/Experimental/ChaosClothAsset/Source/ChaosClothAssetEngine/Private/ChaosClothAsset/ClothLodTransitionDataCache.cpp
new file mode 100644
--- /dev/null
+++ b/Engine/Plugins/Experimental/ChaosClothAsset/Source/ChaosClothAssetEngine/Private/ChaosClothAsset/ClothLodTransitionDataCache.cpp
@@ -0,0 +1,21 @@
+// Copyright Epic Games, Inc. All Rights Reserved.
+
+#include "ChaosClothAsset/ClothLodTransitionDataCache.h"
+
+bool FChaosClothAssetLodTransitionDataCache::Serialize(FArchive& Ar)
+{
+	// Serialize normal tagged property data
+	if (Ar.IsLoading() || Ar.IsSaving())
+	{
+		UScriptStruct* const Struct = FChaosClothAssetLodTransitionDataCache::StaticStruct();
+		Struct->SerializeTaggedProperties(Ar, (uint8*)this, Struct, nullptr);
+	}
+
+	Ar << ModelHash;
+
+	Ar << LODTransitionUpData;
+	Ar << LODTransitionDownData;
+
+	// Return true to confirm that serialization has already been taken care of
+	return true;
+}
\ No newline at end of file
diff --git a/Engine/Plugins/Experimental/ChaosClothAsset/Source/ChaosClothAssetEngine/Private/ChaosClothAsset/ClothSimulationModel.cpp b/Engine/Plugins/Experimental/ChaosClothAsset/Source/ChaosClothAssetEngine/Private/ChaosClothAsset/ClothSimulationModel.cpp
--- a/Engine/Plugins/Experimental/ChaosClothAsset/Source/ChaosClothAssetEngine/Private/ChaosClothAsset/ClothSimulationModel.cpp
+++ b/Engine/Plugins/Experimental/ChaosClothAsset/Source/ChaosClothAssetEngine/Private/ChaosClothAsset/ClothSimulationModel.cpp
@@ -3,6 +3,7 @@
 #include "ChaosClothAsset/ClothSimulationModel.h"
 #include "ChaosClothAsset/CollectionClothFacade.h"
 #include "ChaosClothAsset/ClothGeometryTools.h"
+#include "ChaosClothAsset/ClothLodTransitionDataCache.h"
 #include "Utils/ClothingMeshUtils.h"
 #include "ReferenceSkeleton.h"
 
@@ -105,6 +106,26 @@ namespace UE::Chaos::ClothAsset
 		}
 		return ReferenceBoneIndex;
 	}
+
+	FMD5Hash CalculateTransitionDataHash(const FChaosClothSimulationLodModel& CurrModel)
+	{
+		FMD5 MD5;
+		const TArray<FVector3f>& Positions = CurrModel.Positions;
+		const int32 NumPositions = Positions.Num();
+		MD5.Update(reinterpret_cast<const uint8*>(&NumPositions), sizeof(NumPositions));
+		MD5.Update(reinterpret_cast<const uint8*>(Positions.GetData()),
+			NumPositions * sizeof(FVector3f));
+
+		const TArray<uint32>& Indices = CurrModel.Indices;
+		const int32 NumIndices = Indices.Num();
+		MD5.Update(reinterpret_cast<const uint8*>(&NumIndices), sizeof(NumIndices));
+		MD5.Update(reinterpret_cast<const uint8*>(Indices.GetData()), NumIndices * sizeof(uint32));
+
+		FMD5Hash MD5Hash;
+		MD5Hash.Set(MD5);
+
+		return MD5Hash;
+	}
 }  // End namespace UE::Chaos::ClothAsset
 
 bool FChaosClothSimulationLodModel::Serialize(FArchive& Ar)
@@ -126,11 +147,12 @@ bool FChaosClothSimulationLodModel::Serialize(FArchive& Ar)
 	return true;
 }
 
-FChaosClothSimulationModel::FChaosClothSimulationModel(const TArray<TSharedRef<const FManagedArrayCollection>>& ClothCollections, const FReferenceSkeleton& ReferenceSkeleton)
+FChaosClothSimulationModel::FChaosClothSimulationModel(const TArray<TSharedRef<const FManagedArrayCollection>>& ClothCollections, 
+	const FReferenceSkeleton& ReferenceSkeleton,
+	TArray<FChaosClothAssetLodTransitionDataCache>* InOutTransitionCache)
 {
 	using namespace UE::Chaos::ClothAsset;
 
-
 	// Initialize LOD models
 	const int32 NumLods = ClothCollections.Num();
 	ClothSimulationLodModels.SetNum(NumLods);
@@ -198,55 +220,97 @@ FChaosClothSimulationModel::FChaosClothSimulationModel(const TArray<TSharedRef<c
 	// Initialize Reference bone index
 	ReferenceBoneIndex = CalculateReferenceBoneIndex(ClothSimulationLodModels, ReferenceSkeleton, UsedBoneIndices);
 
-	// Compute LOD Transition data
-	
+	CalculateLODTransitionUpDownData(InOutTransitionCache);
+}
+
+TArray<TConstArrayView<TTuple<int32, int32, float>>> FChaosClothSimulationModel::GetTethers(int32 LODIndex) const
+{
+	TArray<TConstArrayView<TTuple<int32, int32, float>>> Tethers;
+	if (IsValidLodIndex(LODIndex))
+	{
+		const FClothTetherData& ClothTetherData = ClothSimulationLodModels[LODIndex].TetherData;
+
+		const int32 NumTetherBatches = ClothTetherData.Tethers.Num();
+		Tethers.Reserve(NumTetherBatches);
+		for (int32 Index = 0; Index < NumTetherBatches; ++Index)
+		{
+			Tethers.Emplace(TConstArrayView<TTuple<int32, int32, float>>(ClothTetherData.Tethers[Index]));
+		}
+	}
+	return Tethers;
+}
+
+void FChaosClothSimulationModel::CalculateLODTransitionUpDownData(TArray<FChaosClothAssetLodTransitionDataCache>* InOutTransitionCache)
+{
 	// Parameters for GenerateMeshToMeshVertData
 	const FPointWeightMap* const MaxDistances = nullptr;  // No need to update the vertex contribution on the transition maps
 	constexpr bool bUseSmoothTransitions = false;         // Smooth transitions are only used at rendering for now and not during LOD transitions
 	constexpr bool bUseMultipleInfluences = false;        // Multiple influences must not be used for LOD transitions
 	constexpr float SkinningKernelRadius = 0.f;           // KernelRadius is only required when using multiple influences
 
-	for (int32 LODIndex = 0; LODIndex < NumLods; ++LODIndex)
+	// Update cache hashes.
+	TArray<bool> TransitionCacheHashValid;
+	TransitionCacheHashValid.SetNumZeroed(ClothSimulationLodModels.Num());
+	if (InOutTransitionCache)
+	{
+		InOutTransitionCache->SetNum(ClothSimulationLodModels.Num());
+
+		for (int32 LODIndex = 0; LODIndex < ClothSimulationLodModels.Num(); ++LODIndex)
+		{
+			const FMD5Hash CurrCacheHash = UE::Chaos::ClothAsset::CalculateTransitionDataHash(ClothSimulationLodModels[LODIndex]);
+			TransitionCacheHashValid[LODIndex] = (CurrCacheHash == (*InOutTransitionCache)[LODIndex].ModelHash);
+			(*InOutTransitionCache)[LODIndex].ModelHash = CurrCacheHash;
+		}
+	}
+
+	for (int32 LODIndex = 0; LODIndex < ClothSimulationLodModels.Num(); ++LODIndex)
 	{
 		const int32 PrevLODIndex = LODIndex - 1;
 		const int32 NextLODIndex = LODIndex + 1;
-
+				
 		FChaosClothSimulationLodModel& CurrModel = ClothSimulationLodModels[LODIndex];
 		const ClothingMeshUtils::ClothMeshDesc CurrentMeshDesc(CurrModel.Positions, CurrModel.Indices);
 
 		if (ClothSimulationLodModels.IsValidIndex(PrevLODIndex))
 		{
-			const FChaosClothSimulationLodModel& PrevModel = ClothSimulationLodModels[PrevLODIndex];
-			const ClothingMeshUtils::ClothMeshDesc PrevMeshDesc(PrevModel.Positions, PrevModel.Indices);
+			if (TransitionCacheHashValid[LODIndex] && TransitionCacheHashValid[PrevLODIndex])
+			{
+				CurrModel.LODTransitionUpData = (*InOutTransitionCache)[LODIndex].LODTransitionUpData;
+			}
+			else
+			{
+				const FChaosClothSimulationLodModel& PrevModel = ClothSimulationLodModels[PrevLODIndex];
+				const ClothingMeshUtils::ClothMeshDesc PrevMeshDesc(PrevModel.Positions, PrevModel.Indices);
 
-			ClothingMeshUtils::GenerateMeshToMeshVertData(CurrModel.LODTransitionUpData, CurrentMeshDesc, PrevMeshDesc,
-				MaxDistances, bUseSmoothTransitions, bUseMultipleInfluences, SkinningKernelRadius);
+				ClothingMeshUtils::GenerateMeshToMeshVertData(CurrModel.LODTransitionUpData, CurrentMeshDesc, PrevMeshDesc,
+					MaxDistances, bUseSmoothTransitions, bUseMultipleInfluences, SkinningKernelRadius);
+
+				if (InOutTransitionCache)
+				{
+					(*InOutTransitionCache)[LODIndex].LODTransitionUpData = CurrModel.LODTransitionUpData;
+				}
+			}
 		}
 
 		if (ClothSimulationLodModels.IsValidIndex(NextLODIndex))
 		{
-			const FChaosClothSimulationLodModel& NextModel = ClothSimulationLodModels[NextLODIndex];
-			const ClothingMeshUtils::ClothMeshDesc NextMeshDesc(NextModel.Positions, NextModel.Indices);
-
-			ClothingMeshUtils::GenerateMeshToMeshVertData(CurrModel.LODTransitionDownData, CurrentMeshDesc, NextMeshDesc,
-				MaxDistances, bUseSmoothTransitions, bUseMultipleInfluences, SkinningKernelRadius);
-		}
-	}
-}
+			if (TransitionCacheHashValid[LODIndex] && TransitionCacheHashValid[NextLODIndex])
+			{
+				CurrModel.LODTransitionDownData = (*InOutTransitionCache)[LODIndex].LODTransitionDownData;
+			}
+			else
+			{
+				const FChaosClothSimulationLodModel& NextModel = ClothSimulationLodModels[NextLODIndex];
+				const ClothingMeshUtils::ClothMeshDesc NextMeshDesc(NextModel.Positions, NextModel.Indices);
 
-TArray<TConstArrayView<TTuple<int32, int32, float>>> FChaosClothSimulationModel::GetTethers(int32 LODIndex) const
-{
-	TArray<TConstArrayView<TTuple<int32, int32, float>>> Tethers;
-	if (IsValidLodIndex(LODIndex))
-	{
-		const FClothTetherData& ClothTetherData = ClothSimulationLodModels[LODIndex].TetherData;
+				ClothingMeshUtils::GenerateMeshToMeshVertData(CurrModel.LODTransitionDownData, CurrentMeshDesc, NextMeshDesc,
+					MaxDistances, bUseSmoothTransitions, bUseMultipleInfluences, SkinningKernelRadius);
 
-		const int32 NumTetherBatches = ClothTetherData.Tethers.Num();
-		Tethers.Reserve(NumTetherBatches);
-		for (int32 Index = 0; Index < NumTetherBatches; ++Index)
-		{
-			Tethers.Emplace(TConstArrayView<TTuple<int32, int32, float>>(ClothTetherData.Tethers[Index]));
+				if (InOutTransitionCache)
+				{
+					(*InOutTransitionCache)[LODIndex].LODTransitionDownData = CurrModel.LODTransitionDownData;
+				}
+			}
 		}
 	}
-	return Tethers;
 }
\ No newline at end of file
diff --git a/Engine/Plugins/Experimental/ChaosClothAsset/Source/ChaosClothAssetEngine/Private/ChaosClothAsset/ClothSimulationModel.h b/Engine/Plugins/Experimental/ChaosClothAsset/Source/ChaosClothAssetEngine/Private/ChaosClothAsset/ClothSimulationModel.h
--- a/Engine/Plugins/Experimental/ChaosClothAsset/Source/ChaosClothAssetEngine/Private/ChaosClothAsset/ClothSimulationModel.h
+++ b/Engine/Plugins/Experimental/ChaosClothAsset/Source/ChaosClothAssetEngine/Private/ChaosClothAsset/ClothSimulationModel.h
@@ -12,6 +12,7 @@
 
 struct FReferenceSkeleton;
 struct FManagedArrayCollection;
+struct FChaosClothAssetLodTransitionDataCache;
 
 /**
  * Cloth simulation LOD model.
@@ -100,7 +101,7 @@ struct FChaosClothSimulationModel
 	int32 ReferenceBoneIndex = INDEX_NONE;
 
 	FChaosClothSimulationModel() = default;
-	FChaosClothSimulationModel(const TArray<TSharedRef<const FManagedArrayCollection>>& ClothCollections, const FReferenceSkeleton& ReferenceSkeleton);
+	FChaosClothSimulationModel(const TArray<TSharedRef<const FManagedArrayCollection>>& ClothCollections, const FReferenceSkeleton& ReferenceSkeleton, TArray<FChaosClothAssetLodTransitionDataCache>* InOutTransitionCache = nullptr);
 
 	int32 GetNumLods() const { return ClothSimulationLodModels.Num(); }
 
@@ -116,4 +117,7 @@ struct FChaosClothSimulationModel
 	TConstArrayView<uint32> GetPatternToWeldedIndices(int32 LodIndex) const { return IsValidLodIndex(LodIndex) ? ClothSimulationLodModels[LodIndex].PatternToWeldedIndices : TConstArrayView<uint32>(); }
 	TConstArrayView<FClothVertBoneData> GetBoneData(int32 LodIndex) const { return IsValidLodIndex(LodIndex) ? ClothSimulationLodModels[LodIndex].BoneData : TConstArrayView<FClothVertBoneData>(); }
 	TArray<TConstArrayView<TTuple<int32, int32, float>>> GetTethers(int32 LodIndex) const;
+
+
+	void CalculateLODTransitionUpDownData(TArray<FChaosClothAssetLodTransitionDataCache>* InOutTransitionCache = nullptr);
 };
diff --git a/Engine/Plugins/Experimental/ChaosClothAsset/Source/ChaosClothAssetEngine/Public/ChaosClothAsset/ClothAsset.h b/Engine/Plugins/Experimental/ChaosClothAsset/Source/ChaosClothAssetEngine/Public/ChaosClothAsset/ClothAsset.h
--- a/Engine/Plugins/Experimental/ChaosClothAsset/Source/ChaosClothAssetEngine/Public/ChaosClothAsset/ClothAsset.h
+++ b/Engine/Plugins/Experimental/ChaosClothAsset/Source/ChaosClothAssetEngine/Public/ChaosClothAsset/ClothAsset.h
@@ -12,6 +12,7 @@ class FSkeletalMeshModel;
 struct FChaosClothSimulationModel;
 struct FSkeletalMeshLODInfo;
 struct FManagedArrayCollection;
+struct FChaosClothAssetLodTransitionDataCache;
 class FSkinnedAssetCompilationContext;
 class UDataflow;
 
@@ -109,7 +110,7 @@ public:
 	TSharedPtr<const FChaosClothSimulationModel> GetClothSimulationModel() const { return ClothSimulationModel; }
 
 	/** Build this asset static render and simulation data. This needs to be done every time the asset has changed. */
-	void Build();
+	void Build(TArray<FChaosClothAssetLodTransitionDataCache>* InOutTransitionCache = nullptr);
 
 	/**
 	 * Copy the draped simulation mesh patterns into the render mesh data.
@@ -188,7 +189,7 @@ private:
 #endif
 
 	/** Build the clothing simulation meshes from the Cloth Collection. */
-	void BuildClothSimulationModel();
+	void BuildClothSimulationModel(TArray<FChaosClothAssetLodTransitionDataCache>* InOutTransitionCache = nullptr);
 
 	/** Initialize all render resources. */
 	void InitResources();
diff --git a/Engine/Plugins/Experimental/ChaosClothAsset/Source/ChaosClothAssetEngine/Public/ChaosClothAsset/ClothLodTransitionDataCache.h b/Engine/Plugins/Experimental/ChaosClothAsset/Source/ChaosClothAssetEngine/Public/ChaosClothAsset/ClothLodTransitionDataCache.h
new file mode 100644
--- /dev/null
+++ b/Engine/Plugins/Experimental/ChaosClothAsset/Source/ChaosClothAssetEngine/Public/ChaosClothAsset/ClothLodTransitionDataCache.h
@@ -0,0 +1,31 @@
+// Copyright Epic Games, Inc. All Rights Reserved.
+
+#pragma once
+#include "Misc/SecureHash.h"
+#include "SkeletalMeshTypes.h"
+#include "ClothLodTransitionDataCache.generated.h"
+
+USTRUCT()
+struct FChaosClothAssetLodTransitionDataCache
+{
+	GENERATED_BODY()
+
+	/** Hash of this model. Require valid hash of Up and Down models to be able to re-use the transition data.**/
+	FMD5Hash ModelHash;
+
+	/** LOD Transition mesh to mesh skinning weights. */
+	TArray<FMeshToMeshVertData> LODTransitionUpData;
+	TArray<FMeshToMeshVertData> LODTransitionDownData;
+
+	// Custom serialize for MD5Hash
+	bool Serialize(FArchive& Ar);
+};
+
+template<>
+struct TStructOpsTypeTraits<FChaosClothAssetLodTransitionDataCache> : public TStructOpsTypeTraitsBase2<FChaosClothAssetLodTransitionDataCache>
+{
+	enum
+	{
+		WithSerializer = true,
+	};
+};
diff --git a/Engine/Plugins/Experimental/ChaosClothAssetEditor/Source/ChaosClothAssetDataflowNodes/Private/ChaosClothAsset/TerminalNode.cpp b/Engine/Plugins/Experimental/ChaosClothAssetEditor/Source/ChaosClothAssetDataflowNodes/Private/ChaosClothAsset/TerminalNode.cpp
--- a/Engine/Plugins/Experimental/ChaosClothAssetEditor/Source/ChaosClothAssetDataflowNodes/Private/ChaosClothAsset/TerminalNode.cpp
+++ b/Engine/Plugins/Experimental/ChaosClothAssetEditor/Source/ChaosClothAssetDataflowNodes/Private/ChaosClothAsset/TerminalNode.cpp
@@ -4,6 +4,7 @@
 #include "ChaosClothAsset/ClothAsset.h"
 #include "ChaosClothAsset/ClothDataflowTools.h"
 #include "ChaosClothAsset/ClothGeometryTools.h"
+#include "ChaosClothAsset/ClothLODTransitionDataCache.h"
 #include "ChaosClothAsset/CollectionClothFacade.h"
 #include "Animation/Skeleton.h"
 #include "Chaos/CollectionPropertyFacade.h"
@@ -149,7 +150,7 @@ void FChaosClothAssetTerminalNode::SetAssetValue(TObjectPtr<UObject> Asset, Data
 		ClothAsset->SetPhysicsAsset(PhysicsAsset);
 
 		// Rebuild the asset static data
-		ClothAsset->Build();
+		ClothAsset->Build(&LODTransitionDataCache);
 	}
 }
 
diff --git a/Engine/Plugins/Experimental/ChaosClothAssetEditor/Source/ChaosClothAssetDataflowNodes/Public/ChaosClothAsset/TerminalNode.h b/Engine/Plugins/Experimental/ChaosClothAssetEditor/Source/ChaosClothAssetDataflowNodes/Public/ChaosClothAsset/TerminalNode.h
--- a/Engine/Plugins/Experimental/ChaosClothAssetEditor/Source/ChaosClothAssetDataflowNodes/Public/ChaosClothAsset/TerminalNode.h
+++ b/Engine/Plugins/Experimental/ChaosClothAssetEditor/Source/ChaosClothAssetDataflowNodes/Public/ChaosClothAsset/TerminalNode.h
@@ -4,6 +4,7 @@
 
 #include "Dataflow/DataflowTerminalNode.h"
 #include "GeometryCollection/ManagedArrayCollection.h"
+#include "ChaosClothAsset/ClothLodTransitionDataCache.h"
 #include "TerminalNode.generated.h"
 
 /** Cloth terminal node to generate a cloth asset from a cloth collection. */
@@ -52,4 +53,7 @@ private:
 	//~ End FDataflowNode interface
 
 	TArray<const FManagedArrayCollection*> GetCollectionLods() const;
+
+	UPROPERTY()
+	mutable TArray<FChaosClothAssetLodTransitionDataCache> LODTransitionDataCache;
 };
