Chaos Cloth - Rotated the Cloth Asset DDC key to allow for the new number of materials/sections to go through, and prevent a crash in BuildFromLODModel when reloading some assets post the multi material fix changes at CL 26883263.

#rnx

diff --git a/Engine/Plugins/Experimental/ChaosClothAsset/Source/ChaosClothAssetEngine/Private/ChaosClothAsset/ClothAsset.cpp b/Engine/Plugins/Experimental/ChaosClothAsset/Source/ChaosClothAssetEngine/Private/ChaosClothAsset/ClothAsset.cpp
--- a/Engine/Plugins/Experimental/ChaosClothAsset/Source/ChaosClothAssetEngine/Private/ChaosClothAsset/ClothAsset.cpp
+++ b/Engine/Plugins/Experimental/ChaosClothAsset/Source/ChaosClothAssetEngine/Private/ChaosClothAsset/ClothAsset.cpp
@@ -26,7 +26,7 @@
 
 // If Chaos cloth asset derived data needs to be rebuilt (new format, serialization differences, etc.) replace the version GUID below with a new one. 
 // In case of merge conflicts with DDC versions, you MUST generate a new GUID and set this new GUID as the version.
-#define CHAOS_CLOTH_ASSET_DERIVED_DATA_VERSION TEXT("19823D996CA54F279B9A9FA8ED7A8EB6")
+#define CHAOS_CLOTH_ASSET_DERIVED_DATA_VERSION TEXT("5C041E93563044A69646B5E98F152B7C")
 
 namespace UE::Chaos::ClothAsset::Private
 {
