Chaos Cloth Asset:
- Replaced the USkeleton asset path name by a USkeletalMesh asset path name to fix a skinning issues when the USkeletalMesh reference skeleton doesn't match the given USkeleton reference skeleton.
- Deprecated public access to the Dataflow pointer and Terminal node properties.

#rnx

diff --git a/Engine/Plugins/Enterprise/DatasmithImporter/Source/DatasmithImporter/Private/DatasmithImporter.cpp b/Engine/Plugins/Enterprise/DatasmithImporter/Source/DatasmithImporter/Private/DatasmithImporter.cpp
--- a/Engine/Plugins/Enterprise/DatasmithImporter/Source/DatasmithImporter/Private/DatasmithImporter.cpp
+++ b/Engine/Plugins/Enterprise/DatasmithImporter/Source/DatasmithImporter/Private/DatasmithImporter.cpp
@@ -539,15 +539,11 @@ void FDatasmithImporter::ImportClothes(FDatasmithImportContext& ImportContext)
 			}
 		}
 
-		// Set a default skeleton
-		constexpr bool bRebuildModels = false;
-		ClothAsset->SetSkeleton(nullptr, bRebuildModels);
-
 		// Set the render mesh to duplicate the sim mesh
 		ClothAsset->CopySimMeshToRenderMesh();
 
-		// Rebuild static data
-		ClothAsset->Build();
+		// Set a default skeleton and rebuild the asset
+		ClothAsset->SetReferenceSkeleton(nullptr);  // This creates a default reference skeleton, redo the bindings, and rebuild the asset
 
 		ImportContext.ImportedClothes.Add(ClothElement, ClothAsset);
 
diff --git a/Engine/Plugins/Experimental/ChaosClothAsset/Source/ChaosClothAsset/Private/ChaosClothAsset/ClothCollection.cpp b/Engine/Plugins/Experimental/ChaosClothAsset/Source/ChaosClothAsset/Private/ChaosClothAsset/ClothCollection.cpp
--- a/Engine/Plugins/Experimental/ChaosClothAsset/Source/ChaosClothAsset/Private/ChaosClothAsset/ClothCollection.cpp
+++ b/Engine/Plugins/Experimental/ChaosClothAsset/Source/ChaosClothAsset/Private/ChaosClothAsset/ClothCollection.cpp
@@ -19,11 +19,11 @@ namespace UE::Chaos::ClothAsset::Private
 
 	// Lods Group
 	static const FName PhysicsAssetPathNameAttribute(TEXT("PhysicsAssetPathName"));
-	static const FName SkeletonAssetPathNameAttribute(TEXT("SkeletonAssetPathName"));
+	static const FName SkeletalMeshPathNameAttribute(TEXT("SkeletalMeshPathName"));
 	static const TArray<FName> LodsGroupAttributes =
 	{
 		PhysicsAssetPathNameAttribute,
-		SkeletonAssetPathNameAttribute
+		SkeletalMeshPathNameAttribute
 	};
 
 	// Seam Group
@@ -176,7 +176,7 @@ namespace UE::Chaos::ClothAsset
 
 		// LODs Group
 		PhysicsAssetPathName = ManagedArrayCollection->FindAttribute<FString>(PhysicsAssetPathNameAttribute, LodsGroup);
-		SkeletonAssetPathName = ManagedArrayCollection->FindAttribute<FString>(SkeletonAssetPathNameAttribute, LodsGroup);
+		SkeletalMeshPathName = ManagedArrayCollection->FindAttribute<FString>(SkeletalMeshPathNameAttribute, LodsGroup);
 
 		// Seam Group
 		SeamStitchStart = ManagedArrayCollection->FindAttribute<int32>(SeamStitchStartAttribute, SeamsGroup);
@@ -236,7 +236,7 @@ namespace UE::Chaos::ClothAsset
 		return 
 			// LODs Group
 			PhysicsAssetPathName &&
-			SkeletonAssetPathName &&
+			SkeletalMeshPathName &&
 
 			// Seam Group
 			SeamStitchStart &&
@@ -307,7 +307,7 @@ namespace UE::Chaos::ClothAsset
 
 		// LODs Group
 		PhysicsAssetPathName = &ManagedArrayCollection->AddAttribute<FString>(PhysicsAssetPathNameAttribute, LodsGroup);
-		SkeletonAssetPathName = &ManagedArrayCollection->AddAttribute<FString>(SkeletonAssetPathNameAttribute, LodsGroup);
+		SkeletalMeshPathName = &ManagedArrayCollection->AddAttribute<FString>(SkeletalMeshPathNameAttribute, LodsGroup);
 
 		// Seams Group
 		SeamStitchStart = &ManagedArrayCollection->AddAttribute<int32>(SeamStitchStartAttribute, SeamsGroup, SeamStitchesDependency);
diff --git a/Engine/Plugins/Experimental/ChaosClothAsset/Source/ChaosClothAsset/Private/ChaosClothAsset/ClothCollection.h b/Engine/Plugins/Experimental/ChaosClothAsset/Source/ChaosClothAsset/Private/ChaosClothAsset/ClothCollection.h
--- a/Engine/Plugins/Experimental/ChaosClothAsset/Source/ChaosClothAsset/Private/ChaosClothAsset/ClothCollection.h
+++ b/Engine/Plugins/Experimental/ChaosClothAsset/Source/ChaosClothAsset/Private/ChaosClothAsset/ClothCollection.h
@@ -159,7 +159,7 @@ namespace UE::Chaos::ClothAsset
 
 		//~ LODs Group (There should be only one LOD per ClothCollection)
 		const TManagedArray<FString>* GetPhysicsAssetPathName() const { return PhysicsAssetPathName; }
-		const TManagedArray<FString>* GetSkeletonAssetPathName() const { return SkeletonAssetPathName; }
+		const TManagedArray<FString>* GetSkeletalMeshPathName() const { return SkeletalMeshPathName; }
 
 		//~ Seam Group
 		const TManagedArray<int32>* GetSeamStitchStart() const { return SeamStitchStart; }
@@ -215,7 +215,7 @@ namespace UE::Chaos::ClothAsset
 
 		//~ LODs Group (There should be only one LOD per ClothCollection)
 		TManagedArray<FString>* GetPhysicsAssetPathName(){ return PhysicsAssetPathName; }
-		TManagedArray<FString>* GetSkeletonAssetPathName(){ return SkeletonAssetPathName; }
+		TManagedArray<FString>* GetSkeletalMeshPathName() { return SkeletalMeshPathName; }
 
 		//~ Seam Group
 		TManagedArray<int32>* GetSeamStitchStart() { return SeamStitchStart; }
@@ -275,7 +275,7 @@ namespace UE::Chaos::ClothAsset
 
 		//~ LODs Group
 		TManagedArray<FString>* PhysicsAssetPathName;
-		TManagedArray<FString>* SkeletonAssetPathName;
+		TManagedArray<FString>* SkeletalMeshPathName;
 
 		//~ Seam Group
 		TManagedArray<int32>* SeamStitchStart;
diff --git a/Engine/Plugins/Experimental/ChaosClothAsset/Source/ChaosClothAsset/Private/ChaosClothAsset/CollectionClothFacade.cpp b/Engine/Plugins/Experimental/ChaosClothAsset/Source/ChaosClothAsset/Private/ChaosClothAsset/CollectionClothFacade.cpp
--- a/Engine/Plugins/Experimental/ChaosClothAsset/Source/ChaosClothAsset/Private/ChaosClothAsset/CollectionClothFacade.cpp
+++ b/Engine/Plugins/Experimental/ChaosClothAsset/Source/ChaosClothAsset/Private/ChaosClothAsset/CollectionClothFacade.cpp
@@ -29,10 +29,10 @@ namespace UE::Chaos::ClothAsset
 		return ClothCollection->GetPhysicsAssetPathName() && ClothCollection->GetNumElements(FClothCollection::LodsGroup) > 0 ? (*ClothCollection->GetPhysicsAssetPathName())[0] : EmptyString;
 	}
 
-	const FString& FCollectionClothConstFacade::GetSkeletonAssetPathName() const
+	const FString& FCollectionClothConstFacade::GetSkeletalMeshPathName() const
 	{
 		static const FString EmptyString;
-		return ClothCollection->GetSkeletonAssetPathName() && ClothCollection->GetNumElements(FClothCollection::LodsGroup) > 0 ? (*ClothCollection->GetSkeletonAssetPathName())[0] : EmptyString;
+		return ClothCollection->GetSkeletalMeshPathName() && ClothCollection->GetNumElements(FClothCollection::LodsGroup) > 0 ? (*ClothCollection->GetSkeletalMeshPathName())[0] : EmptyString;
 	}
 
 	int32 FCollectionClothConstFacade::GetNumSimVertices2D() const
@@ -308,7 +308,7 @@ namespace UE::Chaos::ClothAsset
 
 		// LODs Group
 		SetPhysicsAssetPathName(Other.GetPhysicsAssetPathName());
-		SetSkeletonAssetPathName(Other.GetSkeletonAssetPathName());
+		SetSkeletalMeshPathName(Other.GetSkeletalMeshPathName());
 		
 		Append(Other);
 	}
@@ -388,11 +388,11 @@ namespace UE::Chaos::ClothAsset
 		}
 
 	}
-	void FCollectionClothFacade::SetSkeletonAssetPathName(const FString& PathName)
+	void FCollectionClothFacade::SetSkeletalMeshPathName(const FString& PathName)
 	{
 		if (ClothCollection->GetNumElements(FClothCollection::LodsGroup))
 		{
-			(*GetClothCollection()->GetSkeletonAssetPathName())[0] = PathName;
+			(*GetClothCollection()->GetSkeletalMeshPathName())[0] = PathName;
 		}
 	}
 
diff --git a/Engine/Plugins/Experimental/ChaosClothAsset/Source/ChaosClothAsset/Public/ChaosClothAsset/CollectionClothFacade.h b/Engine/Plugins/Experimental/ChaosClothAsset/Source/ChaosClothAsset/Public/ChaosClothAsset/CollectionClothFacade.h
--- a/Engine/Plugins/Experimental/ChaosClothAsset/Source/ChaosClothAsset/Public/ChaosClothAsset/CollectionClothFacade.h
+++ b/Engine/Plugins/Experimental/ChaosClothAsset/Source/ChaosClothAsset/Public/ChaosClothAsset/CollectionClothFacade.h
@@ -39,7 +39,7 @@ namespace UE::Chaos::ClothAsset
 		/** Return the physics asset path names used for this collection. */
 		const FString& GetPhysicsAssetPathName() const;
 		/** Return the skeleton asset path names used for this collection. */
-		const FString& GetSkeletonAssetPathName() const;
+		const FString& GetSkeletalMeshPathName() const;
 
 		//~ Sim Vertices 2D Group
 		/** Return the total number of 2D simulation vertices for this collection. */
@@ -165,8 +165,8 @@ namespace UE::Chaos::ClothAsset
 		//~ LOD (single per collection) Group
 		/** Set the physics asset path name. */
 		void SetPhysicsAssetPathName(const FString& PathName);
-		/** Set the skeleton asset path name. */
-		void SetSkeletonAssetPathName(const FString& PathName);
+		/** Set the skeletal mesh asset path name and the reference skeleton that will be used with this asset. */
+		void SetSkeletalMeshPathName(const FString& PathName);
 
 		//~ Pattern Sim Vertices 2D Group
 		/** SetNumSimVertices2D per pattern within pattern facade. */
diff --git a/Engine/Plugins/Experimental/ChaosClothAsset/Source/ChaosClothAssetEngine/Private/ChaosClothAsset/ClothAsset.cpp b/Engine/Plugins/Experimental/ChaosClothAsset/Source/ChaosClothAssetEngine/Private/ChaosClothAsset/ClothAsset.cpp
--- a/Engine/Plugins/Experimental/ChaosClothAsset/Source/ChaosClothAssetEngine/Private/ChaosClothAsset/ClothAsset.cpp
+++ b/Engine/Plugins/Experimental/ChaosClothAsset/Source/ChaosClothAssetEngine/Private/ChaosClothAsset/ClothAsset.cpp
@@ -110,7 +110,8 @@ UChaosClothAsset::UChaosClothAsset(const FObjectInitializer& ObjectInitializer)
 
 	// Set default skeleton (must be done after having added the LOD)
 	constexpr bool bRebuildModels = false;
-	SetSkeleton(nullptr, bRebuildModels);
+	constexpr bool bRebindMeshes = false;
+	SetReferenceSkeleton(nullptr, bRebuildModels, bRebindMeshes);
 }
 
 UChaosClothAsset::UChaosClothAsset(FVTableHelper& Helper)
@@ -118,7 +119,9 @@ UChaosClothAsset::UChaosClothAsset(FVTableHelper& Helper)
 {
 }
 
+PRAGMA_DISABLE_DEPRECATION_WARNINGS
 UChaosClothAsset::~UChaosClothAsset() = default;
+PRAGMA_ENABLE_DEPRECATION_WARNINGS
 
 FSkeletalMeshLODInfo* UChaosClothAsset::GetLODInfo(int32 Index)
 {
@@ -158,11 +161,6 @@ void UChaosClothAsset::Serialize(FArchive& Ar)
 	UE::Chaos::ClothAsset::Private::Serialize(ChaosArchive, ClothCollections);
 
 	Ar << GetRefSkeleton();
-	if (Ar.IsLoading())
-	{
-		constexpr bool bRebuildNameMap = false;
-		GetRefSkeleton().RebuildRefSkeleton(GetSkeleton(), bRebuildNameMap);
-	}
 
 	if (bCooked && !IsTemplate() && !Ar.IsCountingMemory())
 	{
@@ -214,7 +212,8 @@ void UChaosClothAsset::BeginPostLoadInternal(FSkinnedAssetPostLoadContext& Conte
 		ClothFacade.DefineSchema();
 		bAnyInvalidLods = true;
 	}
-	for(int32 LODIndex = 0; LODIndex < ClothCollections.Num(); ++LODIndex)
+	bool bAnyInvalidSkeletons = false;
+	for (int32 LODIndex = 0; LODIndex < ClothCollections.Num(); ++LODIndex)
 	{
 		TSharedRef<FManagedArrayCollection>& ClothCollection = ClothCollections[LODIndex];
 
@@ -226,14 +225,23 @@ void UChaosClothAsset::BeginPostLoadInternal(FSkinnedAssetPostLoadContext& Conte
 			ClothFacade = UE::Chaos::ClothAsset::FCollectionClothFacade(ClothCollection);
 			ClothFacade.DefineSchema();
 			bAnyInvalidLods = true;
+			bAnyInvalidSkeletons = true;
+		}
+		else if (ClothFacade.GetSkeletalMeshPathName().IsEmpty())
+		{
+			bAnyInvalidSkeletons = true;
 		}
 	}
 	if (bAnyInvalidLods)
 	{
-		constexpr bool bRebuildModels = false;
-		SetSkeleton(Skeleton, bRebuildModels);  // Re-update the collection with the skeleton information if any
 		SetPhysicsAsset(PhysicsAsset);  // Re-update the collection with the physics asset information if any
 	}
+	if (bAnyInvalidSkeletons)
+	{
+		constexpr bool bRebuildModels = false;
+		constexpr bool bRebindMeshes = true;  // Best to rebind the mesh when reloading broken data
+		SetReferenceSkeleton(nullptr, bRebuildModels, bRebindMeshes);
+	}
 
 	BuildClothSimulationModel();  // TODO: Cache ClothSimulationModel?
 
@@ -689,27 +697,43 @@ void UChaosClothAsset::SetPhysicsAsset(UPhysicsAsset* InPhysicsAsset)
 	}
 }
 
-void UChaosClothAsset::SetSkeleton(USkeleton* InSkeleton, bool bRebuildModels)
+void UChaosClothAsset::SetReferenceSkeleton(const FReferenceSkeleton* ReferenceSkeleton, bool bRebuildModels, bool bRebindMeshes)
 {
 	using namespace UE::Chaos::ClothAsset;
-	check(ClothCollections.Num());
 
-	const TCHAR* const DefaultSkeletonPathName = TEXT("/Engine/EditorMeshes/SkeletalMesh/DefaultSkeletalMesh_Skeleton.DefaultSkeletalMesh_Skeleton");
-	Skeleton = InSkeleton ?
-		InSkeleton :
-		LoadObject<USkeleton>(nullptr, DefaultSkeletonPathName, nullptr, LOAD_None, nullptr);
+	// Update the reference skeleton
+	if (ReferenceSkeleton)
+	{
+		RefSkeleton = *ReferenceSkeleton;
+	}
+	else
+	{
+		// Create a default reference skeleton
+		RefSkeleton.Empty(1);
+		FReferenceSkeletonModifier ReferenceSkeletonModifier(RefSkeleton, nullptr);
+
+		FMeshBoneInfo MeshBoneInfo;
+		constexpr const TCHAR* RootName = TEXT("Root");
+		MeshBoneInfo.ParentIndex = INDEX_NONE;
+#if WITH_EDITORONLY_DATA
+		MeshBoneInfo.ExportName = RootName;
+#endif
+		MeshBoneInfo.Name = FName(RootName);
+		ReferenceSkeletonModifier.Add(MeshBoneInfo, FTransform::Identity);
 
-	RefSkeleton = Skeleton->GetReferenceSkeleton();
+		bRebindMeshes = true; // Force the binding when a default reference skeleton is being created 
+	}
 
-	for (TSharedRef<FManagedArrayCollection>& ClothCollection : ClothCollections)
+	// Rebind the meshes
+	if (bRebindMeshes)
 	{
-		constexpr bool bBindSimMesh = true;
-		constexpr bool bBindRenderMesh = true;
-		FClothGeometryTools::BindMeshToRootBone(ClothCollection, bBindSimMesh, bBindRenderMesh);
-
-		FCollectionClothFacade Cloth(ClothCollection);
-		Cloth.SetSkeletonAssetPathName(Skeleton->GetPathName());
+		for (TSharedRef<FManagedArrayCollection>& ClothCollection : ClothCollections)
+		{
+			FClothGeometryTools::BindMeshToRootBone(ClothCollection, true, true);
+		}
 	}
+
+	// Rebind the models
 	if (bRebuildModels)
 	{
 		Build();
@@ -724,24 +748,14 @@ void UChaosClothAsset::UpdateSkeletonFromCollection(bool bRebuildModels)
 	FCollectionClothConstFacade ClothFacade(ClothCollections[0]);
 	check(ClothFacade.IsValid());
 
-	USkeleton* InSkeleton = nullptr;
-	const FString& SkeletonAssetPathName = ClothFacade.GetSkeletonAssetPathName();
-	InSkeleton = LoadObject<USkeleton>(nullptr, *SkeletonAssetPathName, nullptr, LOAD_None, nullptr);
+	const FString& SkeletalMeshPathName = ClothFacade.GetSkeletalMeshPathName();
+	USkeletalMesh* const SkeletalMesh = SkeletalMeshPathName.IsEmpty() ? nullptr :
+		LoadObject<USkeletalMesh>(nullptr, *SkeletalMeshPathName, nullptr, LOAD_None, nullptr);
 
-	if (!InSkeleton)
-	{
-		SetSkeleton(nullptr, bRebuildModels);
-	}
-	else
-	{
-		Skeleton = InSkeleton;
-		RefSkeleton = Skeleton->GetReferenceSkeleton();
+	SetSkeleton(SkeletalMesh ? SkeletalMesh->GetSkeleton() : nullptr); // For completion only, this is not being used and might mismatch the skeletal mesh's reference skeleton
 
-		if (bRebuildModels)
-		{
-			Build();
-		}
-	}
+	constexpr bool bRebindMeshes = false;  // The collection should contain the correct binding at the time SkeletalMeshPathName was set
+	SetReferenceSkeleton(SkeletalMesh ? &SkeletalMesh->GetRefSkeleton() : nullptr, bRebuildModels, bRebindMeshes);
 }
 
 void UChaosClothAsset::CopySimMeshToRenderMesh(UMaterialInterface* Material)
diff --git a/Engine/Plugins/Experimental/ChaosClothAsset/Source/ChaosClothAssetEngine/Public/ChaosClothAsset/ClothAsset.h b/Engine/Plugins/Experimental/ChaosClothAsset/Source/ChaosClothAssetEngine/Public/ChaosClothAsset/ClothAsset.h
--- a/Engine/Plugins/Experimental/ChaosClothAsset/Source/ChaosClothAssetEngine/Public/ChaosClothAsset/ClothAsset.h
+++ b/Engine/Plugins/Experimental/ChaosClothAsset/Source/ChaosClothAssetEngine/Public/ChaosClothAsset/ClothAsset.h
@@ -9,12 +9,13 @@
 
 class FSkeletalMeshRenderData;
 class FSkeletalMeshModel;
+class FSkinnedAssetCompilationContext;
+class USkeletalMesh;
+class UDataflow;
 struct FChaosClothSimulationModel;
 struct FSkeletalMeshLODInfo;
 struct FManagedArrayCollection;
 struct FChaosClothAssetLodTransitionDataCache;
-class FSkinnedAssetCompilationContext;
-class UDataflow;
 
 UENUM()
 enum class EClothAssetAsyncProperties : uint64
@@ -75,9 +76,9 @@ public:
 	virtual TArray<class USkeletalMeshSocket*> GetActiveSocketList() const override { static TArray<class USkeletalMeshSocket*> Dummy; return Dummy; }
 	virtual USkeletalMeshSocket* FindSocket(FName InSocketName) const override	{ return nullptr; }
 	virtual USkeletalMeshSocket* FindSocketInfo(FName InSocketName, FTransform& OutTransform, int32& OutBoneIndex, int32& OutIndex) const override { return nullptr; }
-	virtual USkeleton* GetSkeleton() override									{ return Skeleton; }
+	virtual USkeleton* GetSkeleton() override									{ return Skeleton; }  // Note: The USkeleton isn't a reliable source of reference skeleton
 	virtual const USkeleton* GetSkeleton() const override						{ return Skeleton; }
-	virtual void SetSkeleton(USkeleton* InSkeleton) override					{ constexpr bool bRebuildModels = true; SetSkeleton(InSkeleton, bRebuildModels); }
+	virtual void SetSkeleton(USkeleton* InSkeleton) override					{ Skeleton = InSkeleton; }
 	virtual UMeshDeformer* GetDefaultMeshDeformer() const override				{ return nullptr; }
 	virtual class UMaterialInterface* GetOverlayMaterial() const override		{ return nullptr; }
 	virtual float GetOverlayMaterialMaxDrawDistance() const override			{ return 0.f; }
@@ -123,38 +124,40 @@ public:
 	void SetPhysicsAsset(UPhysicsAsset* InPhysicsAsset);
 
 	/**
-	 * Set the skeleton asset and bone hierachy to use for this cloth.
-	 * @param InSkeleton The skeleton asset to use for this cloth, or nullptr to setup a default skeleton asset.
+	 * Set the skeleton asset using the LOD0 reference skeleton of the skeletal mesh specified in the cloth collection, or set a default reference skeleton if there isn't any.
 	 * @param bRebuildModels Whether to rebuild the simulation and rendering models (editor builds only).
 	 */
-	void SetSkeleton(USkeleton* InSkeleton, bool bRebuildModels);
+	void UpdateSkeletonFromCollection(bool bRebuildModels);
 
 	/**
-	 * Set the skeleton asset using the LOD0 skeleton specified in the cloth collection, or load a default skeleton if it is invalid.
-	 * @param bRebuildModels Whether to rebuild the simulation and rendering models (editor builds only).
+	 * Set the specified reference skeleton for this cloth, and rebuild the models.
+	 * Use nullptr to set to a default single bone root reference skeleton.
+	 * @param bRebuildModels Whether to rebuild the mesh static models. This is always required, but could be skipped if done soon after.
+	 * @param bRebindMeshes Bindings could be invalid after a change of reference skeleton, use this option to clear the existing binding and bind the sim and render meshes to the root bone.
 	 */
-	void UpdateSkeletonFromCollection(bool bRebuildModels);
+	void SetReferenceSkeleton(const FReferenceSkeleton* ReferenceSkeleton, bool bRebuildModels = true, bool bRebindMeshes = true);
 
 	/** Set the bone hierachy to use for this cloth. */
-	UE_DEPRECATED(5.3, "Use SetSkeleton instead")
+	UE_DEPRECATED(5.3, "Use SetReferenceSkeleton(const FReferenceSkeleton*, bool, bool) instead")
 	PRAGMA_DISABLE_DEPRECATION_WARNINGS
 	void SetReferenceSkeleton(const FReferenceSkeleton& InReferenceSkeleton, bool bRebuildClothSimulationModel = true) { RefSkeleton = InReferenceSkeleton; UpdateSkeleton(bRebuildClothSimulationModel); }
 	PRAGMA_ENABLE_DEPRECATION_WARNINGS
 
 	/** Set the skinning weights for all of the sim vertices to be bound to the root node of the reference skeleton. */
-	UE_DEPRECATED(5.3, "Use SetSkeleton instead")
+	UE_DEPRECATED(5.3, "Use FClothGeometryTools::BindMeshToRootBone or SetReferenceSkeleton(const FReferenceSkeleton*, bool, bool) instead.")
 	void BindSimMeshToRootBone();
 
 	//
 	// Dataflow
 	//
+	UE_DEPRECATED(5.3, "Do not use. Will be made private in 5.4")
 	UPROPERTY(EditAnywhere, Category = "Dataflow")
 	TObjectPtr<UDataflow> DataflowAsset;
 
+	UE_DEPRECATED(5.3, "Do not use. Will be made private in 5.4")
 	UPROPERTY(EditAnywhere, Category = "Dataflow")
 	FString DataflowTerminal = "ClothAssetTerminal";
 
-
 private:
 	//~ Begin USkinnedAsset interface
 	/** Initial step for the Post Load process - Can't be done in parallel. */
@@ -217,13 +220,12 @@ private:
 	/** Reregister all components using this asset to reset the simulation in case anything has changed. */
 	void ReregisterComponents();
 
-	// TODO: determine if it should be serialized or transient
 	/** List of materials for this cloth asset. */
 	UPROPERTY(EditAnywhere, Category = Materials)
 	TArray<FSkeletalMaterial> Materials;
 
-	/** Skeleton asset used at import time. */
-	UPROPERTY(EditAnywhere, NoClear, Setter = SetSkeleton, Category = Skeleton)
+	/** Skeleton asset used at creation time. This is of limited use since this USkeleton's reference skeleton might not necessarily match the one created for this asset. */
+	UPROPERTY(EditAnywhere, Setter = SetSkeleton, Category = Skeleton)
 	TObjectPtr<USkeleton> Skeleton;
 
 	/** Physics asset used for collision. */
diff --git a/Engine/Plugins/Experimental/ChaosClothAssetEditor/Source/ChaosClothAssetDataflowNodes/Private/ChaosClothAsset/DatasmithImportNode.cpp b/Engine/Plugins/Experimental/ChaosClothAssetEditor/Source/ChaosClothAssetDataflowNodes/Private/ChaosClothAsset/DatasmithImportNode.cpp
--- a/Engine/Plugins/Experimental/ChaosClothAssetEditor/Source/ChaosClothAssetDataflowNodes/Private/ChaosClothAsset/DatasmithImportNode.cpp
+++ b/Engine/Plugins/Experimental/ChaosClothAssetEditor/Source/ChaosClothAssetDataflowNodes/Private/ChaosClothAsset/DatasmithImportNode.cpp
@@ -22,8 +22,21 @@ FChaosClothAssetDatasmithImportNode::FChaosClothAssetDatasmithImportNode(const D
 
 void FChaosClothAssetDatasmithImportNode::Serialize(FArchive& Archive)
 {
+	using namespace UE::Chaos::ClothAsset;
+
 	::Chaos::FChaosArchive ChaosArchive(Archive);
 	ImportCache.Serialize(ChaosArchive);
+	if (Archive.IsLoading())
+	{
+		// Make sure to always have a valid cloth collection on reload, some new attributes could be missing from the cached collection
+		const TSharedRef<FManagedArrayCollection> ClothCollection = MakeShared<FManagedArrayCollection>(MoveTemp(ImportCache));
+		FCollectionClothFacade ClothFacade(ClothCollection);
+		if (!ClothFacade.IsValid())
+		{
+			ClothFacade.DefineSchema();
+		}
+		ImportCache = MoveTemp(*ClothCollection);
+	}
 	Archive << ImportHash;
 }
 
diff --git a/Engine/Plugins/Experimental/ChaosClothAssetEditor/Source/ChaosClothAssetDataflowNodes/Private/ChaosClothAsset/SkeletalMeshImportNode.cpp b/Engine/Plugins/Experimental/ChaosClothAssetEditor/Source/ChaosClothAssetDataflowNodes/Private/ChaosClothAsset/SkeletalMeshImportNode.cpp
--- a/Engine/Plugins/Experimental/ChaosClothAssetEditor/Source/ChaosClothAssetDataflowNodes/Private/ChaosClothAsset/SkeletalMeshImportNode.cpp
+++ b/Engine/Plugins/Experimental/ChaosClothAssetEditor/Source/ChaosClothAssetDataflowNodes/Private/ChaosClothAsset/SkeletalMeshImportNode.cpp
@@ -83,15 +83,7 @@ void FChaosClothAssetSkeletalMeshImportNode::Evaluate(Dataflow::FContext& Contex
 				ClothFacade.SetPhysicsAssetPathName(PhysicsAsset->GetPathName());
 			}
 
-			if (const USkeleton* Skeleton = SkeletalMesh->GetSkeleton())
-			{
-				ClothFacade.SetSkeletonAssetPathName(Skeleton->GetPathName());
-			}
-			else
-			{
-				static const TCHAR* const DefaultSkeletonPathName = TEXT("/Engine/EditorMeshes/SkeletalMesh/DefaultSkeletalMesh_Skeleton.DefaultSkeletalMesh_Skeleton");
-				ClothFacade.SetSkeletonAssetPathName(DefaultSkeletonPathName);
-			}
+			ClothFacade.SetSkeletalMeshPathName(SkeletalMesh->GetPathName());
 
 			// Must be bound to root bone by default
 			FClothGeometryTools::BindMeshToRootBone(ClothCollection, bImportSimMesh, bImportRenderMesh);
diff --git a/Engine/Plugins/Experimental/ChaosClothAssetEditor/Source/ChaosClothAssetDataflowNodes/Private/ChaosClothAsset/StaticMeshImportNode.cpp b/Engine/Plugins/Experimental/ChaosClothAssetEditor/Source/ChaosClothAssetDataflowNodes/Private/ChaosClothAsset/StaticMeshImportNode.cpp
--- a/Engine/Plugins/Experimental/ChaosClothAssetEditor/Source/ChaosClothAssetDataflowNodes/Private/ChaosClothAsset/StaticMeshImportNode.cpp
+++ b/Engine/Plugins/Experimental/ChaosClothAssetEditor/Source/ChaosClothAssetDataflowNodes/Private/ChaosClothAsset/StaticMeshImportNode.cpp
@@ -128,8 +128,6 @@ void FChaosClothAssetStaticMeshImportNode::Evaluate(Dataflow::FContext& Context,
 
 	if (Out->IsA<FManagedArrayCollection>(&Collection))
 	{
-		static const TCHAR* const DefaultSkeletonPathName = TEXT("/Engine/EditorMeshes/SkeletalMesh/DefaultSkeletalMesh_Skeleton.DefaultSkeletalMesh_Skeleton");
-
 		// Evaluate out collection
 		const TSharedRef<FManagedArrayCollection> ClothCollection = MakeShared<FManagedArrayCollection>();
 		FCollectionClothFacade ClothFacade(ClothCollection);
@@ -167,9 +165,6 @@ void FChaosClothAssetStaticMeshImportNode::Evaluate(Dataflow::FContext& Context,
 								FText::FromString(StaticMesh->GetName())));
 					}
 				}
-
-				// Set a default skeleton
-				ClothFacade.SetSkeletonAssetPathName(DefaultSkeletonPathName);
 				
 				// Bind to root bone by default
 				FClothGeometryTools::BindMeshToRootBone(ClothCollection, bImportAsSimMesh, bImportAsRenderMesh);
diff --git a/Engine/Plugins/Experimental/ChaosClothAssetEditor/Source/ChaosClothAssetDataflowNodes/Private/ChaosClothAsset/TransferSkinWeightsNode.cpp b/Engine/Plugins/Experimental/ChaosClothAssetEditor/Source/ChaosClothAssetDataflowNodes/Private/ChaosClothAsset/TransferSkinWeightsNode.cpp
--- a/Engine/Plugins/Experimental/ChaosClothAssetEditor/Source/ChaosClothAssetDataflowNodes/Private/ChaosClothAsset/TransferSkinWeightsNode.cpp
+++ b/Engine/Plugins/Experimental/ChaosClothAssetEditor/Source/ChaosClothAssetDataflowNodes/Private/ChaosClothAsset/TransferSkinWeightsNode.cpp
@@ -470,7 +470,7 @@ void FChaosClothAssetTransferSkinWeightsNode::Evaluate(Dataflow::FContext& Conte
 			FDynamicMesh3 SourceDynamicMesh;
 			SkeletalMeshToDynamicMesh(SkeletalMesh, LodIndex, SourceDynamicMesh);
 			MeshTransforms::ApplyTransform(SourceDynamicMesh, Transform, true);
-			const FReferenceSkeleton& TargetRefSkeleton = SkeletalMesh->GetSkeleton()->GetReferenceSkeleton();
+			const FReferenceSkeleton& TargetRefSkeleton = SkeletalMesh->GetRefSkeleton();
 
 			//
 			// Setup the bone weight transfer operator for the source mesh.
@@ -484,7 +484,7 @@ void FChaosClothAssetTransferSkinWeightsNode::Evaluate(Dataflow::FContext& Conte
 			// Transfer the bone weights from the source Skeletal mesh to the Cloth asset.
 			//
 			FCollectionClothFacade ClothFacade(ClothCollection);
-			ClothFacade.SetSkeletonAssetPathName(SkeletalMesh->GetSkeleton()->GetPathName());
+			ClothFacade.SetSkeletalMeshPathName(SkeletalMesh->GetPathName());
 
 			bool bTransferResult = false;
 			if (TransferMethod == EChaosClothAssetTransferSkinWeightsMethod::InpaintWeights)
diff --git a/Engine/Plugins/Experimental/ChaosClothAssetEditor/Source/ChaosClothAssetEditor/Private/ChaosClothAsset/AssetDefinition_ClothAsset.cpp b/Engine/Plugins/Experimental/ChaosClothAssetEditor/Source/ChaosClothAssetEditor/Private/ChaosClothAsset/AssetDefinition_ClothAsset.cpp
--- a/Engine/Plugins/Experimental/ChaosClothAssetEditor/Source/ChaosClothAssetEditor/Private/ChaosClothAsset/AssetDefinition_ClothAsset.cpp
+++ b/Engine/Plugins/Experimental/ChaosClothAssetEditor/Source/ChaosClothAssetEditor/Private/ChaosClothAsset/AssetDefinition_ClothAsset.cpp
@@ -186,7 +186,7 @@ EAssetCommandResult UAssetDefinition_ClothAsset::OpenAssets(const FAssetOpenArgs
 
 		// Validate the asset
 		UChaosClothAsset* const ClothAsset = CastChecked<UChaosClothAsset>(ClothObjects[0]);
-
+PRAGMA_DISABLE_DEPRECATION_WARNINGS  // TODO: Don't use public property, and have Getter/Setter API instead
 		if (!ClothAsset->DataflowAsset)
 		{
 			if (UDataflow* const NewDataflowAsset = Cast<UDataflow>(ClothAssetDefinitionHelpers::NewOrOpenDataflowAsset(ClothAsset)))
@@ -194,6 +194,7 @@ EAssetCommandResult UAssetDefinition_ClothAsset::OpenAssets(const FAssetOpenArgs
 				ClothAsset->DataflowAsset = NewDataflowAsset;
 			}
 		}
+PRAGMA_ENABLE_DEPRECATION_WARNINGS
 
 		TArray<TObjectPtr<UObject>> Objects;
 		for (UChaosClothAsset* const ClothObject : ClothObjects)
diff --git a/Engine/Plugins/Experimental/ChaosClothAssetEditor/Source/ChaosClothAssetEditor/Private/ChaosClothAsset/ClothEditorToolkit.cpp b/Engine/Plugins/Experimental/ChaosClothAssetEditor/Source/ChaosClothAssetEditor/Private/ChaosClothAsset/ClothEditorToolkit.cpp
--- a/Engine/Plugins/Experimental/ChaosClothAssetEditor/Source/ChaosClothAssetEditor/Private/ChaosClothAsset/ClothEditorToolkit.cpp
+++ b/Engine/Plugins/Experimental/ChaosClothAssetEditor/Source/ChaosClothAssetEditor/Private/ChaosClothAsset/ClothEditorToolkit.cpp
@@ -532,7 +532,9 @@ void FChaosClothAssetEditorToolkit::OnAssetsSavedAs(const TArray<UObject*>& Save
 
 	if (NewClothAsset && NewDataflowAsset)
 	{
+PRAGMA_DISABLE_DEPRECATION_WARNINGS  // TODO: Don't use public property, and have Getter/Setter API instead
 		NewClothAsset->DataflowAsset = NewDataflowAsset;
+PRAGMA_ENABLE_DEPRECATION_WARNINGS
 
 		// Now save the new Cloth asset again since we've updated its Property
 		const TArray<UPackage*> PackagesToSave{NewClothAsset->GetOutermost()};
diff --git a/Engine/Plugins/Experimental/ChaosClothAssetEditor/Source/ChaosClothAssetTools/Private/ChaosClothAsset/ClothAssetFactory.cpp b/Engine/Plugins/Experimental/ChaosClothAssetEditor/Source/ChaosClothAssetTools/Private/ChaosClothAsset/ClothAssetFactory.cpp
--- a/Engine/Plugins/Experimental/ChaosClothAssetEditor/Source/ChaosClothAssetTools/Private/ChaosClothAsset/ClothAssetFactory.cpp
+++ b/Engine/Plugins/Experimental/ChaosClothAssetEditor/Source/ChaosClothAssetTools/Private/ChaosClothAsset/ClothAssetFactory.cpp
@@ -23,15 +23,5 @@ UObject* UChaosClothAssetFactory::FactoryCreateNew(UClass* Class, UObject* Paren
 	UChaosClothAsset* const ClothAsset = NewObject<UChaosClothAsset>(Parent, Class, Name, Flags | RF_Transactional | RF_Public | RF_Standalone);
 	ClothAsset->MarkPackageDirty();
 
-	// Add an empty default LOD, to avoid LOD mismatch with render data
-	TArray<TSharedRef<FManagedArrayCollection>>& ClothCollections = ClothAsset->GetClothCollections();
-	TSharedRef<FManagedArrayCollection>& ClothCollection = ClothCollections.Emplace_GetRef(MakeShared<FManagedArrayCollection>());
-
-	FCollectionClothFacade ClothFacade(ClothCollection);
-	ClothFacade.DefineSchema();
-
-	// Set the default skeleton on this new LOD and rebuild the static data models (which is done by default with this override)
-	ClothAsset->SetSkeleton(nullptr);
-
 	return ClothAsset;
 }
diff --git a/Engine/Plugins/Experimental/ChaosClothAssetEditor/Source/ChaosClothAssetTools/Private/ChaosClothAsset/ClothPatternToDynamicMesh.cpp b/Engine/Plugins/Experimental/ChaosClothAssetEditor/Source/ChaosClothAssetTools/Private/ChaosClothAsset/ClothPatternToDynamicMesh.cpp
--- a/Engine/Plugins/Experimental/ChaosClothAssetEditor/Source/ChaosClothAssetTools/Private/ChaosClothAsset/ClothPatternToDynamicMesh.cpp
+++ b/Engine/Plugins/Experimental/ChaosClothAssetEditor/Source/ChaosClothAssetTools/Private/ChaosClothAsset/ClothPatternToDynamicMesh.cpp
@@ -6,11 +6,11 @@
 #include "ChaosClothAsset/ClothAsset.h"
 #include "ChaosClothAsset/ClothComponent.h"
 #include "ChaosClothAsset/CollectionClothFacade.h"
-#include "ToDynamicMesh.h"
 #include "DynamicMesh/MeshNormals.h"
 #include "DynamicMesh/NonManifoldMappingSupport.h"
+#include "Engine/SkeletalMesh.h"
 #include "SkeletalMeshAttributes.h"
-#include "Animation/Skeleton.h"
+#include "ToDynamicMesh.h"
 
 namespace UE::Chaos::ClothAsset
 {
@@ -220,11 +220,11 @@ public:
 		}
 
 		// Set the reference skeleton if available
-		const FString& SkeletonPathName = Cloth.GetSkeletonAssetPathName();
-		const USkeleton* const Skeleton = !SkeletonPathName.IsEmpty() ?
-			LoadObject<USkeleton>(nullptr, *SkeletonPathName, nullptr, LOAD_None, nullptr) :
+		const FString& SkeletalMeshPathName = ClothFacade.GetSkeletalMeshPathName();
+		const USkeletalMesh* const SkeletalMesh = !SkeletalMeshPathName.IsEmpty() ?
+			LoadObject<USkeletalMesh>(nullptr, *SkeletalMeshPathName, nullptr, LOAD_None, nullptr) :
 			nullptr;
-		RefSkeleton = Skeleton ? &Skeleton->GetReferenceSkeleton() : nullptr;
+		RefSkeleton = SkeletalMesh ? &SkeletalMesh->GetRefSkeleton() : nullptr;
 	}
 
 	int32 NumTris() const
diff --git a/Engine/Plugins/Experimental/ChaosClothAssetEditor/Source/ChaosClothAssetTools/Private/ChaosClothAsset/ClothingAssetToClothAssetExporter.cpp b/Engine/Plugins/Experimental/ChaosClothAssetEditor/Source/ChaosClothAssetTools/Private/ChaosClothAsset/ClothingAssetToClothAssetExporter.cpp
--- a/Engine/Plugins/Experimental/ChaosClothAssetEditor/Source/ChaosClothAssetTools/Private/ChaosClothAsset/ClothingAssetToClothAssetExporter.cpp
+++ b/Engine/Plugins/Experimental/ChaosClothAssetEditor/Source/ChaosClothAssetTools/Private/ChaosClothAsset/ClothingAssetToClothAssetExporter.cpp
@@ -141,13 +141,12 @@ void UClothingAssetToChaosClothAssetExporter::Export(const UClothingAssetBase* C
 	// Assign the physics asset if any (must be done after having added the LODs)
 	ClothAsset->SetPhysicsAsset(ClothingAssetCommon->PhysicsAsset);
 
-	// Set the skeleton from the skeletal mesh (must be done after having added the LODs)
-	constexpr bool bRebuildModels = false;  // Build is called last
+	// Set the reference skeleton from the skeletal mesh and rebuild the asset (must be done after having added the LODs)
 	USkeletalMesh* const SkeletalMesh = CastChecked<USkeletalMesh>(ClothingAssetCommon->GetOuter());
-	ClothAsset->SetSkeleton(SkeletalMesh->GetSkeleton(), bRebuildModels);
+	ClothAsset->SetReferenceSkeleton(&SkeletalMesh->GetRefSkeleton());
 
-	// Build the asset, since it is already loaded, it won't rebuild on load
-	ClothAsset->Build();
+	// Save the USkeleton for completion, although this is not being used and might mismatch the skeletal mesh's reference skeleton
+	ClothAsset->SetSkeleton(SkeletalMesh->GetSkeleton());
 }
 
 #undef LOCTEXT_NAMESPACE
