Chaos Cloth Asset - Fixed a crash when trying to use an empty cloth asset with a cloth component.

#rnx

diff --git a/Engine/Plugins/Experimental/ChaosClothAsset/Source/ChaosClothAssetEngine/Private/ChaosClothAsset/ClothComponent.cpp b/Engine/Plugins/Experimental/ChaosClothAsset/Source/ChaosClothAssetEngine/Private/ChaosClothAsset/ClothComponent.cpp
--- a/Engine/Plugins/Experimental/ChaosClothAsset/Source/ChaosClothAssetEngine/Private/ChaosClothAsset/ClothComponent.cpp
+++ b/Engine/Plugins/Experimental/ChaosClothAsset/Source/ChaosClothAssetEngine/Private/ChaosClothAsset/ClothComponent.cpp
@@ -106,7 +106,7 @@ void UChaosClothComponent::RecreateClothSimulationProxy()
 		if (GetClothAsset())
 		{
 			const TSharedPtr<const FChaosClothSimulationModel> ClothSimulationModel = GetClothAsset()->GetClothSimulationModel();
-			if (ensure(ClothSimulationModel) && ClothSimulationModel->GetNumLods())
+			if (ClothSimulationModel && ClothSimulationModel->GetNumLods())
 			{
 				// Create the simulation proxy (note CreateClothSimulationProxy() can be overloaded)
 				ClothSimulationProxy = CreateClothSimulationProxy();
@@ -394,7 +394,7 @@ void UChaosClothComponent::UpdateComponentSpaceTransforms()
 {
 	check(IsRegistered());
 
-	if (!LeaderPoseComponent.IsValid() && GetClothAsset())
+	if (!LeaderPoseComponent.IsValid() && GetClothAsset() && GetClothAsset()->GetResourceForRendering())
 	{
 		FSkeletalMeshLODRenderData& LODData = GetClothAsset()->GetResourceForRendering()->LODRenderData[GetPredictedLODLevel()];
 		GetClothAsset()->FillComponentSpaceTransforms(GetClothAsset()->GetRefSkeleton().GetRefBonePose(), LODData.RequiredBones, GetEditableComponentSpaceTransforms());
@@ -407,7 +407,7 @@ void UChaosClothComponent::UpdateComponentSpaceTransforms()
 
 void UChaosClothComponent::UpdateVisibility()
 {
-	if (GetClothAsset())
+	if (GetClothAsset() && GetClothAsset()->GetResourceForRendering())
 	{
 		const FSkeletalMeshRenderData* const SkeletalMeshRenderData = GetClothAsset()->GetResourceForRendering();
 		const int32 FirstValidLODIdx = SkeletalMeshRenderData ? SkeletalMeshRenderData->GetFirstValidLODIdx(0) : INDEX_NONE;
