Chaos Cloth Asset: Move UpdateUVChannelData to happen inside InitResources (or to be called in its place when relevant). This follows how the SkeletalMesh works.

This fixes a crash that happens if you remove an LOD as the LODInfo and FSkeletalMeshRenderData need to both be updated before UpdateUVChannelData is called.

diff --git a/Engine/Plugins/Experimental/ChaosClothAsset/Source/ChaosClothAssetEngine/Private/ChaosClothAsset/ClothAsset.cpp b/Engine/Plugins/Experimental/ChaosClothAsset/Source/ChaosClothAssetEngine/Private/ChaosClothAsset/ClothAsset.cpp
--- a/Engine/Plugins/Experimental/ChaosClothAsset/Source/ChaosClothAssetEngine/Private/ChaosClothAsset/ClothAsset.cpp
+++ b/Engine/Plugins/Experimental/ChaosClothAsset/Source/ChaosClothAssetEngine/Private/ChaosClothAsset/ClothAsset.cpp
@@ -284,6 +284,11 @@ void UChaosClothAsset::FinishPostLoadInternal(FSkinnedAssetPostLoadContext& Cont
 	{
 		InitResources();
 	}
+	else
+	{
+		// Update any missing data when cooking.
+		UpdateUVChannelData(false);
+	}
 
 	CalculateInvRefMatrices();
 	CalculateBounds();
@@ -320,6 +325,9 @@ void UChaosClothAsset::InitResources()
 {
 	LLM_SCOPE_BYNAME(TEXT("ClothAsset/InitResources"));
 
+	// Build the material channel data used by the texture streamer
+	UpdateUVChannelData(false);
+
 	if (SkeletalMeshRenderData.IsValid())
 	{
 		TArray<UMorphTarget*> DummyMorphTargets;  // Not even used by InitResources() at the moment
@@ -464,10 +472,9 @@ void UChaosClothAsset::Build(TArray<FChaosClothAssetLodTransitionDataCache>* InO
 	BuildMeshModel();
 #endif
 
-	// Load/save render data from/to DDC
 #if WITH_EDITOR
+	// Load/save render data from/to DDC
 	ExecuteBuildInternal(Context);
-	FinishBuildInternal(Context);
 #endif
 
 	if (FApp::CanEverRender())
@@ -475,6 +482,10 @@ void UChaosClothAsset::Build(TArray<FChaosClothAssetLodTransitionDataCache>* InO
 		InitResources();
 	}
 
+#if WITH_EDITOR
+	FinishBuildInternal(Context);
+#endif
+
 	// Re-register any components using this asset to restart the simulation with the updated asset
 	ReregisterComponents();
 }
@@ -489,6 +500,9 @@ void UChaosClothAsset::ExecuteBuildInternal(FSkinnedAssetBuildContext& Context)
 
 	// rebuild render data from imported model
 	CacheDerivedData(&Context);
+
+	// Build the material channel data used by the texture streamer
+	UpdateUVChannelData(true);
 }
 
 void UChaosClothAsset::BeginBuildInternal(FSkinnedAssetBuildContext& Context)
@@ -543,9 +557,6 @@ void UChaosClothAsset::BuildMeshModel()
 				ClothAssetBuilder->BuildLod(*LODModel, *this, LodIndex);
 				MeshModel->LODModels.Add(LODModel);
 			}
-
-			// Build the material channel data used by the texture streamer
-			UpdateUVChannelData(true);
 		}
 	}
 }
