Fully enable Cloth Component caching with the Chaos Cache Manager.

Main changes:
* implementing the callbacks for reading and writing simulation data in the ClothComponentCacheAdapter
* changing the behavior of the ClothComponent to bTickInEditor when there is cache data available so that the cache will playback correctly (with the correct dependency on any leader SKM) when scrubbing in Sequencer. This was done by adding the bTickOnceInEditor which when set, an in-editor tick will consume the bTickOnceInEditor and bTickInEditor flag, so subsequent ticks will not call the ClothComponent tick.

diff --git a/Engine/Plugins/Experimental/ChaosClothAsset/Source/ChaosClothAssetEngine/Private/ChaosClothAsset/ClothComponent.cpp b/Engine/Plugins/Experimental/ChaosClothAsset/Source/ChaosClothAssetEngine/Private/ChaosClothAsset/ClothComponent.cpp
--- a/Engine/Plugins/Experimental/ChaosClothAsset/Source/ChaosClothAssetEngine/Private/ChaosClothAsset/ClothComponent.cpp
+++ b/Engine/Plugins/Experimental/ChaosClothAsset/Source/ChaosClothAssetEngine/Private/ChaosClothAsset/ClothComponent.cpp
@@ -118,6 +118,13 @@ void UChaosClothComponent::UpdateConfigProperties()
 }
 #endif
 
+void UChaosClothComponent::WaitForExistingParallelClothSimulation_GameThread()
+{
+	// Should only kick new parallel cloth simulations from game thread, so should be safe to also wait for existing ones there.
+	check(IsInGameThread());
+	HandleExistingParallelSimulation();
+}
+
 void UChaosClothComponent::RecreateClothSimulationProxy()
 {
 	if (IsRegistered())
@@ -233,6 +240,15 @@ void UChaosClothComponent::TickComponent(float DeltaTime, enum ELevelTick TickTy
 	{
 		HandleExistingParallelSimulation();
 	}
+
+#if WITH_EDITOR
+	if (TickType == LEVELTICK_ViewportsOnly && bTickOnceInEditor)
+	{
+		// Only tick once in editor when requested. This is used to update from caches by the Chaos Cache Manager.
+		bTickInEditor = false;
+		bTickOnceInEditor = false;
+	}
+#endif
 }
 
 bool UChaosClothComponent::RequiresPreEndOfFrameSync() const
diff --git a/Engine/Plugins/Experimental/ChaosClothAsset/Source/ChaosClothAssetEngine/Private/ChaosClothAsset/ClothComponentCacheAdapter.cpp b/Engine/Plugins/Experimental/ChaosClothAsset/Source/ChaosClothAssetEngine/Private/ChaosClothAsset/ClothComponentCacheAdapter.cpp
--- a/Engine/Plugins/Experimental/ChaosClothAsset/Source/ChaosClothAssetEngine/Private/ChaosClothAsset/ClothComponentCacheAdapter.cpp
+++ b/Engine/Plugins/Experimental/ChaosClothAsset/Source/ChaosClothAssetEngine/Private/ChaosClothAsset/ClothComponentCacheAdapter.cpp
@@ -41,12 +41,19 @@ namespace UE::Chaos::ClothAsset
 
 	void FClothComponentCacheAdapter::Record_PostSolve(UPrimitiveComponent* InComponent, const FTransform& InRootTransform, FPendingFrameWrite& OutFrame, ::Chaos::FReal InTime) const
 	{
-		ensureMsgf(false, TEXT("Record not implemented."));
+		if (const ::Chaos::FClothingSimulationSolver* ClothSolver = GetClothSolver(InComponent))
+		{
+			::Chaos::FClothingCacheSchema::RecordPostSolve(*ClothSolver, OutFrame, InTime);
+		}
 	}
 
 	void FClothComponentCacheAdapter::Playback_PreSolve(UPrimitiveComponent* InComponent, UChaosCache* InCache, ::Chaos::FReal InTime, FPlaybackTickRecord& TickRecord, TArray<::Chaos::TPBDRigidParticleHandle<::Chaos::FReal, 3>*>& OutUpdatedRigids) const
 	{
-		ensureMsgf(false, TEXT("Playback not implemented."));
+		check(InCache);
+		if (::Chaos::FClothingSimulationSolver* ClothSolver = GetClothSolver(InComponent))
+		{
+			::Chaos::FClothingCacheSchema::PlaybackPreSolve(*InCache, InTime, TickRecord, *ClothSolver);
+		}
 	}
 
 	FGuid FClothComponentCacheAdapter::GetGuid() const
@@ -64,7 +71,11 @@ namespace UE::Chaos::ClothAsset
 
 	::Chaos::FPhysicsSolverEvents* FClothComponentCacheAdapter::BuildEventsSolver(UPrimitiveComponent* InComponent) const
 	{
-		ensureMsgf(false, TEXT("Playback or record not implemented."));
+		if (UChaosClothComponent* ClothComp = GetClothComponent(InComponent))
+		{
+			ClothComp->RecreateClothSimulationProxy();
+			return ClothComp->ClothSimulationProxy->Solver.Get();
+		}
 		return nullptr;
 	}
 	
@@ -80,31 +91,52 @@ namespace UE::Chaos::ClothAsset
 			return;
 		}
 
-		if (UChaosClothComponent* ClothComp = GetClothComponent(InComponent))
+		if (UChaosClothComponent* const ClothComponent = GetClothComponent(InComponent))
 		{
-			FClothingSimulationCacheData CacheData;
-			::Chaos::FClothingCacheSchema::LoadCacheData(InCache, InTime, CacheData);
-
-			FClothSimulationProxy* Proxy = ClothComp->ClothSimulationProxy.Get();
-
-			Proxy->Tick_GameThread(0, &CacheData);
-			Proxy->CompleteParallelSimulation_GameThread();
-			ClothComp->MarkRenderDynamicDataDirty();
-			ClothComp->DoDeferredRenderUpdates_Concurrent();
+			if (FClothSimulationProxy* const Proxy = ClothComponent->ClothSimulationProxy.Get())
+			{
+				FClothingSimulationCacheData CacheData;
+				::Chaos::FClothingCacheSchema::LoadCacheData(InCache, InTime, CacheData);
+
+				Proxy->CacheData = MakeUnique<FClothingSimulationCacheData>(MoveTemp(CacheData));
+				Proxy->SolverMode = FClothSimulationProxy::ESolverMode::Default;
+#if WITH_EDITOR
+				ClothComponent->SetTickOnceInEditor();
+#endif
+			}
 		}
 	}
 
 	bool FClothComponentCacheAdapter::InitializeForRecord(UPrimitiveComponent* InComponent, UChaosCache* InCache)
 	{
+		if (FClothSimulationProxy* const Proxy = GetProxy(InComponent))
+		{
+			Proxy->SolverMode = FClothSimulationProxy::ESolverMode::EnableSolverForSimulateRecord;
+		}
 		return true;
 	}
 
 	bool FClothComponentCacheAdapter::InitializeForPlayback(UPrimitiveComponent* InComponent, UChaosCache* InCache, float InTime)
 	{
 		::Chaos::EnsureIsInGameThreadContext();
+
+		if (FClothSimulationProxy* const Proxy = GetProxy(InComponent))
+		{
+			Proxy->SolverMode = FClothSimulationProxy::ESolverMode::DisableSolverForPlayback;
+		}
 		return true;
 	}
 
+	void FClothComponentCacheAdapter::WaitForSolverTasks(UPrimitiveComponent* InComponent) const
+	{
+		::Chaos::EnsureIsInGameThreadContext();
+
+		if (UChaosClothComponent* const ClothComponent = GetClothComponent(InComponent))
+		{
+			ClothComponent->WaitForExistingParallelClothSimulation_GameThread();
+		}
+	}
+
 	UChaosClothComponent* FClothComponentCacheAdapter::GetClothComponent(UPrimitiveComponent* InComponent) const
 	{
 		if (InComponent)
@@ -131,4 +163,22 @@ namespace UE::Chaos::ClothAsset
 		return nullptr;
 	}
 
+	FClothSimulationProxy* FClothComponentCacheAdapter::GetProxy(UPrimitiveComponent* InComponent) const
+	{
+		if (UChaosClothComponent* const ClothComponent = GetClothComponent(InComponent))
+		{
+			return ClothComponent->ClothSimulationProxy.Get();
+		}
+		return nullptr;
+	}
+
+	::Chaos::FClothingSimulationSolver* FClothComponentCacheAdapter::GetClothSolver(UPrimitiveComponent* InComponent) const
+	{
+		if (FClothSimulationProxy* const Proxy = GetProxy(InComponent))
+		{
+			return Proxy->Solver.Get();
+		}
+		return nullptr;
+	}
+
 }    // namespace Chaos
diff --git a/Engine/Plugins/Experimental/ChaosClothAsset/Source/ChaosClothAssetEngine/Private/ChaosClothAsset/ClothComponentCacheAdapter.h b/Engine/Plugins/Experimental/ChaosClothAsset/Source/ChaosClothAssetEngine/Private/ChaosClothAsset/ClothComponentCacheAdapter.h
--- a/Engine/Plugins/Experimental/ChaosClothAsset/Source/ChaosClothAssetEngine/Private/ChaosClothAsset/ClothComponentCacheAdapter.h
+++ b/Engine/Plugins/Experimental/ChaosClothAsset/Source/ChaosClothAssetEngine/Private/ChaosClothAsset/ClothComponentCacheAdapter.h
@@ -12,6 +12,8 @@ namespace Chaos
 
 namespace UE::Chaos::ClothAsset
 {	
+	class FClothSimulationProxy;
+
 	/**
 	 * Skeletal mesh cache adapter to be able to cache cloth simulation datas through the chaos cache system
 	 */
@@ -38,11 +40,15 @@ namespace UE::Chaos::ClothAsset
 										 ::Chaos::FReal										InTime,
 										 FPlaybackTickRecord&								TickRecord,
 										 TArray<::Chaos::TPBDRigidParticleHandle<::Chaos::FReal, 3>*>& OutUpdatedRigids) const override;
+		virtual void                   WaitForSolverTasks(UPrimitiveComponent* InComponent) const override;
 		// ~End FComponentCacheAdapter interface
 
 	private :
 		/** Return the cloth component.
 		 */
 		UChaosClothComponent* GetClothComponent(UPrimitiveComponent* InComponent) const;
+
+		FClothSimulationProxy* GetProxy(UPrimitiveComponent* InComponent) const;
+		::Chaos::FClothingSimulationSolver* GetClothSolver(UPrimitiveComponent* InComponent) const;
 	};
 }    // namespace Chaos
diff --git a/Engine/Plugins/Experimental/ChaosClothAsset/Source/ChaosClothAssetEngine/Private/ChaosClothAsset/ClothSimulationProxy.cpp b/Engine/Plugins/Experimental/ChaosClothAsset/Source/ChaosClothAssetEngine/Private/ChaosClothAsset/ClothSimulationProxy.cpp
--- a/Engine/Plugins/Experimental/ChaosClothAsset/Source/ChaosClothAssetEngine/Private/ChaosClothAsset/ClothSimulationProxy.cpp
+++ b/Engine/Plugins/Experimental/ChaosClothAsset/Source/ChaosClothAssetEngine/Private/ChaosClothAsset/ClothSimulationProxy.cpp
@@ -167,19 +167,16 @@ PRAGMA_ENABLE_DEPRECATION_WARNINGS
 		CompleteParallelSimulation_GameThread();
 	}
 
-	bool FClothSimulationProxy::Tick_GameThread(float DeltaTime, FClothingSimulationCacheData* CacheData)
+	bool FClothSimulationProxy::Tick_GameThread(float DeltaTime)
 	{
 		SCOPE_CYCLE_COUNTER(STAT_ClothSimulationProxy_TickGame);
 
 		// Fill a new context, note the context is also needed when the simulation is suspended
-		ClothSimulationContext->Fill(ClothComponent, DeltaTime, MaxDeltaTime, false, CacheData);
-
-		const bool bUseCache = ClothSimulationContext->CacheData.CacheIndices.Num() > 0;
-		if (bUseCache)
-		{
-			Solver->SetEnableSolver(false);
-		}
+		constexpr bool bIsInitializationFalse = false;
+		FillSimulationContext(DeltaTime, bIsInitializationFalse);
+		Solver->SetEnableSolver(ShouldEnableSolver(Solver->GetEnableSolver()));
 
+		const bool bUseCache = ClothSimulationContext->CacheData.HasData();
 		const bool bCreateParallelTask = (DeltaTime > 0.f && !ClothComponent.IsSimulationSuspended()) || bUseCache;
 		if (bCreateParallelTask)
 		{
@@ -461,8 +458,27 @@ PRAGMA_ENABLE_DEPRECATION_WARNINGS
 		Solver->SetSolverLOD(ClothSimulationContext->LodIndex);
 	}
 
-	void FClothSimulationProxy::FillSimulationContext(float DeltaTime, bool bIsInitialization, FClothingSimulationCacheData* CacheData)
+	void FClothSimulationProxy::FillSimulationContext(float DeltaTime, bool bIsInitialization)
 	{
-		ClothSimulationContext->Fill(ClothComponent, DeltaTime, MaxDeltaTime, bIsInitialization, CacheData);
+		ClothSimulationContext->Fill(ClothComponent, DeltaTime, MaxDeltaTime, bIsInitialization, CacheData.Get());
+		CacheData.Reset();
+	}
+
+	bool FClothSimulationProxy::ShouldEnableSolver(bool bSolverCurrentlyEnabled) const
+	{
+		switch (SolverMode)
+		{
+		case ESolverMode::EnableSolverForSimulateRecord:
+			return true;
+		case ESolverMode::DisableSolverForPlayback:
+			return false;
+		case ESolverMode::Default:
+		default:
+			if (ClothSimulationContext->CacheData.HasData())
+			{
+				return false;
+			}
+		}
+		return bSolverCurrentlyEnabled;
 	}
 }
diff --git a/Engine/Plugins/Experimental/ChaosClothAsset/Source/ChaosClothAssetEngine/Public/ChaosClothAsset/ClothComponent.h b/Engine/Plugins/Experimental/ChaosClothAsset/Source/ChaosClothAssetEngine/Public/ChaosClothAsset/ClothComponent.h
--- a/Engine/Plugins/Experimental/ChaosClothAsset/Source/ChaosClothAssetEngine/Public/ChaosClothAsset/ClothComponent.h
+++ b/Engine/Plugins/Experimental/ChaosClothAsset/Source/ChaosClothAssetEngine/Public/ChaosClothAsset/ClothComponent.h
@@ -107,6 +107,14 @@ public:
 	void UpdateConfigProperties();
 #endif
 
+	/** Stalls on any currently running clothing simulations.*/
+	void WaitForExistingParallelClothSimulation_GameThread();
+
+#if WITH_EDITOR
+	/** This will cause the component to tick once in editor. Both flags will be consumed on that tick */
+	void SetTickOnceInEditor() { bTickOnceInEditor = true; bTickInEditor = true; }
+#endif
+
 protected:
 	//~ Begin UObject Interface
 	virtual void PostLoad() override;
@@ -186,6 +194,10 @@ private:
 	UPROPERTY()
 	float BlendWeight = 1.f;
 
+#if WITH_EDITOR
+	bool bTickOnceInEditor = false;
+#endif
+
 	TArray<TSharedPtr<FManagedArrayCollection>> PropertyCollections;
 	TArray<TUniquePtr<::Chaos::Softs::FCollectionPropertyFacade>> CollectionPropertyFacades;
 
diff --git a/Engine/Plugins/Experimental/ChaosClothAsset/Source/ChaosClothAssetEngine/Public/ChaosClothAsset/ClothSimulationProxy.h b/Engine/Plugins/Experimental/ChaosClothAsset/Source/ChaosClothAssetEngine/Public/ChaosClothAsset/ClothSimulationProxy.h
--- a/Engine/Plugins/Experimental/ChaosClothAsset/Source/ChaosClothAssetEngine/Public/ChaosClothAsset/ClothSimulationProxy.h
+++ b/Engine/Plugins/Experimental/ChaosClothAsset/Source/ChaosClothAssetEngine/Public/ChaosClothAsset/ClothSimulationProxy.h
@@ -45,7 +45,7 @@ namespace UE::Chaos::ClothAsset
 		FClothSimulationProxy& operator=(FClothSimulationProxy&&) = delete;
 
 		/** Start the simulation if it isn't paused or suspended and return true, or simply update the existing simulation data and return false otherwise. */
-		bool Tick_GameThread(float DeltaTime, FClothingSimulationCacheData* CacheData = nullptr);
+		bool Tick_GameThread(float DeltaTime);
 
 		/** Wait for the parallel task to complete if one was running, and update the simulation data. */
 		void CompleteParallelSimulation_GameThread();
@@ -73,9 +73,11 @@ namespace UE::Chaos::ClothAsset
 		void Tick();
 		void WriteSimulationData();
 		void InitializeConfigs();
-		void FillSimulationContext(float DeltaTime, bool bIsInitialization = false, FClothingSimulationCacheData* CacheData = nullptr);
+		void FillSimulationContext(float DeltaTime, bool bIsInitialization = false);
 
 	private:
+		bool ShouldEnableSolver(bool bSolverCurrentlyEnabled) const;
+
 		// Internal physics thread object
 		friend class FClothSimulationProxyParallelTask;
 
@@ -104,6 +106,19 @@ PRAGMA_ENABLE_DEPRECATION_WARNINGS
 		TArray<TUniquePtr<::Chaos::FClothingSimulationCollider>> Colliders;
 		TUniquePtr<::Chaos::FClothVisualization> Visualization;
 
+		// Chaos Cache needs to have access to the solver.
+		friend class FClothComponentCacheAdapter;
+
+		// Additional data used by the cache adapter
+		enum struct ESolverMode : uint8
+		{
+			Default = 0, // Default behavior. Enable solver if no cache data available.
+			EnableSolverForSimulateRecord = 1, // Normal simulation. Also used when Recording.
+			DisableSolverForPlayback = 2, // Solver is disabled. Used when live playing back cache.
+		};
+		TUniquePtr<FClothingSimulationCacheData> CacheData;
+		ESolverMode SolverMode = ESolverMode::Default;
+
 		// Properties that must be readable from all threads
 		std::atomic<int32> NumCloths;
 		std::atomic<int32> NumKinematicParticles;
