Cloth asset: fix type hashing in terminal node (hash the data, not the pointers). Also, hash the maps since they are stored on the asset and require rebuilding to pick up currently.

diff --git a/Engine/Plugins/Experimental/ChaosClothAsset/Source/ChaosClothAsset/Private/ChaosClothAsset/ClothCollection.h b/Engine/Plugins/Experimental/ChaosClothAsset/Source/ChaosClothAsset/Private/ChaosClothAsset/ClothCollection.h
--- a/Engine/Plugins/Experimental/ChaosClothAsset/Source/ChaosClothAsset/Private/ChaosClothAsset/ClothCollection.h
+++ b/Engine/Plugins/Experimental/ChaosClothAsset/Source/ChaosClothAsset/Private/ChaosClothAsset/ClothCollection.h
@@ -139,6 +139,9 @@ namespace UE::Chaos::ClothAsset
 
 		static void CopyArrayViewDataAndApplyOffset(const TArrayView<TArray<int32>>& To, const TConstArrayView<TArray<int32>>& From, const int32 Offset);
 
+		template<typename T>
+		static inline uint32 GetElementsTypeHash(const TManagedArray<T>* ElementArray);
+
 		//~ Weight maps
 		template<typename T, TEMPLATE_REQUIRES(TIsUserAttributeType<T>::Value)>
 		TArray<FName> GetUserDefinedAttributeNames(const FName& GroupName) const;
@@ -447,4 +450,14 @@ namespace UE::Chaos::ClothAsset
 			To[Index] = From[Index] + Offset;
 		}
 	}
+
+	template<typename T>
+	inline uint32 FClothCollection::GetElementsTypeHash(const TManagedArray<T>* ElementArray)
+	{
+		if (ElementArray)
+		{
+			return GetTypeHash(*ElementArray);
+		}
+		return GetTypeHash(ElementArray);
+	}
 }  // End namespace UE::Chaos::ClothAsset
diff --git a/Engine/Plugins/Experimental/ChaosClothAsset/Source/ChaosClothAsset/Private/ChaosClothAsset/CollectionClothFacade.cpp b/Engine/Plugins/Experimental/ChaosClothAsset/Source/ChaosClothAsset/Private/ChaosClothAsset/CollectionClothFacade.cpp
--- a/Engine/Plugins/Experimental/ChaosClothAsset/Source/ChaosClothAsset/Private/ChaosClothAsset/CollectionClothFacade.cpp
+++ b/Engine/Plugins/Experimental/ChaosClothAsset/Source/ChaosClothAsset/Private/ChaosClothAsset/CollectionClothFacade.cpp
@@ -42,59 +42,59 @@ namespace UE::Chaos::ClothAsset
 
 		//~ LOD (single per collection) Group
 		ResultHash = HashCombineFast(ResultHash, GetTypeHash(GetSkeletalMeshPathName()));
-		ResultHash = HashCombineFast(ResultHash, GetTypeHash(GetSkeletalMeshPathName()));
+		ResultHash = HashCombineFast(ResultHash, GetTypeHash(GetPhysicsAssetPathName()));
 
 		//~ Seam Stitches Group
-		ResultHash = HashCombineFast(ResultHash, GetTypeHash(ClothCollection->GetSeamStitchStart()));
-		ResultHash = HashCombineFast(ResultHash, GetTypeHash(ClothCollection->GetSeamStitchEnd()));
+		ResultHash = HashCombineFast(ResultHash, ClothCollection->GetElementsTypeHash(ClothCollection->GetSeamStitchStart()));
+		ResultHash = HashCombineFast(ResultHash, ClothCollection->GetElementsTypeHash(ClothCollection->GetSeamStitchEnd()));
 
 		//~ Seam Stitches Group
-		ResultHash = HashCombineFast(ResultHash, GetTypeHash(ClothCollection->GetSeamStitch2DEndIndices()));
-		ResultHash = HashCombineFast(ResultHash, GetTypeHash(ClothCollection->GetSeamStitch3DIndex()));
+		ResultHash = HashCombineFast(ResultHash, ClothCollection->GetElementsTypeHash(ClothCollection->GetSeamStitch2DEndIndices()));
+		ResultHash = HashCombineFast(ResultHash, ClothCollection->GetElementsTypeHash(ClothCollection->GetSeamStitch3DIndex()));
 
 		//~ Sim Patterns Group
-		ResultHash = HashCombineFast(ResultHash, GetTypeHash(ClothCollection->GetSimVertices2DStart()));
-		ResultHash = HashCombineFast(ResultHash, GetTypeHash(ClothCollection->GetSimVertices2DEnd()));
-		ResultHash = HashCombineFast(ResultHash, GetTypeHash(ClothCollection->GetSimFacesStart()));
-		ResultHash = HashCombineFast(ResultHash, GetTypeHash(ClothCollection->GetSimFacesEnd()));
+		ResultHash = HashCombineFast(ResultHash, ClothCollection->GetElementsTypeHash(ClothCollection->GetSimVertices2DStart()));
+		ResultHash = HashCombineFast(ResultHash, ClothCollection->GetElementsTypeHash(ClothCollection->GetSimVertices2DEnd()));
+		ResultHash = HashCombineFast(ResultHash, ClothCollection->GetElementsTypeHash(ClothCollection->GetSimFacesStart()));
+		ResultHash = HashCombineFast(ResultHash, ClothCollection->GetElementsTypeHash(ClothCollection->GetSimFacesEnd()));
 
 		//~ Render Patterns Group
-		ResultHash = HashCombineFast(ResultHash, GetTypeHash(ClothCollection->GetRenderVerticesStart()));
-		ResultHash = HashCombineFast(ResultHash, GetTypeHash(ClothCollection->GetRenderVerticesEnd()));
-		ResultHash = HashCombineFast(ResultHash, GetTypeHash(ClothCollection->GetRenderFacesStart()));
-		ResultHash = HashCombineFast(ResultHash, GetTypeHash(ClothCollection->GetRenderFacesEnd()));
-		ResultHash = HashCombineFast(ResultHash, GetTypeHash(ClothCollection->GetRenderMaterialPathName()));
+		ResultHash = HashCombineFast(ResultHash, ClothCollection->GetElementsTypeHash(ClothCollection->GetRenderVerticesStart()));
+		ResultHash = HashCombineFast(ResultHash, ClothCollection->GetElementsTypeHash(ClothCollection->GetRenderVerticesEnd()));
+		ResultHash = HashCombineFast(ResultHash, ClothCollection->GetElementsTypeHash(ClothCollection->GetRenderFacesStart()));
+		ResultHash = HashCombineFast(ResultHash, ClothCollection->GetElementsTypeHash(ClothCollection->GetRenderFacesEnd()));
+		ResultHash = HashCombineFast(ResultHash, ClothCollection->GetElementsTypeHash(ClothCollection->GetRenderMaterialPathName()));
 
 		//~ Sim Faces Group
-		ResultHash = HashCombineFast(ResultHash, GetTypeHash(ClothCollection->GetSimIndices2D()));
-		ResultHash = HashCombineFast(ResultHash, GetTypeHash(ClothCollection->GetSimIndices3D()));
+		ResultHash = HashCombineFast(ResultHash, ClothCollection->GetElementsTypeHash(ClothCollection->GetSimIndices2D()));
+		ResultHash = HashCombineFast(ResultHash, ClothCollection->GetElementsTypeHash(ClothCollection->GetSimIndices3D()));
 
 		//~ Sim Vertices 2D Group
-		ResultHash = HashCombineFast(ResultHash, GetTypeHash(ClothCollection->GetSimPosition2D()));
-		ResultHash = HashCombineFast(ResultHash, GetTypeHash(ClothCollection->GetSimVertex3DLookup()));
+		ResultHash = HashCombineFast(ResultHash, ClothCollection->GetElementsTypeHash(ClothCollection->GetSimPosition2D()));
+		ResultHash = HashCombineFast(ResultHash, ClothCollection->GetElementsTypeHash(ClothCollection->GetSimVertex3DLookup()));
 		
 		//~ Sim Vertices 3D Group
-		ResultHash = HashCombineFast(ResultHash, GetTypeHash(ClothCollection->GetSimPosition3D()));
-		ResultHash = HashCombineFast(ResultHash, GetTypeHash(ClothCollection->GetSimNormal()));
-		ResultHash = HashCombineFast(ResultHash, GetTypeHash(ClothCollection->GetSimBoneIndices()));
-		ResultHash = HashCombineFast(ResultHash, GetTypeHash(ClothCollection->GetSimBoneWeights()));
-		ResultHash = HashCombineFast(ResultHash, GetTypeHash(ClothCollection->GetTetherKinematicIndex()));
-		ResultHash = HashCombineFast(ResultHash, GetTypeHash(ClothCollection->GetTetherReferenceLength()));
-		ResultHash = HashCombineFast(ResultHash, GetTypeHash(ClothCollection->GetSimVertex2DLookup()));
-		ResultHash = HashCombineFast(ResultHash, GetTypeHash(ClothCollection->GetSeamStitchLookup()));
+		ResultHash = HashCombineFast(ResultHash, ClothCollection->GetElementsTypeHash(ClothCollection->GetSimPosition3D()));
+		ResultHash = HashCombineFast(ResultHash, ClothCollection->GetElementsTypeHash(ClothCollection->GetSimNormal()));
+		ResultHash = HashCombineFast(ResultHash, ClothCollection->GetElementsTypeHash(ClothCollection->GetSimBoneIndices()));
+		ResultHash = HashCombineFast(ResultHash, ClothCollection->GetElementsTypeHash(ClothCollection->GetSimBoneWeights()));
+		ResultHash = HashCombineFast(ResultHash, ClothCollection->GetElementsTypeHash(ClothCollection->GetTetherKinematicIndex()));
+		ResultHash = HashCombineFast(ResultHash, ClothCollection->GetElementsTypeHash(ClothCollection->GetTetherReferenceLength()));
+		ResultHash = HashCombineFast(ResultHash, ClothCollection->GetElementsTypeHash(ClothCollection->GetSimVertex2DLookup()));
+		ResultHash = HashCombineFast(ResultHash, ClothCollection->GetElementsTypeHash(ClothCollection->GetSeamStitchLookup()));
 
 		//~ Render Faces Group
-		ResultHash = HashCombineFast(ResultHash, GetTypeHash(ClothCollection->GetRenderIndices()));
+		ResultHash = HashCombineFast(ResultHash, ClothCollection->GetElementsTypeHash(ClothCollection->GetRenderIndices()));
 
 		//~ Render Vertices Group
-		ResultHash = HashCombineFast(ResultHash, GetTypeHash(ClothCollection->GetRenderPosition()));
-		ResultHash = HashCombineFast(ResultHash, GetTypeHash(ClothCollection->GetRenderNormal()));
-		ResultHash = HashCombineFast(ResultHash, GetTypeHash(ClothCollection->GetRenderTangentU()));
-		ResultHash = HashCombineFast(ResultHash, GetTypeHash(ClothCollection->GetRenderTangentV()));
-		ResultHash = HashCombineFast(ResultHash, GetTypeHash(ClothCollection->GetRenderUVs()));
-		ResultHash = HashCombineFast(ResultHash, GetTypeHash(ClothCollection->GetRenderColor()));
-		ResultHash = HashCombineFast(ResultHash, GetTypeHash(ClothCollection->GetRenderBoneIndices()));
-		ResultHash = HashCombineFast(ResultHash, GetTypeHash(ClothCollection->GetRenderBoneWeights()));
+		ResultHash = HashCombineFast(ResultHash, ClothCollection->GetElementsTypeHash(ClothCollection->GetRenderPosition()));
+		ResultHash = HashCombineFast(ResultHash, ClothCollection->GetElementsTypeHash(ClothCollection->GetRenderNormal()));
+		ResultHash = HashCombineFast(ResultHash, ClothCollection->GetElementsTypeHash(ClothCollection->GetRenderTangentU()));
+		ResultHash = HashCombineFast(ResultHash, ClothCollection->GetElementsTypeHash(ClothCollection->GetRenderTangentV()));
+		ResultHash = HashCombineFast(ResultHash, ClothCollection->GetElementsTypeHash(ClothCollection->GetRenderUVs()));
+		ResultHash = HashCombineFast(ResultHash, ClothCollection->GetElementsTypeHash(ClothCollection->GetRenderColor()));
+		ResultHash = HashCombineFast(ResultHash, ClothCollection->GetElementsTypeHash(ClothCollection->GetRenderBoneIndices()));
+		ResultHash = HashCombineFast(ResultHash, ClothCollection->GetElementsTypeHash(ClothCollection->GetRenderBoneWeights()));
 
 		if (bIncludeWeightMaps)
 		{
@@ -110,7 +110,7 @@ namespace UE::Chaos::ClothAsset
 		for (const FName& WeightMapName : WeightMapNames)
 		{
 			ResultHash = HashCombineFast(ResultHash, GetTypeHash(WeightMapName));
-			ResultHash = HashCombineFast(ResultHash, GetTypeHash(ClothCollection->GetUserDefinedAttribute<float>(WeightMapName, FClothCollection::SimVertices3DGroup)));
+			ResultHash = HashCombineFast(ResultHash, ClothCollection->GetElementsTypeHash(ClothCollection->GetUserDefinedAttribute<float>(WeightMapName, FClothCollection::SimVertices3DGroup)));
 		}
 		return ResultHash;
 	}
diff --git a/Engine/Plugins/Experimental/ChaosClothAssetEditor/Source/ChaosClothAssetDataflowNodes/Private/ChaosClothAsset/TerminalNode.cpp b/Engine/Plugins/Experimental/ChaosClothAssetEditor/Source/ChaosClothAssetDataflowNodes/Private/ChaosClothAsset/TerminalNode.cpp
--- a/Engine/Plugins/Experimental/ChaosClothAssetEditor/Source/ChaosClothAssetDataflowNodes/Private/ChaosClothAsset/TerminalNode.cpp
+++ b/Engine/Plugins/Experimental/ChaosClothAssetEditor/Source/ChaosClothAssetDataflowNodes/Private/ChaosClothAsset/TerminalNode.cpp
@@ -24,11 +24,11 @@ namespace UE::Chaos::ClothAsset::Private
 		uint32 Checksum = 0;
 		for (const TSharedRef<FManagedArrayCollection>& ClothCollection : InClothCollections)
 		{
-			constexpr bool bIncludeWeightMapsFalse = false;
+			constexpr bool bIncludeWeightMapsTrue = true; // Currently, editing weight maps is destructive
 			FCollectionClothConstFacade Cloth(ClothCollection);
 			if (Cloth.HasValidData())
 			{
-				Checksum = Cloth.CalculateTypeHash(bIncludeWeightMapsFalse, Checksum);
+				Checksum = Cloth.CalculateTypeHash(bIncludeWeightMapsTrue, Checksum);
 			}
 		}
 		return Checksum;
