Chaos Cloth Asset:
- Made the selection node take a group name as a selection type.
- Changed the selection facade schema to allow for group dependency to the selected index sets.
- Added group dependency to the selection attributes.
- Added a detail customization to allow for choosing the selection group from the collection being processed by the selection node.
- Added Dataflow OnSelected and OnDeselected callback methods to allow for custom processing when a toolkit changes the node selection, without having to disrupt the const evaluation of the nodes.
- Changed the order of operations in SDataflowGraphEditor::OnSelectedNodesChanged() to allow for the nodes to refresh their UI before the selection is actually changed.
- Added a group index dependency getter to the managed array collection.

#rnx

diff --git a/Engine/Plugins/Experimental/ChaosClothAsset/Source/ChaosClothAsset/Private/ChaosClothAsset/ClothCollection.cpp b/Engine/Plugins/Experimental/ChaosClothAsset/Source/ChaosClothAsset/Private/ChaosClothAsset/ClothCollection.cpp
--- a/Engine/Plugins/Experimental/ChaosClothAsset/Source/ChaosClothAsset/Private/ChaosClothAsset/ClothCollection.cpp
+++ b/Engine/Plugins/Experimental/ChaosClothAsset/Source/ChaosClothAsset/Private/ChaosClothAsset/ClothCollection.cpp
@@ -1,22 +1,11 @@
 // Copyright Epic Games, Inc. All Rights Reserved.
 
 #include "ChaosClothAsset/ClothCollection.h"
+#include "ChaosClothAsset/ClothCollectionGroup.h"
 #include "GeometryCollection/ManagedArrayCollection.h"
 
 namespace UE::Chaos::ClothAsset::Private
 {
-	// Groups
-	static const FName LodsGroup(TEXT("Lods")); // One lod per collection, but keeping the name for now to represent whatever is a single scalar for the whole collection
-	static const FName SeamsGroup(TEXT("Seams"));
-	static const FName SeamStitchesGroup(TEXT("SeamStitches"));
-	static const FName SimPatternsGroup(TEXT("SimPatterns"));
-	static const FName RenderPatternsGroup(TEXT("RenderPatterns"));
-	static const FName SimFacesGroup(TEXT("SimFaces"));
-	static const FName SimVertices2DGroup(TEXT("SimVertices2D"));
-	static const FName SimVertices3DGroup(TEXT("SimVertices3D"));
-	static const FName RenderFacesGroup(TEXT("RenderFaces"));
-	static const FName RenderVerticesGroup(TEXT("RenderVertices"));
-
 	// Lods Group
 	static const FName PhysicsAssetPathNameAttribute(TEXT("PhysicsAssetPathName"));
 	static const FName SkeletalMeshPathNameAttribute(TEXT("SkeletalMeshPathName"));
@@ -142,93 +131,81 @@ namespace UE::Chaos::ClothAsset::Private
 
 	static const TMap<FName, TArray<FName>> FixedAttributeNamesMap =
 	{
-		{ LodsGroup, LodsGroupAttributes },
-		{ SeamsGroup, SeamsGroupAttributes },
-		{ SeamStitchesGroup, SeamStitchesGroupAttributes },
-		{ SimPatternsGroup, SimPatternsGroupAttributes },
-		{ RenderPatternsGroup, RenderPatternsGroupAttributes },
-		{ SimFacesGroup, SimFacesGroupAttributes },
-		{ SimVertices2DGroup, SimVertices2DGroupAttributes },
-		{ SimVertices3DGroup, SimVertices3DGroupAttributes },
-		{ RenderFacesGroup, RenderFacesGroupAttributes },
-		{ RenderVerticesGroup, RenderVerticesGroupAttributes }
+		{ ClothCollectionGroup::Lods, LodsGroupAttributes },
+		{ ClothCollectionGroup::Seams, SeamsGroupAttributes },
+		{ ClothCollectionGroup::SeamStitches, SeamStitchesGroupAttributes },
+		{ ClothCollectionGroup::SimPatterns, SimPatternsGroupAttributes },
+		{ ClothCollectionGroup::RenderPatterns, RenderPatternsGroupAttributes },
+		{ ClothCollectionGroup::SimFaces, SimFacesGroupAttributes },
+		{ ClothCollectionGroup::SimVertices2D, SimVertices2DGroupAttributes },
+		{ ClothCollectionGroup::SimVertices3D, SimVertices3DGroupAttributes },
+		{ ClothCollectionGroup::RenderFaces, RenderFacesGroupAttributes },
+		{ ClothCollectionGroup::RenderVertices, RenderVerticesGroupAttributes }
 	};
 }  // End namespace UE::Chaos::ClothAsset::Private
 
 namespace UE::Chaos::ClothAsset
 {
-	// Groups
-	const FName FClothCollection::LodsGroup = Private::LodsGroup;
-	const FName FClothCollection::SeamsGroup = Private::SeamsGroup;
-	const FName FClothCollection::SeamStitchesGroup = Private::SeamStitchesGroup;
-	const FName FClothCollection::SimPatternsGroup = Private::SimPatternsGroup;
-	const FName FClothCollection::RenderPatternsGroup = Private::RenderPatternsGroup;
-	const FName FClothCollection::SimFacesGroup = Private::SimFacesGroup;
-	const FName FClothCollection::SimVertices2DGroup = Private::SimVertices2DGroup;
-	const FName FClothCollection::SimVertices3DGroup = Private::SimVertices3DGroup;
-	const FName FClothCollection::RenderFacesGroup = Private::RenderFacesGroup;
-	const FName FClothCollection::RenderVerticesGroup = Private::RenderVerticesGroup;
-
 	FClothCollection::FClothCollection(const TSharedRef<FManagedArrayCollection>& InManagedArrayCollection)
 		: ManagedArrayCollection(InManagedArrayCollection)
 	{
 		using namespace UE::Chaos::ClothAsset::Private;
 
 		// LODs Group
-		PhysicsAssetPathName = ManagedArrayCollection->FindAttribute<FString>(PhysicsAssetPathNameAttribute, LodsGroup);
-		SkeletalMeshPathName = ManagedArrayCollection->FindAttribute<FString>(SkeletalMeshPathNameAttribute, LodsGroup);
+		PhysicsAssetPathName = ManagedArrayCollection->FindAttribute<FString>(PhysicsAssetPathNameAttribute, ClothCollectionGroup::Lods);
+		SkeletalMeshPathName = ManagedArrayCollection->FindAttribute<FString>(SkeletalMeshPathNameAttribute, ClothCollectionGroup::Lods);
 
 		// Seam Group
-		SeamStitchStart = ManagedArrayCollection->FindAttribute<int32>(SeamStitchStartAttribute, SeamsGroup);
-		SeamStitchEnd = ManagedArrayCollection->FindAttribute<int32>(SeamStitchEndAttribute, SeamsGroup);
+		SeamStitchStart = ManagedArrayCollection->FindAttribute<int32>(SeamStitchStartAttribute, ClothCollectionGroup::Seams);
+		SeamStitchEnd = ManagedArrayCollection->FindAttribute<int32>(SeamStitchEndAttribute, ClothCollectionGroup::Seams);
 
 		// Seam Stitches Group
-		SeamStitch2DEndIndices = ManagedArrayCollection->FindAttribute<FIntVector2>(SeamStitch2DEndIndicesAttribute, SeamStitchesGroup);
-		SeamStitch3DIndex = ManagedArrayCollection->FindAttribute<int32>(SeamStitch3DIndexAttribute, SeamStitchesGroup);
+		SeamStitch2DEndIndices = ManagedArrayCollection->FindAttribute<FIntVector2>(SeamStitch2DEndIndicesAttribute, ClothCollectionGroup::SeamStitches);
+		SeamStitch3DIndex = ManagedArrayCollection->FindAttribute<int32>(SeamStitch3DIndexAttribute, ClothCollectionGroup::SeamStitches);
 
 		// Sim Patterns Group
-		SimVertices2DStart = ManagedArrayCollection->FindAttribute<int32>(SimVertices2DStartAttribute, SimPatternsGroup);
-		SimVertices2DEnd = ManagedArrayCollection->FindAttribute<int32>(SimVertices2DEndAttribute, SimPatternsGroup);
-		SimFacesStart = ManagedArrayCollection->FindAttribute<int32>(SimFacesStartAttribute, SimPatternsGroup);
-		SimFacesEnd = ManagedArrayCollection->FindAttribute<int32>(SimFacesEndAttribute, SimPatternsGroup);
+		SimVertices2DStart = ManagedArrayCollection->FindAttribute<int32>(SimVertices2DStartAttribute, ClothCollectionGroup::SimPatterns);
+		SimVertices2DEnd = ManagedArrayCollection->FindAttribute<int32>(SimVertices2DEndAttribute, ClothCollectionGroup::SimPatterns);
+		SimFacesStart = ManagedArrayCollection->FindAttribute<int32>(SimFacesStartAttribute, ClothCollectionGroup::SimPatterns);
+		SimFacesEnd = ManagedArrayCollection->FindAttribute<int32>(SimFacesEndAttribute, ClothCollectionGroup::SimPatterns);
 
 		// Render Patterns Group
-		RenderVerticesStart = ManagedArrayCollection->FindAttribute<int32>(RenderVerticesStartAttribute, RenderPatternsGroup);
-		RenderVerticesEnd = ManagedArrayCollection->FindAttribute<int32>(RenderVerticesEndAttribute, RenderPatternsGroup);
-		RenderFacesStart = ManagedArrayCollection->FindAttribute<int32>(RenderFacesStartAttribute, RenderPatternsGroup);
-		RenderFacesEnd = ManagedArrayCollection->FindAttribute<int32>(RenderFacesEndAttribute, RenderPatternsGroup);
-		RenderMaterialPathName = ManagedArrayCollection->FindAttribute<FString>(RenderMaterialPathNameAttribute, RenderPatternsGroup);
+		RenderVerticesStart = ManagedArrayCollection->FindAttribute<int32>(RenderVerticesStartAttribute, ClothCollectionGroup::RenderPatterns);
+		RenderVerticesEnd = ManagedArrayCollection->FindAttribute<int32>(RenderVerticesEndAttribute, ClothCollectionGroup::RenderPatterns);
+		RenderFacesStart = ManagedArrayCollection->FindAttribute<int32>(RenderFacesStartAttribute, ClothCollectionGroup::RenderPatterns);
+		RenderFacesEnd = ManagedArrayCollection->FindAttribute<int32>(RenderFacesEndAttribute, ClothCollectionGroup::RenderPatterns);
+		RenderMaterialPathName = ManagedArrayCollection->FindAttribute<FString>(RenderMaterialPathNameAttribute, ClothCollectionGroup::RenderPatterns);
 
 		// Sim Faces Group
-		SimIndices2D = ManagedArrayCollection->FindAttribute<FIntVector3>(SimIndices2DAttribute, SimFacesGroup);
-		SimIndices3D = ManagedArrayCollection->FindAttribute<FIntVector3>(SimIndices3DAttribute, SimFacesGroup);
+		SimIndices2D = ManagedArrayCollection->FindAttribute<FIntVector3>(SimIndices2DAttribute, ClothCollectionGroup::SimFaces);
+		SimIndices3D = ManagedArrayCollection->FindAttribute<FIntVector3>(SimIndices3DAttribute, ClothCollectionGroup::SimFaces);
 
 		// Sim Vertices 2D Group
-		SimPosition2D = ManagedArrayCollection->FindAttribute<FVector2f>(SimPosition2DAttribute, SimVertices2DGroup);
-		SimVertex3DLookup = ManagedArrayCollection->FindAttribute<int32>(SimVertex3DLookupAttribute, SimVertices2DGroup);
+		SimPosition2D = ManagedArrayCollection->FindAttribute<FVector2f>(SimPosition2DAttribute, ClothCollectionGroup::SimVertices2D);
+		SimVertex3DLookup = ManagedArrayCollection->FindAttribute<int32>(SimVertex3DLookupAttribute, ClothCollectionGroup::SimVertices2D);
 
 		// Sim Vertices 3D Group
-		SimPosition3D = ManagedArrayCollection->FindAttribute<FVector3f>(SimPosition3DAttribute, SimVertices3DGroup);
-		SimNormal = ManagedArrayCollection->FindAttribute<FVector3f>(SimNormalAttribute, SimVertices3DGroup);
-		SimBoneIndices = ManagedArrayCollection->FindAttribute<TArray<int32>>(SimBoneIndicesAttribute, SimVertices3DGroup);
-		SimBoneWeights = ManagedArrayCollection->FindAttribute<TArray<float>>(SimBoneWeightsAttribute, SimVertices3DGroup);
-		TetherKinematicIndex = ManagedArrayCollection->FindAttribute<TArray<int32>>(TetherKinematicIndexAttribute, SimVertices3DGroup);
-		TetherReferenceLength = ManagedArrayCollection->FindAttribute<TArray<float>>(TetherReferenceLengthAttribute, SimVertices3DGroup);
-		SimVertex2DLookup = ManagedArrayCollection->FindAttribute<TArray<int32>>(SimVertex2DLookupAttribute, SimVertices3DGroup);
-		SeamStitchLookup = ManagedArrayCollection->FindAttribute<TArray<int32>>(SeamStitchLookupAttribute, SimVertices3DGroup);
+		SimPosition3D = ManagedArrayCollection->FindAttribute<FVector3f>(SimPosition3DAttribute, ClothCollectionGroup::SimVertices3D);
+		SimNormal = ManagedArrayCollection->FindAttribute<FVector3f>(SimNormalAttribute, ClothCollectionGroup::SimVertices3D);
+		SimBoneIndices = ManagedArrayCollection->FindAttribute<TArray<int32>>(SimBoneIndicesAttribute, ClothCollectionGroup::SimVertices3D);
+		SimBoneWeights = ManagedArrayCollection->FindAttribute<TArray<float>>(SimBoneWeightsAttribute, ClothCollectionGroup::SimVertices3D);
+		TetherKinematicIndex = ManagedArrayCollection->FindAttribute<TArray<int32>>(TetherKinematicIndexAttribute, ClothCollectionGroup::SimVertices3D);
+		TetherReferenceLength = ManagedArrayCollection->FindAttribute<TArray<float>>(TetherReferenceLengthAttribute, ClothCollectionGroup::SimVertices3D);
+		SimVertex2DLookup = ManagedArrayCollection->FindAttribute<TArray<int32>>(SimVertex2DLookupAttribute, ClothCollectionGroup::SimVertices3D);
+		SeamStitchLookup = ManagedArrayCollection->FindAttribute<TArray<int32>>(SeamStitchLookupAttribute, ClothCollectionGroup::SimVertices3D);
 
 		// Render Faces Group
-		RenderIndices = ManagedArrayCollection->FindAttribute<FIntVector3>(RenderIndicesAttribute, RenderFacesGroup);
+		RenderIndices = ManagedArrayCollection->FindAttribute<FIntVector3>(RenderIndicesAttribute, ClothCollectionGroup::RenderFaces);
 
 		// Render Vertices Group
-		RenderPosition = ManagedArrayCollection->FindAttribute<FVector3f>(RenderPositionAttribute, RenderVerticesGroup);
-		RenderNormal = ManagedArrayCollection->FindAttribute<FVector3f>(RenderNormalAttribute, RenderVerticesGroup);
-		RenderTangentU = ManagedArrayCollection->FindAttribute<FVector3f>(RenderTangentUAttribute, RenderVerticesGroup);
-		RenderTangentV = ManagedArrayCollection->FindAttribute<FVector3f>(RenderTangentVAttribute, RenderVerticesGroup);
-		RenderUVs = ManagedArrayCollection->FindAttribute<TArray<FVector2f>>(RenderUVsAttribute, RenderVerticesGroup);
-		RenderColor = ManagedArrayCollection->FindAttribute<FLinearColor>(RenderColorAttribute, RenderVerticesGroup);
-		RenderBoneIndices = ManagedArrayCollection->FindAttribute<TArray<int32>>(RenderBoneIndicesAttribute, RenderVerticesGroup);
-		RenderBoneWeights = ManagedArrayCollection->FindAttribute<TArray<float>>(RenderBoneWeightsAttribute, RenderVerticesGroup);
+		RenderPosition = ManagedArrayCollection->FindAttribute<FVector3f>(RenderPositionAttribute, ClothCollectionGroup::RenderVertices);
+		RenderNormal = ManagedArrayCollection->FindAttribute<FVector3f>(RenderNormalAttribute, ClothCollectionGroup::RenderVertices);
+		RenderTangentU = ManagedArrayCollection->FindAttribute<FVector3f>(RenderTangentUAttribute, ClothCollectionGroup::RenderVertices);
+		RenderTangentV = ManagedArrayCollection->FindAttribute<FVector3f>(RenderTangentVAttribute, ClothCollectionGroup::RenderVertices);
+		RenderUVs = ManagedArrayCollection->FindAttribute<TArray<FVector2f>>(RenderUVsAttribute, ClothCollectionGroup::RenderVertices);
+		RenderColor = ManagedArrayCollection->FindAttribute<FLinearColor>(RenderColorAttribute, ClothCollectionGroup::RenderVertices);
+		RenderBoneIndices = ManagedArrayCollection->FindAttribute<TArray<int32>>(RenderBoneIndicesAttribute, ClothCollectionGroup::RenderVertices);
+		RenderBoneWeights = ManagedArrayCollection->FindAttribute<TArray<float>>(RenderBoneWeightsAttribute, ClothCollectionGroup::RenderVertices);
 	}
 
 	bool FClothCollection::IsValid() const
@@ -298,68 +275,68 @@ namespace UE::Chaos::ClothAsset
 		// Dependencies
 		constexpr bool bSaved = true;
 		constexpr bool bAllowCircularDependency = true;
-		FManagedArrayCollection::FConstructionParameters SeamStitchesDependency(SeamStitchesGroup, bSaved, bAllowCircularDependency);
-		FManagedArrayCollection::FConstructionParameters RenderFacesDependency(RenderFacesGroup, bSaved, bAllowCircularDependency);
-		FManagedArrayCollection::FConstructionParameters RenderVerticesDependency(RenderVerticesGroup, bSaved, bAllowCircularDependency);
-		FManagedArrayCollection::FConstructionParameters SimFacesDependency(SimFacesGroup, bSaved, bAllowCircularDependency);
-		FManagedArrayCollection::FConstructionParameters SimVertices2DDependency(SimVertices2DGroup, bSaved, bAllowCircularDependency);
-		FManagedArrayCollection::FConstructionParameters SimVertices3DDependency(SimVertices3DGroup, bSaved, bAllowCircularDependency);  // Any attribute with this dependency must handle welding and splitting in FCollectionClothSeamFacade
+		FManagedArrayCollection::FConstructionParameters SeamStitchesDependency(ClothCollectionGroup::SeamStitches, bSaved, bAllowCircularDependency);
+		FManagedArrayCollection::FConstructionParameters RenderFacesDependency(ClothCollectionGroup::RenderFaces, bSaved, bAllowCircularDependency);
+		FManagedArrayCollection::FConstructionParameters RenderVerticesDependency(ClothCollectionGroup::RenderVertices, bSaved, bAllowCircularDependency);
+		FManagedArrayCollection::FConstructionParameters SimFacesDependency(ClothCollectionGroup::SimFaces, bSaved, bAllowCircularDependency);
+		FManagedArrayCollection::FConstructionParameters SimVertices2DDependency(ClothCollectionGroup::SimVertices2D, bSaved, bAllowCircularDependency);
+		FManagedArrayCollection::FConstructionParameters SimVertices3DDependency(ClothCollectionGroup::SimVertices3D, bSaved, bAllowCircularDependency);  // Any attribute with this dependency must handle welding and splitting in FCollectionClothSeamFacade
 
 		// LODs Group
-		PhysicsAssetPathName = &ManagedArrayCollection->AddAttribute<FString>(PhysicsAssetPathNameAttribute, LodsGroup);
-		SkeletalMeshPathName = &ManagedArrayCollection->AddAttribute<FString>(SkeletalMeshPathNameAttribute, LodsGroup);
+		PhysicsAssetPathName = &ManagedArrayCollection->AddAttribute<FString>(PhysicsAssetPathNameAttribute, ClothCollectionGroup::Lods);
+		SkeletalMeshPathName = &ManagedArrayCollection->AddAttribute<FString>(SkeletalMeshPathNameAttribute, ClothCollectionGroup::Lods);
 
 		// Seams Group
-		SeamStitchStart = &ManagedArrayCollection->AddAttribute<int32>(SeamStitchStartAttribute, SeamsGroup, SeamStitchesDependency);
-		SeamStitchEnd = &ManagedArrayCollection->AddAttribute<int32>(SeamStitchEndAttribute, SeamsGroup, SeamStitchesDependency);
+		SeamStitchStart = &ManagedArrayCollection->AddAttribute<int32>(SeamStitchStartAttribute, ClothCollectionGroup::Seams, SeamStitchesDependency);
+		SeamStitchEnd = &ManagedArrayCollection->AddAttribute<int32>(SeamStitchEndAttribute, ClothCollectionGroup::Seams, SeamStitchesDependency);
 
 		// Seam Stitches Group
-		SeamStitch2DEndIndices = &ManagedArrayCollection->AddAttribute<FIntVector2>(SeamStitch2DEndIndicesAttribute, SeamStitchesGroup, SimVertices2DDependency);
-		SeamStitch3DIndex = &ManagedArrayCollection->AddAttribute<int32>(SeamStitch3DIndexAttribute, SeamStitchesGroup, SimVertices3DDependency);
+		SeamStitch2DEndIndices = &ManagedArrayCollection->AddAttribute<FIntVector2>(SeamStitch2DEndIndicesAttribute, ClothCollectionGroup::SeamStitches, SimVertices2DDependency);
+		SeamStitch3DIndex = &ManagedArrayCollection->AddAttribute<int32>(SeamStitch3DIndexAttribute, ClothCollectionGroup::SeamStitches, SimVertices3DDependency);
 
 		// Sim Patterns Group
-		SimVertices2DStart = &ManagedArrayCollection->AddAttribute<int32>(SimVertices2DStartAttribute, SimPatternsGroup, SimVertices2DDependency);
-		SimVertices2DEnd = &ManagedArrayCollection->AddAttribute<int32>(SimVertices2DEndAttribute, SimPatternsGroup, SimVertices2DDependency);
-		SimFacesStart = &ManagedArrayCollection->AddAttribute<int32>(SimFacesStartAttribute, SimPatternsGroup, SimFacesDependency);
-		SimFacesEnd = &ManagedArrayCollection->AddAttribute<int32>(SimFacesEndAttribute, SimPatternsGroup, SimFacesDependency);
+		SimVertices2DStart = &ManagedArrayCollection->AddAttribute<int32>(SimVertices2DStartAttribute, ClothCollectionGroup::SimPatterns, SimVertices2DDependency);
+		SimVertices2DEnd = &ManagedArrayCollection->AddAttribute<int32>(SimVertices2DEndAttribute, ClothCollectionGroup::SimPatterns, SimVertices2DDependency);
+		SimFacesStart = &ManagedArrayCollection->AddAttribute<int32>(SimFacesStartAttribute, ClothCollectionGroup::SimPatterns, SimFacesDependency);
+		SimFacesEnd = &ManagedArrayCollection->AddAttribute<int32>(SimFacesEndAttribute, ClothCollectionGroup::SimPatterns, SimFacesDependency);
 
 		// Render Patterns Group
-		RenderVerticesStart = &ManagedArrayCollection->AddAttribute<int32>(RenderVerticesStartAttribute, RenderPatternsGroup, RenderVerticesDependency);
-		RenderVerticesEnd = &ManagedArrayCollection->AddAttribute<int32>(RenderVerticesEndAttribute, RenderPatternsGroup, RenderVerticesDependency);
-		RenderFacesStart = &ManagedArrayCollection->AddAttribute<int32>(RenderFacesStartAttribute, RenderPatternsGroup, RenderFacesDependency);
-		RenderFacesEnd = &ManagedArrayCollection->AddAttribute<int32>(RenderFacesEndAttribute, RenderPatternsGroup, RenderFacesDependency);
-		RenderMaterialPathName = &ManagedArrayCollection->AddAttribute<FString>(RenderMaterialPathNameAttribute, RenderPatternsGroup);
+		RenderVerticesStart = &ManagedArrayCollection->AddAttribute<int32>(RenderVerticesStartAttribute, ClothCollectionGroup::RenderPatterns, RenderVerticesDependency);
+		RenderVerticesEnd = &ManagedArrayCollection->AddAttribute<int32>(RenderVerticesEndAttribute, ClothCollectionGroup::RenderPatterns, RenderVerticesDependency);
+		RenderFacesStart = &ManagedArrayCollection->AddAttribute<int32>(RenderFacesStartAttribute, ClothCollectionGroup::RenderPatterns, RenderFacesDependency);
+		RenderFacesEnd = &ManagedArrayCollection->AddAttribute<int32>(RenderFacesEndAttribute, ClothCollectionGroup::RenderPatterns, RenderFacesDependency);
+		RenderMaterialPathName = &ManagedArrayCollection->AddAttribute<FString>(RenderMaterialPathNameAttribute, ClothCollectionGroup::RenderPatterns);
 
 		// Sim Faces Group
-		SimIndices2D = &ManagedArrayCollection->AddAttribute<FIntVector3>(SimIndices2DAttribute, SimFacesGroup, SimVertices2DDependency);
-		SimIndices3D = &ManagedArrayCollection->AddAttribute<FIntVector3>(SimIndices3DAttribute, SimFacesGroup, SimVertices3DDependency);
+		SimIndices2D = &ManagedArrayCollection->AddAttribute<FIntVector3>(SimIndices2DAttribute, ClothCollectionGroup::SimFaces, SimVertices2DDependency);
+		SimIndices3D = &ManagedArrayCollection->AddAttribute<FIntVector3>(SimIndices3DAttribute, ClothCollectionGroup::SimFaces, SimVertices3DDependency);
 
 		// Sim Vertices 2D Group
-		SimPosition2D = &ManagedArrayCollection->AddAttribute<FVector2f>(SimPosition2DAttribute, SimVertices2DGroup);
-		SimVertex3DLookup = &ManagedArrayCollection->AddAttribute<int32>(SimVertex3DLookupAttribute, SimVertices2DGroup, SimVertices3DDependency);
+		SimPosition2D = &ManagedArrayCollection->AddAttribute<FVector2f>(SimPosition2DAttribute, ClothCollectionGroup::SimVertices2D);
+		SimVertex3DLookup = &ManagedArrayCollection->AddAttribute<int32>(SimVertex3DLookupAttribute, ClothCollectionGroup::SimVertices2D, SimVertices3DDependency);
 
 		// Sim Vertices 3D Group
-		SimPosition3D = &ManagedArrayCollection->AddAttribute<FVector3f>(SimPosition3DAttribute, SimVertices3DGroup);
-		SimNormal = &ManagedArrayCollection->AddAttribute<FVector3f>(SimNormalAttribute, SimVertices3DGroup);
-		SimBoneIndices = &ManagedArrayCollection->AddAttribute<TArray<int32>>(SimBoneIndicesAttribute, SimVertices3DGroup);
-		SimBoneWeights = &ManagedArrayCollection->AddAttribute<TArray<float>>(SimBoneWeightsAttribute, SimVertices3DGroup);
-		TetherKinematicIndex = &ManagedArrayCollection->AddAttribute<TArray<int32>>(TetherKinematicIndexAttribute, SimVertices3DGroup, SimVertices3DDependency);
-		TetherReferenceLength = &ManagedArrayCollection->AddAttribute<TArray<float>>(TetherReferenceLengthAttribute, SimVertices3DGroup);
-		SimVertex2DLookup = &ManagedArrayCollection->AddAttribute<TArray<int32>>(SimVertex2DLookupAttribute, SimVertices3DGroup, SimVertices2DDependency);
-		SeamStitchLookup = &ManagedArrayCollection->AddAttribute<TArray<int32>>(SeamStitchLookupAttribute, SimVertices3DGroup, SeamStitchesDependency);
+		SimPosition3D = &ManagedArrayCollection->AddAttribute<FVector3f>(SimPosition3DAttribute, ClothCollectionGroup::SimVertices3D);
+		SimNormal = &ManagedArrayCollection->AddAttribute<FVector3f>(SimNormalAttribute, ClothCollectionGroup::SimVertices3D);
+		SimBoneIndices = &ManagedArrayCollection->AddAttribute<TArray<int32>>(SimBoneIndicesAttribute, ClothCollectionGroup::SimVertices3D);
+		SimBoneWeights = &ManagedArrayCollection->AddAttribute<TArray<float>>(SimBoneWeightsAttribute, ClothCollectionGroup::SimVertices3D);
+		TetherKinematicIndex = &ManagedArrayCollection->AddAttribute<TArray<int32>>(TetherKinematicIndexAttribute, ClothCollectionGroup::SimVertices3D, SimVertices3DDependency);
+		TetherReferenceLength = &ManagedArrayCollection->AddAttribute<TArray<float>>(TetherReferenceLengthAttribute, ClothCollectionGroup::SimVertices3D);
+		SimVertex2DLookup = &ManagedArrayCollection->AddAttribute<TArray<int32>>(SimVertex2DLookupAttribute, ClothCollectionGroup::SimVertices3D, SimVertices2DDependency);
+		SeamStitchLookup = &ManagedArrayCollection->AddAttribute<TArray<int32>>(SeamStitchLookupAttribute, ClothCollectionGroup::SimVertices3D, SeamStitchesDependency);
 
 		// Render Faces Group
-		RenderIndices = &ManagedArrayCollection->AddAttribute<FIntVector3>(RenderIndicesAttribute, RenderFacesGroup, RenderVerticesDependency);
+		RenderIndices = &ManagedArrayCollection->AddAttribute<FIntVector3>(RenderIndicesAttribute, ClothCollectionGroup::RenderFaces, RenderVerticesDependency);
 
 		// Render Vertices Group
-		RenderPosition = &ManagedArrayCollection->AddAttribute<FVector3f>(RenderPositionAttribute, RenderVerticesGroup);
-		RenderNormal = &ManagedArrayCollection->AddAttribute<FVector3f>(RenderNormalAttribute, RenderVerticesGroup);
-		RenderTangentU = &ManagedArrayCollection->AddAttribute<FVector3f>(RenderTangentUAttribute, RenderVerticesGroup);
-		RenderTangentV = &ManagedArrayCollection->AddAttribute<FVector3f>(RenderTangentVAttribute, RenderVerticesGroup);
-		RenderUVs = &ManagedArrayCollection->AddAttribute<TArray<FVector2f>>(RenderUVsAttribute, RenderVerticesGroup);
-		RenderColor = &ManagedArrayCollection->AddAttribute<FLinearColor>(RenderColorAttribute, RenderVerticesGroup);
-		RenderBoneIndices = &ManagedArrayCollection->AddAttribute<TArray<int32>>(RenderBoneIndicesAttribute, RenderVerticesGroup);
-		RenderBoneWeights = &ManagedArrayCollection->AddAttribute<TArray<float>>(RenderBoneWeightsAttribute, RenderVerticesGroup);
+		RenderPosition = &ManagedArrayCollection->AddAttribute<FVector3f>(RenderPositionAttribute, ClothCollectionGroup::RenderVertices);
+		RenderNormal = &ManagedArrayCollection->AddAttribute<FVector3f>(RenderNormalAttribute, ClothCollectionGroup::RenderVertices);
+		RenderTangentU = &ManagedArrayCollection->AddAttribute<FVector3f>(RenderTangentUAttribute, ClothCollectionGroup::RenderVertices);
+		RenderTangentV = &ManagedArrayCollection->AddAttribute<FVector3f>(RenderTangentVAttribute, ClothCollectionGroup::RenderVertices);
+		RenderUVs = &ManagedArrayCollection->AddAttribute<TArray<FVector2f>>(RenderUVsAttribute, ClothCollectionGroup::RenderVertices);
+		RenderColor = &ManagedArrayCollection->AddAttribute<FLinearColor>(RenderColorAttribute, ClothCollectionGroup::RenderVertices);
+		RenderBoneIndices = &ManagedArrayCollection->AddAttribute<TArray<int32>>(RenderBoneIndicesAttribute, ClothCollectionGroup::RenderVertices);
+		RenderBoneWeights = &ManagedArrayCollection->AddAttribute<TArray<float>>(RenderBoneWeightsAttribute, ClothCollectionGroup::RenderVertices);
 	}
 
 	int32 FClothCollection::GetNumElements(const FName& GroupName) const
diff --git a/Engine/Plugins/Experimental/ChaosClothAsset/Source/ChaosClothAsset/Private/ChaosClothAsset/ClothCollection.h b/Engine/Plugins/Experimental/ChaosClothAsset/Source/ChaosClothAsset/Private/ChaosClothAsset/ClothCollection.h
--- a/Engine/Plugins/Experimental/ChaosClothAsset/Source/ChaosClothAsset/Private/ChaosClothAsset/ClothCollection.h
+++ b/Engine/Plugins/Experimental/ChaosClothAsset/Source/ChaosClothAsset/Private/ChaosClothAsset/ClothCollection.h
@@ -21,17 +21,6 @@ namespace UE::Chaos::ClothAsset
 	class FClothCollection final
 	{
 	public:
-		static const FName LodsGroup;  // LOD information (only one LOD per collection)
-		static const FName SeamsGroup;  // Collection of seam stitches
-		static const FName SeamStitchesGroup;  // Contains pairs of stitched sim vertex indices
-		static const FName SimPatternsGroup;  // Contains sim pattern relationships to other groups
-		static const FName RenderPatternsGroup;  // Contains render pattern relationships to other groups
-		static const FName SimFacesGroup;  // Contains indices to sim vertices
-		static const FName SimVertices2DGroup;  // Contains 2D positions
-		static const FName SimVertices3DGroup;  // Contains 3D positions
-		static const FName RenderFacesGroup;  // Contains indices to render vertex
-		static const FName RenderVerticesGroup;  // Contains 3D render model
-
 		static constexpr int8 MaxNumBoneInfluences = 12; // This should be <= MAX_TOTAL_INFLUENCES defined in GPUSkinPublicDefs.h 
 		static constexpr int8 MaxNumTetherAttachments = 4; // This should be <= FClothTetherDataPrivate::MaxNumAttachments
 
diff --git a/Engine/Plugins/Experimental/ChaosClothAsset/Source/ChaosClothAsset/Private/ChaosClothAsset/CollectionClothFacade.cpp b/Engine/Plugins/Experimental/ChaosClothAsset/Source/ChaosClothAsset/Private/ChaosClothAsset/CollectionClothFacade.cpp
--- a/Engine/Plugins/Experimental/ChaosClothAsset/Source/ChaosClothAsset/Private/ChaosClothAsset/CollectionClothFacade.cpp
+++ b/Engine/Plugins/Experimental/ChaosClothAsset/Source/ChaosClothAsset/Private/ChaosClothAsset/CollectionClothFacade.cpp
@@ -2,6 +2,7 @@
 
 #include "ChaosClothAsset/CollectionClothFacade.h"
 #include "ChaosClothAsset/ClothCollection.h"
+#include "ChaosClothAsset/ClothCollectionGroup.h"
 #include "ChaosClothAsset/ClothGeometryTools.h"
 #include "Chaos/ChaosArchive.h"
 #include "UObject/UE5MainStreamObjectVersion.h"
@@ -20,7 +21,7 @@ namespace UE::Chaos::ClothAsset
 
 	bool FCollectionClothConstFacade::IsValid() const
 	{
-		return ClothCollection->IsValid() && ClothCollection->GetNumElements(FClothCollection::LodsGroup) == 1;
+		return ClothCollection->IsValid() && ClothCollection->GetNumElements(ClothCollectionGroup::Lods) == 1;
 	}
 
 	bool FCollectionClothConstFacade::HasValidData() const
@@ -110,7 +111,7 @@ namespace UE::Chaos::ClothAsset
 		for (const FName& WeightMapName : WeightMapNames)
 		{
 			ResultHash = HashCombineFast(ResultHash, GetTypeHash(WeightMapName));
-			ResultHash = HashCombineFast(ResultHash, ClothCollection->GetElementsTypeHash(ClothCollection->GetUserDefinedAttribute<float>(WeightMapName, FClothCollection::SimVertices3DGroup)));
+			ResultHash = HashCombineFast(ResultHash, ClothCollection->GetElementsTypeHash(ClothCollection->GetUserDefinedAttribute<float>(WeightMapName, ClothCollectionGroup::SimVertices3D)));
 		}
 		return ResultHash;
 	}
@@ -118,18 +119,18 @@ namespace UE::Chaos::ClothAsset
 	const FString& FCollectionClothConstFacade::GetPhysicsAssetPathName() const
 	{
 		static const FString EmptyString;
-		return ClothCollection->GetPhysicsAssetPathName() && ClothCollection->GetNumElements(FClothCollection::LodsGroup) > 0 ? (*ClothCollection->GetPhysicsAssetPathName())[0] : EmptyString;
+		return ClothCollection->GetPhysicsAssetPathName() && ClothCollection->GetNumElements(ClothCollectionGroup::Lods) > 0 ? (*ClothCollection->GetPhysicsAssetPathName())[0] : EmptyString;
 	}
 
 	const FString& FCollectionClothConstFacade::GetSkeletalMeshPathName() const
 	{
 		static const FString EmptyString;
-		return ClothCollection->GetSkeletalMeshPathName() && ClothCollection->GetNumElements(FClothCollection::LodsGroup) > 0 ? (*ClothCollection->GetSkeletalMeshPathName())[0] : EmptyString;
+		return ClothCollection->GetSkeletalMeshPathName() && ClothCollection->GetNumElements(ClothCollectionGroup::Lods) > 0 ? (*ClothCollection->GetSkeletalMeshPathName())[0] : EmptyString;
 	}
 
 	int32 FCollectionClothConstFacade::GetNumSimVertices2D() const
 	{
-		return ClothCollection->GetNumElements(FClothCollection::SimVertices2DGroup);
+		return ClothCollection->GetNumElements(ClothCollectionGroup::SimVertices2D);
 	}
 
 	TConstArrayView<FVector2f> FCollectionClothConstFacade::GetSimPosition2D() const
@@ -144,7 +145,7 @@ namespace UE::Chaos::ClothAsset
 
 	int32 FCollectionClothConstFacade::GetNumSimVertices3D() const
 	{
-		return ClothCollection->GetNumElements(FClothCollection::SimVertices3DGroup);
+		return ClothCollection->GetNumElements(ClothCollectionGroup::SimVertices3D);
 	}
 
 	TConstArrayView<FVector3f> FCollectionClothConstFacade::GetSimPosition3D() const
@@ -189,7 +190,7 @@ namespace UE::Chaos::ClothAsset
 
 	int32 FCollectionClothConstFacade::GetNumSimFaces() const
 	{
-		return ClothCollection->GetNumElements(FClothCollection::SimFacesGroup);
+		return ClothCollection->GetNumElements(ClothCollectionGroup::SimFaces);
 	}
 
 	TConstArrayView<FIntVector3> FCollectionClothConstFacade::GetSimIndices2D() const
@@ -204,7 +205,7 @@ namespace UE::Chaos::ClothAsset
 
 	int32 FCollectionClothConstFacade::GetNumSimPatterns() const
 	{
-		return ClothCollection->GetNumElements(FClothCollection::SimPatternsGroup);
+		return ClothCollection->GetNumElements(ClothCollectionGroup::SimPatterns);
 	}
 
 	FCollectionClothSimPatternConstFacade FCollectionClothConstFacade::GetSimPattern(int32 PatternIndex) const
@@ -224,7 +225,7 @@ namespace UE::Chaos::ClothAsset
 
 	int32 FCollectionClothConstFacade::GetNumRenderPatterns() const
 	{
-		return ClothCollection->GetNumElements(FClothCollection::RenderPatternsGroup);
+		return ClothCollection->GetNumElements(ClothCollectionGroup::RenderPatterns);
 	}
 
 	FCollectionClothRenderPatternConstFacade FCollectionClothConstFacade::GetRenderPattern(int32 PatternIndex) const
@@ -249,7 +250,7 @@ namespace UE::Chaos::ClothAsset
 
 	int32 FCollectionClothConstFacade::GetNumSeams() const
 	{
-		return ClothCollection->GetNumElements(FClothCollection::SeamsGroup);
+		return ClothCollection->GetNumElements(ClothCollectionGroup::Seams);
 	}
 
 	FCollectionClothSeamConstFacade FCollectionClothConstFacade::GetSeam(int32 SeamIndex) const
@@ -259,7 +260,7 @@ namespace UE::Chaos::ClothAsset
 
 	int32 FCollectionClothConstFacade::GetNumRenderVertices() const
 	{
-		return ClothCollection->GetNumElements(FClothCollection::RenderVerticesGroup);
+		return ClothCollection->GetNumElements(ClothCollectionGroup::RenderVertices);
 	}
 
 	TConstArrayView<FVector3f> FCollectionClothConstFacade::GetRenderPosition() const
@@ -304,7 +305,7 @@ namespace UE::Chaos::ClothAsset
 
 	int32 FCollectionClothConstFacade::GetNumRenderFaces() const
 	{
-		return ClothCollection->GetNumElements(FClothCollection::RenderFacesGroup);
+		return ClothCollection->GetNumElements(ClothCollectionGroup::RenderFaces);
 	}
 
 	TConstArrayView<FIntVector3> FCollectionClothConstFacade::GetRenderIndices() const
@@ -314,17 +315,17 @@ namespace UE::Chaos::ClothAsset
 
 	bool FCollectionClothConstFacade::HasWeightMap(const FName& Name) const
 	{
-		return ClothCollection->HasUserDefinedAttribute<float>(Name, FClothCollection::SimVertices3DGroup);
+		return ClothCollection->HasUserDefinedAttribute<float>(Name, ClothCollectionGroup::SimVertices3D);
 	}
 
 	TArray<FName> FCollectionClothConstFacade::GetWeightMapNames() const
 	{
-		return ClothCollection->GetUserDefinedAttributeNames<float>(FClothCollection::SimVertices3DGroup);
+		return ClothCollection->GetUserDefinedAttributeNames<float>(ClothCollectionGroup::SimVertices3D);
 	}
 
 	TConstArrayView<float> FCollectionClothConstFacade::GetWeightMap(const FName& Name) const
 	{
-		return ClothCollection->GetElements(ClothCollection->GetUserDefinedAttribute<float>(Name, FClothCollection::SimVertices3DGroup));
+		return ClothCollection->GetElements(ClothCollection->GetUserDefinedAttribute<float>(Name, ClothCollectionGroup::SimVertices3D));
 	}
 
 	void FCollectionClothConstFacade::BuildSimulationMesh(TArray<FVector3f>& Positions, TArray<FVector3f>& Normals, TArray<uint32>& Indices, TArray<FVector2f>& PatternsPositions, TArray<uint32>& PatternsIndices,
@@ -386,8 +387,8 @@ namespace UE::Chaos::ClothAsset
 	void FCollectionClothFacade::Reset()
 	{
 		// Hack to reset all Lods data
-		GetClothCollection()->SetNumElements(0, FClothCollection::LodsGroup);
-		GetClothCollection()->SetNumElements(1, FClothCollection::LodsGroup);
+		GetClothCollection()->SetNumElements(0, ClothCollectionGroup::Lods);
+		GetClothCollection()->SetNumElements(1, ClothCollectionGroup::Lods);
 		RemoveAllSimVertices3D();
 		SetNumSimPatterns(0);
 		SetNumRenderPatterns(0);
@@ -419,7 +420,7 @@ namespace UE::Chaos::ClothAsset
 		// Sim Vertices 3D Group
 		const int32 StartNumSimVertices3D = GetNumSimVertices3D();
 		const int32 OtherNumSimVertices3D = Other.GetNumSimVertices3D();
-		GetClothCollection()->SetNumElements(StartNumSimVertices3D + OtherNumSimVertices3D, FClothCollection::SimVertices3DGroup);
+		GetClothCollection()->SetNumElements(StartNumSimVertices3D + OtherNumSimVertices3D, ClothCollectionGroup::SimVertices3D);
 		FClothCollection::CopyArrayViewData(GetSimPosition3D().Right(OtherNumSimVertices3D), Other.GetSimPosition3D());
 		FClothCollection::CopyArrayViewData(GetSimNormal().Right(OtherNumSimVertices3D), Other.GetSimNormal());
 		FClothCollection::CopyArrayViewData(GetSimBoneIndices().Right(OtherNumSimVertices3D), Other.GetSimBoneIndices());
@@ -474,7 +475,7 @@ namespace UE::Chaos::ClothAsset
 
 	void FCollectionClothFacade::SetPhysicsAssetPathName(const FString& PathName)
 	{
-		if (ClothCollection->GetNumElements(FClothCollection::LodsGroup))
+		if (ClothCollection->GetNumElements(ClothCollectionGroup::Lods))
 		{
 			(*GetClothCollection()->GetPhysicsAssetPathName())[0] = PathName;
 		}
@@ -482,7 +483,7 @@ namespace UE::Chaos::ClothAsset
 	}
 	void FCollectionClothFacade::SetSkeletalMeshPathName(const FString& PathName)
 	{
-		if (ClothCollection->GetNumElements(FClothCollection::LodsGroup))
+		if (ClothCollection->GetNumElements(ClothCollectionGroup::Lods))
 		{
 			(*GetClothCollection()->GetSkeletalMeshPathName())[0] = PathName;
 		}
@@ -542,12 +543,12 @@ namespace UE::Chaos::ClothAsset
 	{
 		const int32 NumSimVertices = GetNumSimVertices3D();
 		check(InNumSimVertices >= NumSimVertices);
-		GetClothCollection()->SetNumElements(NumSimVertices - InNumSimVertices, FClothCollection::SimVertices3DGroup);
+		GetClothCollection()->SetNumElements(NumSimVertices - InNumSimVertices, ClothCollectionGroup::SimVertices3D);
 	}
 
 	void FCollectionClothFacade::RemoveSimVertices3D(const TArray<int32>& SortedDeletionList)
 	{
-		GetClothCollection()->RemoveElements(FClothCollection::SimVertices3DGroup, SortedDeletionList);
+		GetClothCollection()->RemoveElements(ClothCollectionGroup::SimVertices3D, SortedDeletionList);
 	}
 
 	void FCollectionClothFacade::CompactSimVertex2DLookup()
@@ -584,7 +585,7 @@ namespace UE::Chaos::ClothAsset
 			GetSimPattern(PatternIndex).Reset();
 		}
 
-		GetClothCollection()->SetNumElements(InNumPatterns, FClothCollection::SimPatternsGroup);
+		GetClothCollection()->SetNumElements(InNumPatterns, ClothCollectionGroup::SimPatterns);
 
 		for (int32 PatternIndex = NumPatterns; PatternIndex < InNumPatterns; ++PatternIndex)
 		{
@@ -611,7 +612,7 @@ namespace UE::Chaos::ClothAsset
 			GetSimPattern(PatternToRemove).Reset();
 		}
 
-		GetClothCollection()->RemoveElements(FClothCollection::SimPatternsGroup, SortedDeletionList);
+		GetClothCollection()->RemoveElements(ClothCollectionGroup::SimPatterns, SortedDeletionList);
 	}
 
 	void FCollectionClothFacade::SetNumRenderPatterns(int32 InNumPatterns)
@@ -623,7 +624,7 @@ namespace UE::Chaos::ClothAsset
 			GetRenderPattern(PatternIndex).Reset();
 		}
 
-		GetClothCollection()->SetNumElements(InNumPatterns, FClothCollection::RenderPatternsGroup);
+		GetClothCollection()->SetNumElements(InNumPatterns, ClothCollectionGroup::RenderPatterns);
 
 		for (int32 PatternIndex = NumPatterns; PatternIndex < InNumPatterns; ++PatternIndex)
 		{
@@ -650,7 +651,7 @@ namespace UE::Chaos::ClothAsset
 			GetRenderPattern(PatternToRemove).Reset();
 		}
 
-		GetClothCollection()->RemoveElements(FClothCollection::RenderPatternsGroup, SortedDeletionList);
+		GetClothCollection()->RemoveElements(ClothCollectionGroup::RenderPatterns, SortedDeletionList);
 	}
 
 	TArrayView<FString> FCollectionClothFacade::GetRenderMaterialPathName()
@@ -667,7 +668,7 @@ namespace UE::Chaos::ClothAsset
 			GetSeam(SeamIndex).Reset();
 		}
 
-		GetClothCollection()->SetNumElements(InNumSeams, FClothCollection::SeamsGroup);
+		GetClothCollection()->SetNumElements(InNumSeams, ClothCollectionGroup::Seams);
 
 		for (int32 SeamIndex = NumSeams; SeamIndex < InNumSeams; ++SeamIndex)
 		{
@@ -693,7 +694,7 @@ namespace UE::Chaos::ClothAsset
 		{
 			GetSeam(SeamToRemove).Reset();
 		}
-		GetClothCollection()->RemoveElements(FClothCollection::SeamsGroup, SortedDeletionList);
+		GetClothCollection()->RemoveElements(ClothCollectionGroup::Seams, SortedDeletionList);
 	}
 
 	//~ Render Vertices Group
@@ -745,23 +746,23 @@ namespace UE::Chaos::ClothAsset
 	void FCollectionClothFacade::AddWeightMap(const FName& Name)
 	{
 		check(IsValid());
-		GetClothCollection()->AddUserDefinedAttribute<float>(Name, FClothCollection::SimVertices3DGroup);
+		GetClothCollection()->AddUserDefinedAttribute<float>(Name, ClothCollectionGroup::SimVertices3D);
 	}
 
 	void FCollectionClothFacade::RemoveWeightMap(const FName& Name)
 	{
 		check(IsValid());
-		GetClothCollection()->RemoveUserDefinedAttribute(Name, FClothCollection::SimVertices3DGroup);
+		GetClothCollection()->RemoveUserDefinedAttribute(Name, ClothCollectionGroup::SimVertices3D);
 	}
 
 	TArrayView<float> FCollectionClothFacade::GetWeightMap(const FName& Name)
 	{
-		return GetClothCollection()->GetElements(GetClothCollection()->GetUserDefinedAttribute<float>(Name, FClothCollection::SimVertices3DGroup));
+		return GetClothCollection()->GetElements(GetClothCollection()->GetUserDefinedAttribute<float>(Name, ClothCollectionGroup::SimVertices3D));
 	}
 
 	void FCollectionClothFacade::SetDefaults()
 	{
-		GetClothCollection()->SetNumElements(1, FClothCollection::LodsGroup);
+		GetClothCollection()->SetNumElements(1, ClothCollectionGroup::Lods);
 	}
 
 } // End namespace UE::Chaos::ClothAsset
diff --git a/Engine/Plugins/Experimental/ChaosClothAsset/Source/ChaosClothAsset/Private/ChaosClothAsset/CollectionClothRenderPatternFacade.cpp b/Engine/Plugins/Experimental/ChaosClothAsset/Source/ChaosClothAsset/Private/ChaosClothAsset/CollectionClothRenderPatternFacade.cpp
--- a/Engine/Plugins/Experimental/ChaosClothAsset/Source/ChaosClothAsset/Private/ChaosClothAsset/CollectionClothRenderPatternFacade.cpp
+++ b/Engine/Plugins/Experimental/ChaosClothAsset/Source/ChaosClothAsset/Private/ChaosClothAsset/CollectionClothRenderPatternFacade.cpp
@@ -2,6 +2,7 @@
 
 #include "ChaosClothAsset/CollectionClothRenderPatternFacade.h"
 #include "ChaosClothAsset/ClothCollection.h"
+#include "ChaosClothAsset/ClothCollectionGroup.h"
 #include "ChaosClothAsset/ClothGeometryTools.h"
 
 namespace UE::Chaos::ClothAsset
@@ -9,7 +10,7 @@ namespace UE::Chaos::ClothAsset
 	const FString& FCollectionClothRenderPatternConstFacade::GetRenderMaterialPathName() const
 	{
 		static const FString EmptyString;
-		return ClothCollection->GetRenderMaterialPathName() && ClothCollection->GetNumElements(FClothCollection::RenderPatternsGroup) > GetElementIndex() ? 
+		return ClothCollection->GetRenderMaterialPathName() && ClothCollection->GetNumElements(ClothCollectionGroup::RenderPatterns) > GetElementIndex() ?
 			(*ClothCollection->GetRenderMaterialPathName())[GetElementIndex()] : EmptyString;
 	}
 
@@ -136,7 +137,7 @@ namespace UE::Chaos::ClothAsset
 		, PatternIndex(InPatternIndex)
 	{
 		check(ClothCollection->IsValid());
-		check(PatternIndex >= 0 && PatternIndex < ClothCollection->GetNumElements(FClothCollection::RenderPatternsGroup));
+		check(PatternIndex >= 0 && PatternIndex < ClothCollection->GetNumElements(ClothCollectionGroup::RenderPatterns));
 	}
 
 	void FCollectionClothRenderPatternFacade::Reset()
@@ -178,7 +179,7 @@ namespace UE::Chaos::ClothAsset
 	{
 		GetClothCollection()->SetNumElements(
 			NumRenderVertices,
-			FClothCollection::RenderVerticesGroup,
+			ClothCollectionGroup::RenderVertices,
 			GetClothCollection()->GetRenderVerticesStart(),
 			GetClothCollection()->GetRenderVerticesEnd(),
 			GetElementIndex());
@@ -195,7 +196,7 @@ namespace UE::Chaos::ClothAsset
 		}
 
 		GetClothCollection()->RemoveElements(
-			FClothCollection::RenderVerticesGroup,
+			ClothCollectionGroup::RenderVertices,
 			GlobalIndexSortedDeletionList,
 			GetClothCollection()->GetRenderVerticesStart(),
 			GetClothCollection()->GetRenderVerticesEnd(),
@@ -278,7 +279,7 @@ namespace UE::Chaos::ClothAsset
 	{
 		GetClothCollection()->SetNumElements(
 			NumRenderFaces,
-			FClothCollection::RenderFacesGroup,
+			ClothCollectionGroup::RenderFaces,
 			GetClothCollection()->GetRenderFacesStart(),
 			GetClothCollection()->GetRenderFacesEnd(),
 			GetElementIndex());
@@ -295,7 +296,7 @@ namespace UE::Chaos::ClothAsset
 		}
 
 		GetClothCollection()->RemoveElements(
-			FClothCollection::RenderFacesGroup,
+			ClothCollectionGroup::RenderFaces,
 			GlobalIndexSortedDeletionList,
 			GetClothCollection()->GetRenderFacesStart(),
 			GetClothCollection()->GetRenderFacesEnd(),
diff --git a/Engine/Plugins/Experimental/ChaosClothAsset/Source/ChaosClothAsset/Private/ChaosClothAsset/CollectionClothSeamFacade.cpp b/Engine/Plugins/Experimental/ChaosClothAsset/Source/ChaosClothAsset/Private/ChaosClothAsset/CollectionClothSeamFacade.cpp
--- a/Engine/Plugins/Experimental/ChaosClothAsset/Source/ChaosClothAsset/Private/ChaosClothAsset/CollectionClothSeamFacade.cpp
+++ b/Engine/Plugins/Experimental/ChaosClothAsset/Source/ChaosClothAsset/Private/ChaosClothAsset/CollectionClothSeamFacade.cpp
@@ -3,6 +3,7 @@
 #include "ChaosClothAsset/CollectionClothSeamFacade.h"
 #include "ChaosClothAsset/CollectionClothFacade.h"
 #include "ChaosClothAsset/ClothCollection.h"
+#include "ChaosClothAsset/ClothCollectionGroup.h"
 #include "ChaosClothAsset/ClothGeometryTools.h"
 #include "Containers/Map.h"
 
@@ -502,7 +503,7 @@ namespace UE::Chaos::ClothAsset
 		, SeamIndex(SeamIndex)
 	{
 		check(ClothCollection->IsValid());
-		check(SeamIndex >= 0 && SeamIndex < ClothCollection->GetNumElements(FClothCollection::SeamsGroup));
+		check(SeamIndex >= 0 && SeamIndex < ClothCollection->GetNumElements(ClothCollectionGroup::Seams));
 	}
 
 	void FCollectionClothSeamFacade::Reset()
@@ -572,7 +573,7 @@ namespace UE::Chaos::ClothAsset
 				//
 
 				// This resize will invalidate any existing SimVertex3DGroup ArrayViews!
-				GetClothCollection()->SetNumElements(OrigNumSimVertices3D + NumNewSimVertices3D, FClothCollection::SimVertices3DGroup);
+				GetClothCollection()->SetNumElements(OrigNumSimVertices3D + NumNewSimVertices3D, ClothCollectionGroup::SimVertices3D);
 
 				// Update 2D <--> 3D lookups
 				UpdateVertexLookupsAfterSplitting(SplittingGroups, Cloth.GetSimVertex3DLookupPrivate(), Cloth.GetSimVertex2DLookupPrivate());
@@ -727,7 +728,7 @@ namespace UE::Chaos::ClothAsset
 			}
 		}
 		VerticesToRemove.Sort();
-		GetClothCollection()->RemoveElements(FClothCollection::SimVertices3DGroup, VerticesToRemove);		
+		GetClothCollection()->RemoveElements(ClothCollectionGroup::SimVertices3D, VerticesToRemove);
 	}
 
 	void FCollectionClothSeamFacade::Initialize(const FCollectionClothSeamConstFacade& Other, const int32 SimVertex2DOffset, const int32 SimVertex3DOffset)
@@ -741,7 +742,7 @@ namespace UE::Chaos::ClothAsset
 	{
 		GetClothCollection()->SetNumElements(
 			NumStitches,
-			FClothCollection::SeamStitchesGroup,
+			ClothCollectionGroup::SeamStitches,
 			GetClothCollection()->GetSeamStitchStart(),
 			GetClothCollection()->GetSeamStitchEnd(),
 			GetElementIndex());
diff --git a/Engine/Plugins/Experimental/ChaosClothAsset/Source/ChaosClothAsset/Private/ChaosClothAsset/CollectionClothSelectionFacade.cpp b/Engine/Plugins/Experimental/ChaosClothAsset/Source/ChaosClothAsset/Private/ChaosClothAsset/CollectionClothSelectionFacade.cpp
--- a/Engine/Plugins/Experimental/ChaosClothAsset/Source/ChaosClothAsset/Private/ChaosClothAsset/CollectionClothSelectionFacade.cpp
+++ b/Engine/Plugins/Experimental/ChaosClothAsset/Source/ChaosClothAsset/Private/ChaosClothAsset/CollectionClothSelectionFacade.cpp
@@ -11,9 +11,6 @@ namespace UE::Chaos::ClothAsset
 	namespace Private
 	{
 		static const FName SelectionGroup(TEXT("Selection"));
-		static const FName NameAttribute(TEXT("Name"));
-		static const FName TypeAttribute(TEXT("Type"));
-		static const FName IndicesAttribute(TEXT("Indices"));
 	}
 
 	// --------------- FCollectionClothSelectionConstFacade -------------------------------
@@ -21,61 +18,50 @@ namespace UE::Chaos::ClothAsset
 	FCollectionClothSelectionConstFacade::FCollectionClothSelectionConstFacade(const TSharedRef<const FManagedArrayCollection>& InManagedArrayCollection) :
 		ManagedArrayCollection(ConstCastSharedRef<FManagedArrayCollection>(InManagedArrayCollection))
 	{
-		Name = ManagedArrayCollection->FindAttribute<FString>(Private::NameAttribute, Private::SelectionGroup);
-		Type = ManagedArrayCollection->FindAttribute<FString>(Private::TypeAttribute, Private::SelectionGroup);
-		Indices = ManagedArrayCollection->FindAttribute<TSet<int32>>(Private::IndicesAttribute, Private::SelectionGroup);
 	}
 
 	bool FCollectionClothSelectionConstFacade::IsValid() const
 	{
-		return Name && Type && Indices;
+		return ManagedArrayCollection->HasGroup(Private::SelectionGroup) &&
+			ManagedArrayCollection->NumElements(Private::SelectionGroup) == 1;
 	}
 
 	int32 FCollectionClothSelectionConstFacade::GetNumSelections() const
 	{
-		return ManagedArrayCollection->NumElements(Private::SelectionGroup);
+		return ManagedArrayCollection->NumAttributes(Private::SelectionGroup);
 	}
 
-	const TArrayView<const FString> FCollectionClothSelectionConstFacade::GetName() const
+	TArray<FName> FCollectionClothSelectionConstFacade::GetNames() const
 	{
-		if (Name)
-		{
-			return TArrayView<const FString>(Name->GetData(), Name->Num());
-		}
-		return TArrayView<const FString>();
+		return ManagedArrayCollection->AttributeNames(Private::SelectionGroup);
 	}
 
-	const TArrayView<const FString> FCollectionClothSelectionConstFacade::GetType() const
+	bool FCollectionClothSelectionConstFacade::HasSelection(const FName& Name) const
 	{
-		if (Type)
-		{
-			return TArrayView<const FString>(Type->GetData(), Type->Num());
-		}
-		return TArrayView<const FString>();
+		return IsValid() && ManagedArrayCollection->HasAttribute(Name, Private::SelectionGroup);
 	}
 
-	const TArrayView<const TSet<int32>> FCollectionClothSelectionConstFacade::GetIndices() const
+	FName FCollectionClothSelectionConstFacade::GetSelectionGroup(const FName& Name) const
 	{
-		if (Indices)
-		{
-			return TArrayView<const TSet<int32>>(Indices->GetData(), Indices->Num());
-		}
-		return TArrayView<const TSet<int32>>();
+		check(IsValid());
+		check(HasSelection(Name));
+		return ManagedArrayCollection->GetDependency(Name, Private::SelectionGroup);
 	}
 
-	int32 FCollectionClothSelectionConstFacade::FindSelection(const FString& SearchName) const
+	const TSet<int32>& FCollectionClothSelectionConstFacade::GetSelectionSet(const FName& Name) const
 	{
-		const TArrayView<const FString> Names = GetName();
-		for (int32 SelectionIndex = 0; SelectionIndex < Names.Num(); ++SelectionIndex)
-		{
-			if (Names[SelectionIndex] == SearchName)
-			{
-				return SelectionIndex;
-			}
-		}
-		return INDEX_NONE;
+		check(IsValid());
+		const TManagedArray<TSet<int32>>* const Selection = ManagedArrayCollection->FindAttributeTyped<TSet<int32>>(Name, Private::SelectionGroup);
+		check(Selection);
+		return *Selection->GetData();
 	}
 
+	const TSet<int32>* FCollectionClothSelectionConstFacade::FindSelectionSet(const FName& Name) const
+	{
+		check(IsValid());
+		const TManagedArray<TSet<int32>>* const Selection = ManagedArrayCollection->FindAttributeTyped<TSet<int32>>(Name, Private::SelectionGroup);
+		return Selection ? Selection->GetData() : nullptr;
+	}
 
 	// --------------- FCollectionClothSelectionFacade -------------------------------
 
@@ -89,82 +75,62 @@ namespace UE::Chaos::ClothAsset
 		{
 			ManagedArrayCollection->AddGroup(Private::SelectionGroup);
 		}
-
-		// AddAttribute will only add the attribute if it doesn't already exist, otherwise it will return the existing one
-		Name = &ManagedArrayCollection->AddAttribute<FString>(Private::NameAttribute, Private::SelectionGroup);
-		Type = &ManagedArrayCollection->AddAttribute<FString>(Private::TypeAttribute, Private::SelectionGroup);
-		Indices = &ManagedArrayCollection->AddAttribute<TSet<int32>>(Private::IndicesAttribute, Private::SelectionGroup);
-	}
-
-	TArrayView<FString> FCollectionClothSelectionFacade::GetName()
-	{
-		if (Name)
+		const int32 NumElements = ManagedArrayCollection->NumElements(Private::SelectionGroup);
+		if (NumElements > 1)
 		{
-			return TArrayView<FString>(Name->GetData(), Name->Num());
+			ManagedArrayCollection->RemoveElements(Private::SelectionGroup, NumElements - 1, 1);
 		}
-		return TArrayView<FString>();
-	}
-
-	TArrayView<FString> FCollectionClothSelectionFacade::GetType()
-	{
-		if (Type)
+		else if (NumElements == 0)
 		{
-			return TArrayView<FString>(Type->GetData(), Type->Num());
+			ManagedArrayCollection->AddElements(1, Private::SelectionGroup);
 		}
-		return TArrayView<FString>();
 	}
 
-	TArrayView<TSet<int32>> FCollectionClothSelectionFacade::GetIndices()
+	TSet<int32>& FCollectionClothSelectionFacade::GetSelectionSet(const FName& Name)
 	{
-		if (Indices)
-		{
-			return TArrayView<TSet<int32>>(Indices->GetData(), Indices->Num());
-		}
-		return TArrayView<TSet<int32>>();
+		check(IsValid());
+		TManagedArray<TSet<int32>>* const Selection = ManagedArrayCollection->FindAttributeTyped<TSet<int32>>(Name, Private::SelectionGroup);
+		check(Selection);
+		return *Selection->GetData();
 	}
 
-	int32 FCollectionClothSelectionFacade::FindOrAddSelection(const FString& InName)
+	TSet<int32>* FCollectionClothSelectionFacade::FindSelectionSet(const FName& Name)
 	{
-		if (!IsValid())
-		{
-			return INDEX_NONE;
-		}
-
-		// Check for existing selection with matching Name
-		const int32 FoundIndex = FindSelection(InName);
-		if (FoundIndex != INDEX_NONE)
-		{
-			return FoundIndex;
-		}
-
-		// No selection with the given name exists
-		checkf(ManagedArrayCollection->HasGroup(Private::SelectionGroup), TEXT("Expected SelectionGroup to exist in the managed array at this point"));
-
-		const int32 NewIndex = ManagedArrayCollection->AddElements(1, Private::SelectionGroup);
-		
-		GetName()[NewIndex] = InName;
-
-		return NewIndex;
+		check(IsValid());
+		TManagedArray<TSet<int32>>* const Selection = ManagedArrayCollection->FindAttributeTyped<TSet<int32>>(Name, Private::SelectionGroup);
+		return Selection ? Selection->GetData() : nullptr;
 	}
 
-	bool FCollectionClothSelectionFacade::SetSelection(int32 SelectionIndex, const FString& InType, const TSet<int32>& InSelectedIndices)
+	TSet<int32>& FCollectionClothSelectionFacade::FindOrAddSelectionSet(const FName& Name, const FName& GroupName)
 	{
-		if (!IsValid())
+		check(IsValid());
+		ensure(GroupName != NAME_None && Name != NAME_None);
+		constexpr bool bAllowCircularDependency = false;
+
+		TManagedArray<TSet<int32>>* Selection = ManagedArrayCollection->FindAttributeTyped<TSet<int32>>(Name, Private::SelectionGroup);
+		if (Selection)
 		{
-			return false;
-		}
+			check(Selection->Num() == 1);  // This should always be the case if the facade is valid
 
-		if (SelectionIndex < 0 || SelectionIndex >= GetNumSelections())
+			// Recycle the existing selection, with the new group if needed
+			if (ManagedArrayCollection->GetDependency(Name, Private::SelectionGroup) != GroupName)
+			{
+				ManagedArrayCollection->SetDependency(Name, Private::SelectionGroup, GroupName, bAllowCircularDependency);
+				(*Selection)[0].Reset();  // No point in keeping unrelated selection indices since the group has changed better clear everything
+			}
+		}
+		else
 		{
-			return false;
+			// Create a new selection
+			constexpr bool bSaved = true;
+			const FManagedArrayCollection::FConstructionParameters GroupDependency(GroupName, bSaved, bAllowCircularDependency);
+			Selection = &ManagedArrayCollection->AddAttribute<TSet<int32>>(Name, Private::SelectionGroup, GroupDependency);
+
+			check(Selection->Num() == 1);  // This should always be the case if the facade is valid
 		}
 
-		GetType()[SelectionIndex] = InType;
-		GetIndices()[SelectionIndex] = InSelectedIndices;
-		return true;
+		return (*Selection)[0];
 	}
-
-
 }	// namespace  UE::Chaos::ClothAsset
 
 #undef LOCTEXT_NAMESPACE
diff --git a/Engine/Plugins/Experimental/ChaosClothAsset/Source/ChaosClothAsset/Private/ChaosClothAsset/CollectionClothSimPatternFacade.cpp b/Engine/Plugins/Experimental/ChaosClothAsset/Source/ChaosClothAsset/Private/ChaosClothAsset/CollectionClothSimPatternFacade.cpp
--- a/Engine/Plugins/Experimental/ChaosClothAsset/Source/ChaosClothAsset/Private/ChaosClothAsset/CollectionClothSimPatternFacade.cpp
+++ b/Engine/Plugins/Experimental/ChaosClothAsset/Source/ChaosClothAsset/Private/ChaosClothAsset/CollectionClothSimPatternFacade.cpp
@@ -3,6 +3,7 @@
 #include "ChaosClothAsset/CollectionClothSimPatternFacade.h"
 #include "ChaosClothAsset/CollectionClothFacade.h"
 #include "ChaosClothAsset/ClothCollection.h"
+#include "ChaosClothAsset/ClothCollectionGroup.h"
 #include "ChaosClothAsset/ClothGeometryTools.h"
 
 namespace UE::Chaos::ClothAsset
@@ -85,7 +86,7 @@ namespace UE::Chaos::ClothAsset
 		, PatternIndex(InPatternIndex)
 	{
 		check(ClothCollection->IsValid());
-		check(PatternIndex >= 0 && PatternIndex < ClothCollection->GetNumElements(FClothCollection::SimPatternsGroup));
+		check(PatternIndex >= 0 && PatternIndex < ClothCollection->GetNumElements(ClothCollectionGroup::SimPatterns));
 	}
 
 	void FCollectionClothSimPatternFacade::Reset()
@@ -192,10 +193,10 @@ namespace UE::Chaos::ClothAsset
 
 	FIntVector2 FCollectionClothSimPatternFacade::AppendSimVertices(int32 NumSimVertices)
 	{
-		const FIntVector2 OrigNumVertices(GetNumSimVertices2D(), ClothCollection->GetNumElements(FClothCollection::SimVertices3DGroup));
+		const FIntVector2 OrigNumVertices(GetNumSimVertices2D(), ClothCollection->GetNumElements(ClothCollectionGroup::SimVertices3D));
 
 		// Resize arrays
-		GetClothCollection()->SetNumElements(OrigNumVertices[1] + NumSimVertices, FClothCollection::SimVertices3DGroup);
+		GetClothCollection()->SetNumElements(OrigNumVertices[1] + NumSimVertices, ClothCollectionGroup::SimVertices3D);
 		SetNumSimVertices2D(OrigNumVertices[0] + NumSimVertices);
 
 		// Set lookups to each other
@@ -234,7 +235,7 @@ namespace UE::Chaos::ClothAsset
 		}
 
 		GetClothCollection()->RemoveElements(
-			FClothCollection::SimVertices2DGroup,
+			ClothCollectionGroup::SimVertices2D,
 			GlobalIndexSortedDeletionList,
 			GetClothCollection()->GetSimVertices2DStart(),
 			GetClothCollection()->GetSimVertices2DEnd(),
@@ -245,7 +246,7 @@ namespace UE::Chaos::ClothAsset
 	{
 		GetClothCollection()->SetNumElements(
 			NumSimFaces,
-			FClothCollection::SimFacesGroup,
+			ClothCollectionGroup::SimFaces,
 			GetClothCollection()->GetSimFacesStart(),
 			GetClothCollection()->GetSimFacesEnd(),
 			GetElementIndex());
@@ -280,7 +281,7 @@ namespace UE::Chaos::ClothAsset
 		}
 
 		GetClothCollection()->RemoveElements(
-			FClothCollection::SimFacesGroup,
+			ClothCollectionGroup::SimFaces,
 			GlobalIndexSortedDeletionList,
 			GetClothCollection()->GetSimFacesStart(),
 			GetClothCollection()->GetSimFacesEnd(),
@@ -306,7 +307,7 @@ namespace UE::Chaos::ClothAsset
 	{
 		GetClothCollection()->SetNumElements(
 			NumSimVertices,
-			FClothCollection::SimVertices2DGroup,
+			ClothCollectionGroup::SimVertices2D,
 			GetClothCollection()->GetSimVertices2DStart(),
 			GetClothCollection()->GetSimVertices2DEnd(),
 			GetElementIndex());
diff --git a/Engine/Plugins/Experimental/ChaosClothAsset/Source/ChaosClothAsset/Public/ChaosClothAsset/ClothCollectionGroup.h b/Engine/Plugins/Experimental/ChaosClothAsset/Source/ChaosClothAsset/Public/ChaosClothAsset/ClothCollectionGroup.h
new file mode 100644
--- /dev/null
+++ b/Engine/Plugins/Experimental/ChaosClothAsset/Source/ChaosClothAsset/Public/ChaosClothAsset/ClothCollectionGroup.h
@@ -0,0 +1,35 @@
+// Copyright Epic Games, Inc. All Rights Reserved.
+
+#pragma once
+
+#include "UObject/NameTypes.h"
+
+namespace UE::Chaos::ClothAsset
+{
+	/**
+	 * Cloth collection group names.
+	 */
+	namespace ClothCollectionGroup
+	{
+		/** LOD information (only one LOD per collection). */
+		inline const FName Lods = FName(TEXT("Lods"));
+		/** Collection of seam stitches. */
+		inline const FName Seams = FName(TEXT("Seams"));
+		/** Contains pairs of stitched sim vertex indices. */
+		inline const FName SeamStitches = FName(TEXT("SeamStitches"));
+		/** Contains sim pattern relationships to other groups. */
+		inline const FName SimPatterns = FName(TEXT("SimPatterns"));
+		/** Contains render pattern relationships to other groups. */
+		inline const FName RenderPatterns = FName(TEXT("RenderPatterns"));
+		/** Contains indices to sim vertices. */
+		inline const FName SimFaces = FName(TEXT("SimFaces"));
+		/** Contains 2D positions. */
+		inline const FName SimVertices2D = FName(TEXT("SimVertices2D"));
+		/** Contains 3D positions. */
+		inline const FName SimVertices3D = FName(TEXT("SimVertices3D"));
+		/** Contains indices to render vertex. */
+		inline const FName RenderFaces = FName(TEXT("RenderFaces"));
+		/** Contains 3D render model. */
+		inline const FName RenderVertices = FName(TEXT("RenderVertices"));
+	}
+}  // End namespace UE::Chaos::ClothAsset
diff --git a/Engine/Plugins/Experimental/ChaosClothAsset/Source/ChaosClothAsset/Public/ChaosClothAsset/CollectionClothSelectionFacade.h b/Engine/Plugins/Experimental/ChaosClothAsset/Source/ChaosClothAsset/Public/ChaosClothAsset/CollectionClothSelectionFacade.h
--- a/Engine/Plugins/Experimental/ChaosClothAsset/Source/ChaosClothAsset/Public/ChaosClothAsset/CollectionClothSelectionFacade.h
+++ b/Engine/Plugins/Experimental/ChaosClothAsset/Source/ChaosClothAsset/Public/ChaosClothAsset/CollectionClothSelectionFacade.h
@@ -13,10 +13,10 @@ namespace UE::Chaos::ClothAsset
 	 * Cloth Asset collection Selection facade class. Each Selection consists of a set of integer indices, a name, and a string representing the type of element being indexed.
 	 * Const access (read only) version.
 	 */
-	class CHAOSCLOTHASSET_API FCollectionClothSelectionConstFacade
+	class FCollectionClothSelectionConstFacade
 	{
 	public:
-		explicit FCollectionClothSelectionConstFacade(const TSharedRef<const FManagedArrayCollection>& ManagedArrayCollection);
+		CHAOSCLOTHASSET_API explicit FCollectionClothSelectionConstFacade(const TSharedRef<const FManagedArrayCollection>& ManagedArrayCollection);
 
 		FCollectionClothSelectionConstFacade() = delete;
 		FCollectionClothSelectionConstFacade(const FCollectionClothSelectionConstFacade&) = delete;
@@ -26,38 +26,38 @@ namespace UE::Chaos::ClothAsset
 		virtual ~FCollectionClothSelectionConstFacade() = default;
 
 		/** Return whether the facade is defined on the collection. */
-		bool IsValid() const;
+		CHAOSCLOTHASSET_API bool IsValid() const;
 
-		/** Return the total number of selections in the collection */
-		int32 GetNumSelections() const;
+		/** Return the number of selections in the collection. */
+		CHAOSCLOTHASSET_API int32 GetNumSelections() const;
 
-		const TArrayView<const FString> GetName() const;
-		const TArrayView<const FString> GetType() const;
-		const TArrayView<const TSet<int32>> GetIndices() const;
+		/** Return an array of all the selections' name in the collection. */
+		CHAOSCLOTHASSET_API TArray<FName> GetNames() const;
 
-		/** Find a selection with the given name. Return INDEX_NONE if no such selection exists. */
-		int32 FindSelection(const FString& Name) const;
+		/** Return whether a selection with the given name cyrrently exists in the collection. */
+		CHAOSCLOTHASSET_API bool HasSelection(const FName& Name) const;
 
-	protected:
+		/** Return the selection group dependency. The selection must exist to call this function. */
+		CHAOSCLOTHASSET_API FName GetSelectionGroup(const FName& Name) const;
 
-		TSharedRef<FManagedArrayCollection> ManagedArrayCollection;
+		/** Get the selection set for the given selection name. The selection must exist to call this function. */
+		CHAOSCLOTHASSET_API const TSet<int32>& GetSelectionSet(const FName& Name) const;
 
-		// Schema (non-const since these will be used by the subclass as well)
-		TManagedArray<FString>* Name;
-		TManagedArray<FString>* Type;
-		TManagedArray<TSet<int32>>* Indices;
-	};
+		/** Find a selection with the given name, or nullptr if no such selection exists. */
+		CHAOSCLOTHASSET_API const TSet<int32>* FindSelectionSet(const FName& Name) const;
 
+	protected:
+		TSharedRef<FManagedArrayCollection> ManagedArrayCollection;
+	};
 
 	/**
 	 * Cloth Asset collection Selection facade class. Each Selection consists of a set of integer indices, a name, and a string representing the type of element being indexed.
 	 * Non-const access (read/write) version.
 	 */
-	class CHAOSCLOTHASSET_API FCollectionClothSelectionFacade final : public FCollectionClothSelectionConstFacade
+	class FCollectionClothSelectionFacade final : public FCollectionClothSelectionConstFacade
 	{
 	public:
-
-		explicit FCollectionClothSelectionFacade(const TSharedRef<const FManagedArrayCollection>& ManagedArrayCollection);
+		CHAOSCLOTHASSET_API explicit FCollectionClothSelectionFacade(const TSharedRef<const FManagedArrayCollection>& ManagedArrayCollection);
 
 		FCollectionClothSelectionFacade() = delete;
 		FCollectionClothSelectionFacade(const FCollectionClothSelectionFacade&) = delete;
@@ -66,23 +66,20 @@ namespace UE::Chaos::ClothAsset
 		FCollectionClothSelectionFacade& operator=(FCollectionClothSelectionFacade&&) = default;
 		virtual ~FCollectionClothSelectionFacade() override = default;
 
-		/** Add the Selection attributes to the underlying collection */
-		void DefineSchema();
+		/** Add the Selection attributes to the underlying collection. */
+		CHAOSCLOTHASSET_API void DefineSchema();
 
-		TArrayView<FString> GetName();
-		TArrayView<FString> GetType();
-		TArrayView<TSet<int32>> GetIndices();
+		/** Get the selection set for the given selection name. The selection must exist to call this function. */
+		CHAOSCLOTHASSET_API TSet<int32>& GetSelectionSet(const FName& Name);
 
-		/** 
-		* Create a new Selection in the collection with the given name and return its index. If a Selection with the name already exists, return the index of the existing selection. 
-		* 
-		* @param Name The name of the new selection to create if it doesn't already exist
-		* @return The previous number of Selections if we added a new Selection, or the index of the Selection with the given Name if it already existed, or INDEX_NONE if the schema for this facade object is not valid.
-		*/
-		int32 FindOrAddSelection(const FString& Name);
+		/** Find a selection with the given name, or nullptr if no such selection exists. */
+		CHAOSCLOTHASSET_API TSet<int32>* FindSelectionSet(const FName& Name);
 
-		/** Set the Type and Indices for the given Selection. Returns false if the Facade is invalid or if the SelectionIndex is out of range */
-		bool SetSelection(int32 SelectionIndex, const FString& InType, const TSet<int32>& InSelectedIndices);
+		/**
+		 * Find, or add if it doesn't already exist, a selection for the specified group with the given name.
+		 * If the group doesn't already exists, this function will create it.
+		 * If the selection already exists, but depends on a different group, the old selection will be deleted and a new one recreated with the new group dependency.
+		 */
+		CHAOSCLOTHASSET_API TSet<int32>& FindOrAddSelectionSet(const FName& Name, const FName& GroupName);
 	};
-
 }
diff --git a/Engine/Plugins/Experimental/ChaosClothAssetEditor/Source/ChaosClothAssetDataflowNodes/Private/ChaosClothAsset/ChaosClothAssetDataflowNodesModule.cpp b/Engine/Plugins/Experimental/ChaosClothAssetEditor/Source/ChaosClothAssetDataflowNodes/Private/ChaosClothAsset/ChaosClothAssetDataflowNodesModule.cpp
--- a/Engine/Plugins/Experimental/ChaosClothAssetEditor/Source/ChaosClothAssetDataflowNodes/Private/ChaosClothAsset/ChaosClothAssetDataflowNodesModule.cpp
+++ b/Engine/Plugins/Experimental/ChaosClothAssetEditor/Source/ChaosClothAssetDataflowNodes/Private/ChaosClothAsset/ChaosClothAssetDataflowNodesModule.cpp
@@ -5,6 +5,7 @@
 #include "ChaosClothAsset/BindToRootBoneNode.h"
 #include "ChaosClothAsset/ColorScheme.h"
 #include "ChaosClothAsset/CopySimulationToRenderMeshNode.h"
+#include "ChaosClothAsset/SelectionGroupCustomization.h"
 #include "ChaosClothAsset/DatasmithImportNode.h"
 #include "ChaosClothAsset/DeleteElementNode.h"
 #include "ChaosClothAsset/ImportNode.h"
@@ -107,9 +108,10 @@ void FChaosClothAssetDataflowNodesModule::StartupModule()
 	// Register type customizations
 	if (FPropertyEditorModule* const PropertyModule = FModuleManager::GetModulePtr<FPropertyEditorModule>("PropertyEditor"))
 	{
-		PropertyModule->RegisterCustomPropertyTypeLayout("ChaosClothAssetWeightedValue", FOnGetPropertyTypeCustomizationInstance::CreateStatic(&FWeightedValueCustomization::MakeInstance));
-		PropertyModule->RegisterCustomPropertyTypeLayout("ChaosClothAssetWeightedValueNonAnimatable", FOnGetPropertyTypeCustomizationInstance::CreateStatic(&FWeightedValueCustomization::MakeInstance));
-		PropertyModule->RegisterCustomPropertyTypeLayout("ChaosClothAssetWeightedValueNonAnimatableNoLowHighRange", FOnGetPropertyTypeCustomizationInstance::CreateStatic(&FWeightedValueCustomization::MakeInstance));
+		PropertyModule->RegisterCustomPropertyTypeLayout(FChaosClothAssetWeightedValue::StaticStruct()->GetFName(), FOnGetPropertyTypeCustomizationInstance::CreateStatic(&FWeightedValueCustomization::MakeInstance));
+		PropertyModule->RegisterCustomPropertyTypeLayout(FChaosClothAssetWeightedValueNonAnimatable::StaticStruct()->GetFName(), FOnGetPropertyTypeCustomizationInstance::CreateStatic(&FWeightedValueCustomization::MakeInstance));
+		PropertyModule->RegisterCustomPropertyTypeLayout(FChaosClothAssetWeightedValueNonAnimatableNoLowHighRange::StaticStruct()->GetFName(), FOnGetPropertyTypeCustomizationInstance::CreateStatic(&FWeightedValueCustomization::MakeInstance));
+		PropertyModule->RegisterCustomPropertyTypeLayout(FChaosClothAssetNodeSelectionGroup::StaticStruct()->GetFName(), FOnGetPropertyTypeCustomizationInstance::CreateStatic(&FSelectionGroupCustomization::MakeInstance));
 	}
 
 	// Register modular features
@@ -121,9 +123,10 @@ void FChaosClothAssetDataflowNodesModule::ShutdownModule()
 	// Unregister type customizations
 	if (FPropertyEditorModule* const PropertyModule = FModuleManager::GetModulePtr<FPropertyEditorModule>("PropertyEditor"))
 	{
-		PropertyModule->UnregisterCustomPropertyTypeLayout("ChaosClothAssetWeightedValue");
-		PropertyModule->UnregisterCustomPropertyTypeLayout("ChaosClothAssetWeightedValueNonAnimatable");
-		PropertyModule->UnregisterCustomPropertyTypeLayout("ChaosClothAssetWeightedValueNonAnimatableNoLowHighRange");
+		PropertyModule->UnregisterCustomPropertyTypeLayout(FChaosClothAssetWeightedValue::StaticStruct()->GetFName());
+		PropertyModule->UnregisterCustomPropertyTypeLayout(FChaosClothAssetWeightedValueNonAnimatable::StaticStruct()->GetFName());
+		PropertyModule->UnregisterCustomPropertyTypeLayout(FChaosClothAssetWeightedValueNonAnimatableNoLowHighRange::StaticStruct()->GetFName());
+		PropertyModule->UnregisterCustomPropertyTypeLayout(FChaosClothAssetNodeSelectionGroup::StaticStruct()->GetFName());
 	}
 
 	// Unregister modular features
diff --git a/Engine/Plugins/Experimental/ChaosClothAssetEditor/Source/ChaosClothAssetDataflowNodes/Private/ChaosClothAsset/ClothDataflowTools.cpp b/Engine/Plugins/Experimental/ChaosClothAssetEditor/Source/ChaosClothAssetDataflowNodes/Private/ChaosClothAsset/ClothDataflowTools.cpp
--- a/Engine/Plugins/Experimental/ChaosClothAssetEditor/Source/ChaosClothAssetDataflowNodes/Private/ChaosClothAsset/ClothDataflowTools.cpp
+++ b/Engine/Plugins/Experimental/ChaosClothAssetEditor/Source/ChaosClothAssetDataflowNodes/Private/ChaosClothAsset/ClothDataflowTools.cpp
@@ -450,4 +450,11 @@ namespace UE::Chaos::ClothAsset
 
 		UE_LOG(LogChaosClothAssetDataflowNodes, Display, TEXT("%s"), *Text.ToString());
 	}
+
+	void FClothDataflowTools::MakeCollectionName(FString& InOutString)
+	{
+		InOutString = SlugStringForValidName(InOutString, TEXT("_")).Replace(TEXT("\\"), TEXT("_"));
+		bool bCharsWereRemoved;
+		do { InOutString.TrimCharInline(TEXT('_'), &bCharsWereRemoved); } while (bCharsWereRemoved);
+	}
 }  // End namespace UE::Chaos::ClothAsset
diff --git a/Engine/Plugins/Experimental/ChaosClothAssetEditor/Source/ChaosClothAssetDataflowNodes/Private/ChaosClothAsset/ClothDataflowTools.h b/Engine/Plugins/Experimental/ChaosClothAssetEditor/Source/ChaosClothAssetDataflowNodes/Private/ChaosClothAsset/ClothDataflowTools.h
--- a/Engine/Plugins/Experimental/ChaosClothAssetEditor/Source/ChaosClothAssetDataflowNodes/Private/ChaosClothAsset/ClothDataflowTools.h
+++ b/Engine/Plugins/Experimental/ChaosClothAssetEditor/Source/ChaosClothAssetDataflowNodes/Private/ChaosClothAsset/ClothDataflowTools.h
@@ -22,6 +22,14 @@ namespace UE::Chaos::ClothAsset
 		static void AddSimPatternsFromSkeletalMeshSection(const TSharedRef<FManagedArrayCollection>& ClothCollection, const FSkeletalMeshLODModel& SkeletalMeshModel, const int32 SectionIndex, const int32 UVChannelIndex, const FVector2f& UVScale);
 
 		static void LogAndToastWarning(const FDataflowNode& DataflowNode, const FText& Headline, const FText& Details);
+
+		/**
+		 * Turn a string into a valid collection group or attribute name.
+		 * The resulting name won't contains spaces and any other special characters as listed in
+		 * INVALID_OBJECTNAME_CHARACTERS (currently "',/.:|&!~\n\r\t@#(){}[]=;^%$`).
+		 * It will also have all leading underscore removed, as these names are reserved for internal use.
+		 */
+		static void MakeCollectionName(FString& InOutString);
 	};
 }  // End namespace UE::Chaos::ClothAsset
 
diff --git a/Engine/Plugins/Experimental/ChaosClothAssetEditor/Source/ChaosClothAssetDataflowNodes/Private/ChaosClothAsset/SelectionGroupCustomization.cpp b/Engine/Plugins/Experimental/ChaosClothAssetEditor/Source/ChaosClothAssetDataflowNodes/Private/ChaosClothAsset/SelectionGroupCustomization.cpp
new file mode 100644
--- /dev/null
+++ b/Engine/Plugins/Experimental/ChaosClothAssetEditor/Source/ChaosClothAssetDataflowNodes/Private/ChaosClothAsset/SelectionGroupCustomization.cpp
@@ -0,0 +1,155 @@
+// Copyright Epic Games, Inc. All Rights Reserved.
+
+#include "ChaosClothAsset/SelectionGroupCustomization.h"
+#include "ChaosClothAsset/SelectionNode.h"
+#include "ChaosClothAsset/ClothDataflowTools.h"  // For MakeCollectionName
+#include "Widgets/Input/SEditableTextBox.h"
+#include "Widgets/Input/SComboButton.h"
+#include "Widgets/Views/SListView.h"
+#include "Widgets/Views/STableRow.h"
+#include "DetailLayoutBuilder.h"
+#include "DetailWidgetRow.h"
+#include "Editor.h"
+
+#define LOCTEXT_NAMESPACE "DataflowNodeSelectionGroupCustomization"
+
+namespace UE::Chaos::ClothAsset
+{
+	TSharedRef<IPropertyTypeCustomization> FSelectionGroupCustomization::MakeInstance()
+	{
+		return MakeShareable(new FSelectionGroupCustomization);
+	}
+
+	void FSelectionGroupCustomization::CustomizeHeader(TSharedRef<IPropertyHandle> PropertyHandle, FDetailWidgetRow& HeaderRow, IPropertyTypeCustomizationUtils& CustomizationUtils)
+	{
+		uint32 NumChildren;
+		const FPropertyAccess::Result Result = PropertyHandle->GetNumChildren(NumChildren);
+
+		ChildPropertyHandle = (Result == FPropertyAccess::Success && NumChildren) ? PropertyHandle->GetChildHandle(0) : TSharedPtr<IPropertyHandle>();
+
+		GroupNames.Reset();
+
+		HeaderRow
+			.NameContent()
+			[
+				PropertyHandle->CreatePropertyNameWidget(PropertyHandle->GetPropertyDisplayName())
+			]
+			.ValueContent()
+			.MinDesiredWidth(250)
+			.MaxDesiredWidth(350.f)
+			[
+				SNew(SHorizontalBox)
+				+ SHorizontalBox::Slot()
+				.MaxWidth(160.f)
+				[
+					SAssignNew(ComboButton, SComboButton)
+					.ContentPadding(0)
+					.OnGetMenuContent(this, &FSelectionGroupCustomization::OnGetMenuContent)
+					.ButtonContent()
+					[
+						SNew(SEditableTextBox)
+						.Text_Raw(this, &FSelectionGroupCustomization::GetText)
+						.Font(FAppStyle::GetFontStyle(TEXT("PropertyWindow.NormalFont")))
+						.OnTextCommitted(this, &FSelectionGroupCustomization::OnTextCommitted)
+					]
+				]
+			];
+	}
+
+	FText FSelectionGroupCustomization::GetText() const
+	{
+		FText Text;
+		ChildPropertyHandle->GetValueAsFormattedText(Text);
+		return Text;
+	}
+
+	void FSelectionGroupCustomization::OnTextCommitted(const FText& NewText, ETextCommit::Type CommitInfo)
+	{
+		FText CurrentText;
+		ChildPropertyHandle->GetValueAsFormattedText(CurrentText);
+
+		if (!NewText.ToString().Equals(CurrentText.ToString(), ESearchCase::CaseSensitive))
+		{
+			FString String = NewText.ToString();
+			FClothDataflowTools::MakeCollectionName(String);
+			ChildPropertyHandle->SetValueFromFormattedString(String);
+		}
+	}
+
+	void FSelectionGroupCustomization::OnSelectionChanged(TSharedPtr<FText> ItemSelected, ESelectInfo::Type /*SelectInfo*/)
+	{
+		// Set the child property's value
+		if (ItemSelected)
+		{
+			FText CurrentText;
+			ChildPropertyHandle->GetValueAsFormattedText(CurrentText);
+
+			if (!ItemSelected->EqualTo(CurrentText))
+			{
+				ChildPropertyHandle->SetValueFromFormattedString(ItemSelected->ToString());
+			}
+
+			ComboButton->SetIsOpen(false);
+		}
+	}
+
+	TSharedRef<ITableRow> FSelectionGroupCustomization::MakeCategoryViewWidget(TSharedPtr<FText> Item, const TSharedRef< STableViewBase >& OwnerTable)
+	{
+		return SNew(STableRow<TSharedPtr<FString>>, OwnerTable)
+			[
+				SNew(STextBlock).Text(*Item.Get())
+			];
+	}
+
+	TSharedRef<SWidget> FSelectionGroupCustomization::OnGetMenuContent()
+	{
+		// Find all group names in the parent selection node's collection
+		GroupNames.Reset();
+
+		if (const FChaosClothAssetSelectionNode* const SelectionNode = GetOwnerStruct<FChaosClothAssetSelectionNode>())
+		{
+			const TArray<FName> CollectionGroupNames = SelectionNode->GetCachedCollectionGroupNames();
+
+			for (const FName& CollectionGroupName : CollectionGroupNames)
+			{
+				GroupNames.Add(MakeShareable(new FText(FText::FromName(CollectionGroupName))));
+			}
+		}
+
+		return SNew(SVerticalBox)
+			+ SVerticalBox::Slot()
+			.AutoHeight()
+			.MaxHeight(400.0f)
+			[
+				SNew(SListView<TSharedPtr<FText>>)
+				.ListItemsSource(&GroupNames)
+				.SelectionMode(ESelectionMode::Single)
+				.OnGenerateRow(this, &FSelectionGroupCustomization::MakeCategoryViewWidget)
+				.OnSelectionChanged(this, &FSelectionGroupCustomization::OnSelectionChanged)
+			];
+	}
+
+	template<typename T>
+	T* FSelectionGroupCustomization::GetOwnerStruct() const
+	{
+		if (const TSharedPtr<IPropertyHandle> PropertyHandle = ChildPropertyHandle->GetParentHandle())
+		{
+			if (const TSharedPtr<IPropertyHandle> OwnerHandle = PropertyHandle->GetParentHandle())  // Assume that the group struct is only used at the node struct level
+			{
+				if (const TSharedPtr<IPropertyHandleStruct> OwnerHandleStruct = OwnerHandle->AsStruct())
+				{
+					if (TSharedPtr<FStructOnScope> StructOnScope = OwnerHandleStruct->GetStructData())
+					{
+						if (StructOnScope->GetStruct() == T::StaticStruct())
+						{
+							return reinterpret_cast<T*>(StructOnScope->GetStructMemory());
+						}
+					}
+				}
+			}
+		}
+		return nullptr;
+	}
+}  // End namespace UE::Chaos::ClothAsset
+
+#undef LOCTEXT_NAMESPACE
diff --git a/Engine/Plugins/Experimental/ChaosClothAssetEditor/Source/ChaosClothAssetDataflowNodes/Private/ChaosClothAsset/SelectionGroupCustomization.h b/Engine/Plugins/Experimental/ChaosClothAssetEditor/Source/ChaosClothAssetDataflowNodes/Private/ChaosClothAsset/SelectionGroupCustomization.h
new file mode 100644
--- /dev/null
+++ b/Engine/Plugins/Experimental/ChaosClothAssetEditor/Source/ChaosClothAssetDataflowNodes/Private/ChaosClothAsset/SelectionGroupCustomization.h
@@ -0,0 +1,41 @@
+// Copyright Epic Games, Inc. All Rights Reserved.
+
+#pragma once
+
+#include "IPropertyTypeCustomization.h"
+
+class IPropertyHandle;
+class ITableRow;
+class SComboButton;
+class STableViewBase;
+class SWidget;
+
+namespace UE::Chaos::ClothAsset
+{
+	/**
+	 * Customization for selection node group names input.
+	 */
+	class FSelectionGroupCustomization : public IPropertyTypeCustomization
+	{
+	public:
+
+		static TSharedRef<IPropertyTypeCustomization> MakeInstance();
+
+		virtual void CustomizeHeader(TSharedRef<IPropertyHandle> InPropertyHandle, FDetailWidgetRow& HeaderRow, IPropertyTypeCustomizationUtils& CustomizationUtils) override;
+		virtual void CustomizeChildren(TSharedRef<IPropertyHandle> /*InPropertyHandle*/, IDetailChildrenBuilder& /*ChildBuilder*/, IPropertyTypeCustomizationUtils& /*CustomizationUtils*/) override {}
+
+	private:
+		FText GetText() const;
+		void OnTextCommitted(const FText& NewText, ETextCommit::Type CommitInfo);
+		void OnSelectionChanged(TSharedPtr<FText> ItemSelected, ESelectInfo::Type SelectInfo);
+		TSharedRef<ITableRow> MakeCategoryViewWidget(TSharedPtr<FText> Item, const TSharedRef<STableViewBase>& OwnerTable);
+		TSharedRef<SWidget> OnGetMenuContent();
+		
+		template<typename T>
+		T* GetOwnerStruct() const;
+
+		TSharedPtr<IPropertyHandle> ChildPropertyHandle;
+		TSharedPtr<SComboButton> ComboButton;
+		TArray<TSharedPtr<FText>> GroupNames;
+	};
+}
diff --git a/Engine/Plugins/Experimental/ChaosClothAssetEditor/Source/ChaosClothAssetDataflowNodes/Private/ChaosClothAsset/SelectionNode.cpp b/Engine/Plugins/Experimental/ChaosClothAssetEditor/Source/ChaosClothAssetDataflowNodes/Private/ChaosClothAsset/SelectionNode.cpp
--- a/Engine/Plugins/Experimental/ChaosClothAssetEditor/Source/ChaosClothAssetDataflowNodes/Private/ChaosClothAsset/SelectionNode.cpp
+++ b/Engine/Plugins/Experimental/ChaosClothAssetEditor/Source/ChaosClothAssetDataflowNodes/Private/ChaosClothAsset/SelectionNode.cpp
@@ -1,8 +1,9 @@
 // Copyright Epic Games, Inc. All Rights Reserved.
 
 #include "ChaosClothAsset/SelectionNode.h"
-#include "ChaosClothAsset/CollectionClothSelectionFacade.h"
+#include "ChaosClothAsset/ClothCollectionGroup.h"
 #include "ChaosClothAsset/ClothDataflowTools.h"
+#include "ChaosClothAsset/CollectionClothSelectionFacade.h"
 #include "Dataflow/DataflowInputOutput.h"
 
 #include UE_INLINE_GENERATED_CPP_BY_NAME(SelectionNode)
@@ -23,34 +24,73 @@ void FChaosClothAssetSelectionNode::Evaluate(Dataflow::FContext& Context, const
 
 	if (Out->IsA<FManagedArrayCollection>(&Collection))
 	{
-		FManagedArrayCollection InSelectionCollection = GetValue<FManagedArrayCollection>(Context, &Collection);
-		const TSharedRef<FManagedArrayCollection> SelectionCollection = MakeShared<FManagedArrayCollection>(MoveTemp(InSelectionCollection));
-		FCollectionClothSelectionFacade SelectionFacade(SelectionCollection);
-		SelectionFacade.DefineSchema();
-		
-		checkSlow(SelectionFacade.IsValid());
-
-		const int32 PreviousNumSelections = SelectionFacade.GetNumSelections();
-		const int32 NewSelectionIndex = SelectionFacade.FindOrAddSelection(Name);
-		checkf(NewSelectionIndex != INDEX_NONE, TEXT("Failed to find or add a new Selection"));
-
-		if (NewSelectionIndex != PreviousNumSelections)
+		if (Name.IsEmpty() || Group.Name.IsEmpty())
 		{
-			FClothDataflowTools::LogAndToastWarning(*this,
-				LOCTEXT("SelectionAlreadyExistsHeadline", "Selection already exists"),
-				FText::Format(LOCTEXT("SelectionAlreadyExistsDetails", "A Selection with Name \"{0}\" already exists in the Collection. Its data will be overwritten."),
-					FText::FromString(Name)));
+			const FManagedArrayCollection& SelectionCollection = GetValue<FManagedArrayCollection>(Context, &Collection);
+			SetValue(Context, SelectionCollection, &Collection);
 		}
+		else
+		{
+			const FName SelectionName(*Name);
+			const FName SelectionGroupName(*Group.Name);
 
-		const UEnum* const SelectionTypeEnum = StaticEnum<EChaosClothAssetSelectionType>();
+			FManagedArrayCollection InSelectionCollection = GetValue<FManagedArrayCollection>(Context, &Collection);
+			const TSharedRef<FManagedArrayCollection> SelectionCollection = MakeShared<FManagedArrayCollection>(MoveTemp(InSelectionCollection));
 
-		const bool bSuccess = SelectionFacade.SetSelection(NewSelectionIndex, SelectionTypeEnum->GetNameStringByValue(static_cast<int32>(Type)), Indices);
-		checkf(bSuccess, TEXT("Failed to set value on a Selection"));
+			FCollectionClothSelectionFacade SelectionFacade(SelectionCollection);
+			SelectionFacade.DefineSchema();
 
-		SetValue(Context, MoveTemp(*SelectionCollection), &Collection);
+			TSet<int32>& SelectionSet = SelectionFacade.FindOrAddSelectionSet(SelectionName, SelectionGroupName);
+			SelectionSet = Indices;
+
+			SetValue(Context, MoveTemp(*SelectionCollection), &Collection);
+		}
+	}
+	else if (Out->IsA<FString>(&Name))
+	{
 		SetValue(Context, Name, &Name);
 	}
 }
 
+void FChaosClothAssetSelectionNode::OnSelected(Dataflow::FContext& Context)
+{
+	// Re-evaluate the input collection
+	const FManagedArrayCollection& SelectionCollection = GetValue<FManagedArrayCollection>(Context, &Collection);
+
+	// Update the list of used group for the UI customization
+	CachedCollectionGroupNames = SelectionCollection.GroupNames();
+}
+
+void FChaosClothAssetSelectionNode::OnDeselected()
+{
+	// Clean up, to avoid another toolkit picking up the wrong context evaluation
+	CachedCollectionGroupNames.Reset();
+}
+
+void FChaosClothAssetSelectionNode::Serialize(FArchive& Ar)
+{
+	using namespace UE::Chaos::ClothAsset;
+
+	// This is just for convenience and can be removed post 5.4 once the plugin loses its experimental status
+PRAGMA_DISABLE_DEPRECATION_WARNINGS
+	if (Ar.IsLoading() && Type_DEPRECATED != EChaosClothAssetSelectionType::Deprecated)
+	{
+		switch (Type_DEPRECATED)
+		{
+		case EChaosClothAssetSelectionType::SimVertex2D: Group.Name = ClothCollectionGroup::SimVertices2D.ToString(); break;
+		case EChaosClothAssetSelectionType::SimVertex3D: Group.Name = ClothCollectionGroup::SimVertices3D.ToString(); break;
+		case EChaosClothAssetSelectionType::RenderVertex: Group.Name = ClothCollectionGroup::RenderVertices.ToString(); break;
+		case EChaosClothAssetSelectionType::SimFace: Group.Name = ClothCollectionGroup::SimFaces.ToString(); break;
+		case EChaosClothAssetSelectionType::RenderFace: Group.Name = ClothCollectionGroup::RenderFaces.ToString(); break;
+		default: checkNoEntry();
+		}
+		Type_DEPRECATED = EChaosClothAssetSelectionType::Deprecated;  // This is only for clarity since the Type property won't be saved from now on
+
+		FClothDataflowTools::LogAndToastWarning(*this,
+			LOCTEXT("DeprecatedSelectionType", "Outdated Dataflow asset."),
+			LOCTEXT("DeprecatedSelectionDetails", "This node is out of data and contain deprecated data. The asset needs to be re-saved before it stops working at the next version update."));
+	}
+PRAGMA_ENABLE_DEPRECATION_WARNINGS
+}
 
 #undef LOCTEXT_NAMESPACE
diff --git a/Engine/Plugins/Experimental/ChaosClothAssetEditor/Source/ChaosClothAssetDataflowNodes/Private/ChaosClothAsset/SelectionToWeightMapNode.cpp b/Engine/Plugins/Experimental/ChaosClothAssetEditor/Source/ChaosClothAssetDataflowNodes/Private/ChaosClothAsset/SelectionToWeightMapNode.cpp
--- a/Engine/Plugins/Experimental/ChaosClothAssetEditor/Source/ChaosClothAssetDataflowNodes/Private/ChaosClothAsset/SelectionToWeightMapNode.cpp
+++ b/Engine/Plugins/Experimental/ChaosClothAssetEditor/Source/ChaosClothAssetDataflowNodes/Private/ChaosClothAsset/SelectionToWeightMapNode.cpp
@@ -1,9 +1,10 @@
 // Copyright Epic Games, Inc. All Rights Reserved.
 
 #include "ChaosClothAsset/SelectionToWeightMapNode.h"
+#include "ChaosClothAsset/ClothCollectionGroup.h"
+#include "ChaosClothAsset/ClothDataflowTools.h"
 #include "ChaosClothAsset/CollectionClothFacade.h"
 #include "ChaosClothAsset/CollectionClothSelectionFacade.h"
-#include "ChaosClothAsset/ClothDataflowTools.h"
 #include "Dataflow/DataflowInputOutput.h"
 
 #include UE_INLINE_GENERATED_CPP_BY_NAME(SelectionToWeightMapNode)
@@ -18,7 +19,6 @@ FChaosClothAssetSelectionToWeightMapNode::FChaosClothAssetSelectionToWeightMapNo
 	RegisterOutputConnection(&WeightMapName);
 }
 
-
 void FChaosClothAssetSelectionToWeightMapNode::Evaluate(Dataflow::FContext& Context, const FDataflowOutput* Out) const
 {
 	using namespace UE::Chaos::ClothAsset;
@@ -30,36 +30,33 @@ void FChaosClothAssetSelectionToWeightMapNode::Evaluate(Dataflow::FContext& Cont
 		FCollectionClothFacade ClothFacade(ClothCollection);
 
 		const FCollectionClothSelectionConstFacade SelectionFacade(ClothCollection);
-		const FString InSelectionName = GetValue<FString>(Context, &SelectionName);
+		const FName InSelectionName(*GetValue<FString>(Context, &SelectionName));
 
-		if (SelectionFacade.IsValid() && ClothFacade.IsValid())
+		if (SelectionFacade.IsValid() && ClothFacade.IsValid() && InSelectionName != NAME_None)
 		{
-			const int32 FoundSelectionIndex = SelectionFacade.FindSelection(InSelectionName);
-
-			if (FoundSelectionIndex != INDEX_NONE)
+			if (const TSet<int32>* const SelectionSet = SelectionFacade.FindSelectionSet(InSelectionName))
 			{
-				const FString& SelectionType = SelectionFacade.GetType()[FoundSelectionIndex];
+				const FName& SelectionGroup = SelectionFacade.GetSelectionGroup(InSelectionName);
 
-				if (SelectionType == "SimVertex3D" || SelectionType == "SimVertex2D" || SelectionType == "SimFace")
+				if (SelectionGroup == ClothCollectionGroup::SimVertices2D ||
+					SelectionGroup == ClothCollectionGroup::SimVertices3D || 
+					SelectionGroup == ClothCollectionGroup::SimFaces)
 				{
 					const FName InMapName(WeightMapName);
 					ClothFacade.AddWeightMap(InMapName);
 					TArrayView<float> OutClothWeights = ClothFacade.GetWeightMap(InMapName);
 
-					const TSet<int32>& Selection = SelectionFacade.GetIndices()[FoundSelectionIndex];
-
-					if (SelectionType == "SimVertex3D")
+					if (SelectionGroup == ClothCollectionGroup::SimVertices3D)
 					{
 						for (int32 VertexIndex = 0; VertexIndex < OutClothWeights.Num(); ++VertexIndex)
 						{
-							OutClothWeights[VertexIndex] = Selection.Contains(VertexIndex) ? SelectedValue : UnselectedValue;
+							OutClothWeights[VertexIndex] = SelectionSet->Contains(VertexIndex) ? SelectedValue : UnselectedValue;
 						}
 					}
-					else if (SelectionType == "SimVertex2D")
+					else if (SelectionGroup == ClothCollectionGroup::SimVertices2D)
 					{
 						// We are given a selection over the set of 2D vertices, but weight maps only exist for 3D vertices, so
 						// we need a bit of translation
-
 						const TConstArrayView<TArray<int32>> Vertex3DTo2D = ClothFacade.GetSimVertex2DLookup();
 
 						for (int32 Vertex3DIndex = 0; Vertex3DIndex < OutClothWeights.Num(); ++Vertex3DIndex)
@@ -68,7 +65,7 @@ void FChaosClothAssetSelectionToWeightMapNode::Evaluate(Dataflow::FContext& Cont
 							OutClothWeights[Vertex3DIndex] = UnselectedValue;
 							for (const int32 Vertex2DIndex : Vertex3DTo2D[Vertex3DIndex])
 							{
-								if (Selection.Contains(Vertex2DIndex))
+								if (SelectionSet->Contains(Vertex2DIndex))
 								{
 									OutClothWeights[Vertex3DIndex] = SelectedValue;
 									break;
@@ -78,14 +75,15 @@ void FChaosClothAssetSelectionToWeightMapNode::Evaluate(Dataflow::FContext& Cont
 					}
 					else
 					{
-						check(SelectionType == "SimFace");
+						check(SelectionGroup == ClothCollectionGroup::SimFaces);
+
 						// Fill with unselected value and then set all vertices in the selected faces to the selected value.
 						for (int32 Vertex3DIndex = 0; Vertex3DIndex < OutClothWeights.Num(); ++Vertex3DIndex)
 						{
 							OutClothWeights[Vertex3DIndex] = UnselectedValue;
 						}
 						const TConstArrayView<FIntVector3> SimIndices3D = ClothFacade.GetSimIndices3D();
-						for (const int32 FaceIndex : Selection)
+						for (const int32 FaceIndex : *SelectionSet)
 						{
 							OutClothWeights[SimIndices3D[FaceIndex][0]] = SelectedValue;
 							OutClothWeights[SimIndices3D[FaceIndex][1]] = SelectedValue;
@@ -96,16 +94,20 @@ void FChaosClothAssetSelectionToWeightMapNode::Evaluate(Dataflow::FContext& Cont
 				else
 				{
 					FClothDataflowTools::LogAndToastWarning(*this,
-						LOCTEXT("SelectionTypeNotCorrectHeadline", "Selection type is incompatible."),
-						FText::Format(LOCTEXT("SelectionTypeNotCorrectDetails", "Selection with Name \"{0}\" does not have Type \"SimVertex3D\", \"SimVertex2D\", or \"SimFace\"."), FText::FromString(InSelectionName)));
+						LOCTEXT("SelectionTypeNotCorrectHeadline", "Invalid selection group."),
+						FText::Format(LOCTEXT("SelectionTypeNotCorrectDetails", "Selection \"{0}\" does not have its index dependency group set to \"{1}\", \"{2}\", or \"{3}\"."),
+							FText::FromName(InSelectionName),
+							FText::FromName(ClothCollectionGroup::SimVertices2D),
+							FText::FromName(ClothCollectionGroup::SimVertices3D),
+							FText::FromName(ClothCollectionGroup::SimFaces)));
 				}
 			}
 			else
 			{
 				FClothDataflowTools::LogAndToastWarning(*this,
-					LOCTEXT("SelectionNameNotFoundHeadline", "Selection Name was not found."),
-					FText::Format(LOCTEXT("SelectionNameNotFoundDetails", "A Selection with Name \"{0}\" was not found in the Collection."),
-						FText::FromString(InSelectionName)));
+					LOCTEXT("SelectionNameNotFoundHeadline", "Invalid selection name."),
+					FText::Format(LOCTEXT("SelectionNameNotFoundDetails", "Selection \"{0}\" was not found in the collection."),
+						FText::FromName(InSelectionName)));
 			}
 		}
 
diff --git a/Engine/Plugins/Experimental/ChaosClothAssetEditor/Source/ChaosClothAssetDataflowNodes/Private/ChaosClothAsset/WeightedValue.cpp b/Engine/Plugins/Experimental/ChaosClothAssetEditor/Source/ChaosClothAssetDataflowNodes/Private/ChaosClothAsset/WeightedValue.cpp
new file mode 100644
--- /dev/null
+++ b/Engine/Plugins/Experimental/ChaosClothAssetEditor/Source/ChaosClothAssetDataflowNodes/Private/ChaosClothAsset/WeightedValue.cpp
@@ -0,0 +1,12 @@
+// Copyright Epic Games, Inc. All Rights Reserved.
+
+#include "ChaosClothAsset/WeightedValue.h"
+#include "ChaosClothAsset/ClothDataflowTools.h"
+
+namespace UE::Chaos::ClothAsset
+{
+	void FWeightMapTools::MakeWeightMapName(FString& InOutString)
+	{
+		FClothDataflowTools::MakeCollectionName(InOutString);
+	}
+}  // End namespace UE::Chaos::ClothAsset
diff --git a/Engine/Plugins/Experimental/ChaosClothAssetEditor/Source/ChaosClothAssetDataflowNodes/Public/ChaosClothAsset/SelectionNode.h b/Engine/Plugins/Experimental/ChaosClothAssetEditor/Source/ChaosClothAssetDataflowNodes/Public/ChaosClothAsset/SelectionNode.h
--- a/Engine/Plugins/Experimental/ChaosClothAssetEditor/Source/ChaosClothAssetDataflowNodes/Public/ChaosClothAsset/SelectionNode.h
+++ b/Engine/Plugins/Experimental/ChaosClothAssetEditor/Source/ChaosClothAssetDataflowNodes/Public/ChaosClothAsset/SelectionNode.h
@@ -8,7 +8,7 @@
 
 /** What type of element is selected in the Selection */
 UENUM()
-enum class EChaosClothAssetSelectionType : uint8
+enum class UE_DEPRECATED(5.4, "Use FChaosClothAssetNodeSelectionGroup instead") EChaosClothAssetSelectionType : uint8
 {
 	/** 2D simulation vertices */
 	SimVertex2D,
@@ -23,10 +23,26 @@ enum class EChaosClothAssetSelectionType : uint8
 	SimFace,
 
 	/** Render faces */
-	RenderFace
+	RenderFace,
+
+	/** Deprecated marker */
+	Deprecated UMETA(Hidden)
+};
+
+/**
+ * The managed array collection group used in the selection.
+ * This separate structure is required to allow for customization of the UI.
+ */
+USTRUCT()
+struct FChaosClothAssetNodeSelectionGroup
+{
+	GENERATED_USTRUCT_BODY()
+
+	UPROPERTY(EditAnywhere, Category = "Selection Group")
+	FString Name;
 };
 
-/** Integer index set selection node */
+/** Integer index set selection node. */
 USTRUCT(Meta = (DataflowCloth))
 struct FChaosClothAssetSelectionNode : public FDataflowNode
 {
@@ -42,18 +58,33 @@ public:
 	UPROPERTY(EditAnywhere, Category = "Selection", Meta = (DataflowOutput))
 	FString Name;
 
+	/** The type of element the selection refers to */
+	UE_DEPRECATED(5.4, "Use Group instead")
+PRAGMA_DISABLE_DEPRECATION_WARNINGS
+	UPROPERTY()
+	EChaosClothAssetSelectionType Type_DEPRECATED = EChaosClothAssetSelectionType::Deprecated;
+PRAGMA_ENABLE_DEPRECATION_WARNINGS
+
 	/** The type of element the selection refers to */
 	UPROPERTY(EditAnywhere, Category = "Selection")
-	EChaosClothAssetSelectionType Type = EChaosClothAssetSelectionType::SimVertex2D;
+	FChaosClothAssetNodeSelectionGroup Group;
 
 	/** Selected element indices */
-	UPROPERTY(EditAnywhere, Category = "Selection")
+	UPROPERTY(EditAnywhere, Category = "Selection", Meta = (ClampMin = "0"))
 	TSet<int32> Indices;
 
 	FChaosClothAssetSelectionNode(const Dataflow::FNodeParameters& InParam, FGuid InGuid = FGuid::NewGuid());
 
+	/** Return a cached array of all the groups used by the input collection during at the time of the latest evaluation. */
+	const TArray<FName>& GetCachedCollectionGroupNames() const { return CachedCollectionGroupNames; }
+
 private:
 
 	virtual void Evaluate(Dataflow::FContext& Context, const FDataflowOutput* Out) const override;
+	virtual void OnSelected(Dataflow::FContext& Context) override;
+	virtual void OnDeselected() override;
+	virtual void Serialize(FArchive& Ar);
+
+	TArray<FName> CachedCollectionGroupNames;
 };
 
diff --git a/Engine/Plugins/Experimental/ChaosClothAssetEditor/Source/ChaosClothAssetDataflowNodes/Public/ChaosClothAsset/WeightedValue.h b/Engine/Plugins/Experimental/ChaosClothAssetEditor/Source/ChaosClothAssetDataflowNodes/Public/ChaosClothAsset/WeightedValue.h
--- a/Engine/Plugins/Experimental/ChaosClothAssetEditor/Source/ChaosClothAssetDataflowNodes/Public/ChaosClothAsset/WeightedValue.h
+++ b/Engine/Plugins/Experimental/ChaosClothAssetEditor/Source/ChaosClothAssetDataflowNodes/Public/ChaosClothAsset/WeightedValue.h
@@ -16,16 +16,11 @@ namespace UE::Chaos::ClothAsset
 		inline static const FString NotOverridden = TEXT("#");
 
 		/**
-		 * Modifies a string to make it suitable as an attribute name.
+		 * Modifies a string to make it suitable as a weight map name.
 		 * Replaces any deemed special characters and spaces by underscores, and remove leading and ending underscores from the name.
 		 * Note that leading underscores are reserved for non user/internal attribute names.
 		 */
-		static void MakeWeightMapName(FString& InOutString)
-		{
-			InOutString = SlugStringForValidName(InOutString, TEXT("_")).Replace(TEXT("\\"), TEXT("_"));
-			bool bCharsWereRemoved;
-			do { InOutString.TrimCharInline(TEXT('_'), &bCharsWereRemoved); } while (bCharsWereRemoved);
-		}
+		static CHAOSCLOTHASSETDATAFLOWNODES_API void MakeWeightMapName(FString& InOutString);
 	};
 }
 
diff --git a/Engine/Plugins/Experimental/ChaosClothAssetEditor/Source/ChaosClothAssetEditor/Private/ChaosClothAsset/ClothEditorToolkit.cpp b/Engine/Plugins/Experimental/ChaosClothAssetEditor/Source/ChaosClothAssetEditor/Private/ChaosClothAsset/ClothEditorToolkit.cpp
--- a/Engine/Plugins/Experimental/ChaosClothAssetEditor/Source/ChaosClothAssetEditor/Private/ChaosClothAsset/ClothEditorToolkit.cpp
+++ b/Engine/Plugins/Experimental/ChaosClothAssetEditor/Source/ChaosClothAssetEditor/Private/ChaosClothAsset/ClothEditorToolkit.cpp
@@ -209,6 +209,7 @@ FChaosClothAssetEditorToolkit::~FChaosClothAssetEditorToolkit()
 	if (DataflowNode && OnNodeInvalidatedDelegateHandle.IsValid())
 	{
 		DataflowNode->GetOnNodeInvalidatedDelegate().Remove(OnNodeInvalidatedDelegateHandle);
+		DataflowNode->OnDeselected();
 	}
 	DataflowNode.Reset();
 
@@ -933,7 +934,7 @@ void FChaosClothAssetEditorToolkit::OnNodeTitleCommitted(const FText& InNewText,
 
 void FChaosClothAssetEditorToolkit::OnNodeSelectionChanged(const TSet<UObject*>& NewSelection)
 {
-	auto GetClothCollectionIfPossible = [](const TSharedPtr<const FDataflowNode> InDataflowNode, const TSharedPtr<Dataflow::FEngineContext> Context) -> TSharedPtr<FManagedArrayCollection>
+	auto GetClothCollectionIfPossible = [](const TSharedPtr<FDataflowNode> InDataflowNode, const TSharedPtr<Dataflow::FEngineContext> Context) -> TSharedPtr<FManagedArrayCollection>
 	{
 		if (Context.IsValid())
 		{
@@ -971,32 +972,56 @@ void FChaosClothAssetEditorToolkit::OnNodeSelectionChanged(const TSet<UObject*>&
 	{
 		Dataflow->RenderTargets.Reset();
 
-		for (UObject* const Selected : NewSelection)
+		if (!NewSelection.Num())
 		{
-			if (UDataflowEdNode* const Node = Cast<UDataflowEdNode>(Selected))
+			if (DataflowNode)
 			{
-				Dataflow->RenderTargets.Add(Node);
-
-				if (DataflowNode && OnNodeInvalidatedDelegateHandle.IsValid())
+				if (OnNodeInvalidatedDelegateHandle.IsValid())
 				{
 					DataflowNode->GetOnNodeInvalidatedDelegate().Remove(OnNodeInvalidatedDelegateHandle);
 				}
-				DataflowNode = Node->GetDataflowNode();
-
-				if (DataflowNode)
+				DataflowNode->OnDeselected();
+				DataflowNode.Reset();
+			}
+		}
+		else
+		{
+			for (UObject* const Selected : NewSelection)
+			{
+				if (UDataflowEdNode* const Node = Cast<UDataflowEdNode>(Selected))
 				{
-					Collection = GetClothCollectionIfPossible(DataflowNode, this->DataflowContext);
-					DataflowNode->GetOnNodeInvalidatedDelegate();
+					Dataflow->RenderTargets.Add(Node);
 
-					// Set a callback to re-evaluate the node if it is invalidated
-					OnNodeInvalidatedDelegateHandle = DataflowNode->GetOnNodeInvalidatedDelegate().AddLambda(
-						[this, &GetClothCollectionIfPossible](FDataflowNode* InDataflowNode)
+					if (DataflowNode != Node->GetDataflowNode())
+					{
+						if (DataflowNode)
 						{
-							if (DataflowNode.Get() == InDataflowNode)
+							if (OnNodeInvalidatedDelegateHandle.IsValid())
 							{
-								GetClothCollectionIfPossible(DataflowNode, DataflowContext);
+								DataflowNode->GetOnNodeInvalidatedDelegate().Remove(OnNodeInvalidatedDelegateHandle);
 							}
-						});
+							DataflowNode->OnDeselected();
+						}
+						DataflowNode = Node->GetDataflowNode();
+
+						if (DataflowNode)
+						{
+							Collection = GetClothCollectionIfPossible(DataflowNode, DataflowContext);
+							DataflowNode->OnSelected(*DataflowContext);
+
+							// Set a callback to re-evaluate the node if it is invalidated
+							OnNodeInvalidatedDelegateHandle = DataflowNode->GetOnNodeInvalidatedDelegate().AddLambda(
+								[this, &GetClothCollectionIfPossible](FDataflowNode* InDataflowNode)
+								{
+									if (DataflowNode.Get() == InDataflowNode)
+									{
+										GetClothCollectionIfPossible(DataflowNode, DataflowContext);
+										DataflowNode->OnSelected(*DataflowContext);
+									}
+								});
+						}
+					}
+					break;
 				}
 			}
 		}
diff --git a/Engine/Plugins/Experimental/ChaosClothAssetEditor/Source/ChaosClothAssetEditorTools/Private/ChaosClothAsset/ClothMeshSelectionTool.cpp b/Engine/Plugins/Experimental/ChaosClothAssetEditor/Source/ChaosClothAssetEditorTools/Private/ChaosClothAsset/ClothMeshSelectionTool.cpp
--- a/Engine/Plugins/Experimental/ChaosClothAssetEditor/Source/ChaosClothAssetEditorTools/Private/ChaosClothAsset/ClothMeshSelectionTool.cpp
+++ b/Engine/Plugins/Experimental/ChaosClothAssetEditor/Source/ChaosClothAssetEditorTools/Private/ChaosClothAsset/ClothMeshSelectionTool.cpp
@@ -1,6 +1,7 @@
 // Copyright Epic Games, Inc. All Rights Reserved.
 
 #include "ChaosClothAsset/ClothMeshSelectionTool.h"
+#include "ChaosClothAsset/ClothCollectionGroup.h"
 #include "ChaosClothAsset/ClothEditorContextObject.h"
 #include "ChaosClothAsset/ClothPatternVertexType.h"
 #include "ChaosClothAsset/SelectionNode.h"
@@ -308,6 +309,8 @@ void UClothMeshSelectionTool::SetClothEditorContextObject(TObjectPtr<UClothEdito
 
 bool UClothMeshSelectionTool::GetSelectedNodeInfo(FString& OutSelectionName, UE::Geometry::FGroupTopologySelection& OutSelection)
 {
+	using namespace UE::Chaos::ClothAsset;
+
 	const FChaosClothAssetSelectionNode* const MeshSelectionNode = ClothEditorContextObject->GetSingleSelectedNodeOfType<FChaosClothAssetSelectionNode>();
 	check(MeshSelectionNode);
 
@@ -342,7 +345,9 @@ bool UClothMeshSelectionTool::GetSelectedNodeInfo(FString& OutSelectionName, UE:
 	};
 
 
-	if (MeshSelectionNode->Type == EChaosClothAssetSelectionType::SimVertex2D || MeshSelectionNode->Type == EChaosClothAssetSelectionType::SimVertex3D || MeshSelectionNode->Type == EChaosClothAssetSelectionType::RenderVertex)
+	if (MeshSelectionNode->Group.Name == ClothCollectionGroup::SimVertices2D.ToString() ||
+		MeshSelectionNode->Group.Name == ClothCollectionGroup::SimVertices3D.ToString() ||
+		MeshSelectionNode->Group.Name == ClothCollectionGroup::RenderVertices.ToString())
 	{
 		if (bHasNonManifoldMapping)
 		{
@@ -356,7 +361,8 @@ bool UClothMeshSelectionTool::GetSelectedNodeInfo(FString& OutSelectionName, UE:
 			AppendVerticesIfValid(OutSelection.SelectedCornerIDs, MeshSelectionNode->Indices);
 		}
 	}
-	else if (MeshSelectionNode->Type == EChaosClothAssetSelectionType::SimFace || MeshSelectionNode->Type == EChaosClothAssetSelectionType::RenderFace)
+	else if (MeshSelectionNode->Group.Name == ClothCollectionGroup::SimFaces.ToString() ||
+		MeshSelectionNode->Group.Name == ClothCollectionGroup::RenderFaces)
 	{
 		AppendFacesIfValid(OutSelection.SelectedGroupIDs, MeshSelectionNode->Indices);
 	}
@@ -369,11 +375,13 @@ bool UClothMeshSelectionTool::GetSelectedNodeInfo(FString& OutSelectionName, UE:
 
 void UClothMeshSelectionTool::UpdateSelectedNode()
 {
+	using namespace UE::Chaos::ClothAsset;
+
 	const UE::Geometry::FGroupTopologySelection& Selection = SelectionMechanic->GetActiveSelection();
-	const UE::Chaos::ClothAsset::EClothPatternVertexType ViewMode = ClothEditorContextObject->GetConstructionViewMode();
+	const EClothPatternVertexType ViewMode = ClothEditorContextObject->GetConstructionViewMode();
 
 	TSet<int32> Indices;
-	EChaosClothAssetSelectionType Type = EChaosClothAssetSelectionType::SimVertex2D;
+	FName GroupName = ClothCollectionGroup::SimVertices2D;
 
 	if (SelectionMechanic->Properties->bSelectVertices)
 	{
@@ -385,16 +393,15 @@ void UClothMeshSelectionTool::UpdateSelectedNode()
 		switch(ViewMode)
 		{
 		case UE::Chaos::ClothAsset::EClothPatternVertexType::Sim2D:
-			Type = EChaosClothAssetSelectionType::SimVertex2D;
+			GroupName = ClothCollectionGroup::SimVertices2D;
 			break;
 		case UE::Chaos::ClothAsset::EClothPatternVertexType::Sim3D:
-			Type = EChaosClothAssetSelectionType::SimVertex3D;
+			GroupName = ClothCollectionGroup::SimVertices3D;
 			break;
 		case UE::Chaos::ClothAsset::EClothPatternVertexType::Render:
-			Type = EChaosClothAssetSelectionType::RenderVertex;
+			GroupName = ClothCollectionGroup::RenderVertices;
 			break;
 		}
-
 	}
 	else
 	{
@@ -404,10 +411,10 @@ void UClothMeshSelectionTool::UpdateSelectedNode()
 		{
 		case UE::Chaos::ClothAsset::EClothPatternVertexType::Sim2D:
 		case UE::Chaos::ClothAsset::EClothPatternVertexType::Sim3D:
-			Type = EChaosClothAssetSelectionType::SimFace;
+			GroupName = ClothCollectionGroup::SimFaces;
 			break;
 		case UE::Chaos::ClothAsset::EClothPatternVertexType::Render:
-			Type = EChaosClothAssetSelectionType::RenderFace;
+			GroupName = ClothCollectionGroup::RenderFaces;
 			break;
 		}
 	}
@@ -416,7 +423,7 @@ void UClothMeshSelectionTool::UpdateSelectedNode()
 	check(MeshSelectionNode);
 
 	MeshSelectionNode->Name = ToolProperties->Name;
-	MeshSelectionNode->Type = Type;
+	MeshSelectionNode->Group.Name = GroupName.ToString();
 
 	if (SelectionMechanic->Properties->bSelectVertices && bHasNonManifoldMapping)
 	{
diff --git a/Engine/Plugins/Experimental/Dataflow/Source/DataflowEditor/Private/Dataflow/DataflowGraphEditor.cpp b/Engine/Plugins/Experimental/Dataflow/Source/DataflowEditor/Private/Dataflow/DataflowGraphEditor.cpp
--- a/Engine/Plugins/Experimental/Dataflow/Source/DataflowEditor/Private/Dataflow/DataflowGraphEditor.cpp
+++ b/Engine/Plugins/Experimental/Dataflow/Source/DataflowEditor/Private/Dataflow/DataflowGraphEditor.cpp
@@ -179,12 +179,12 @@ void SDataflowGraphEditor::DeleteNode()
 
 void SDataflowGraphEditor::OnSelectedNodesChanged(const TSet<UObject*>& NewSelection)
 {
+	OnSelectionChangedMulticast.Broadcast(NewSelection);  // Broadcast the selection change before refreshing the DetailsView, the nodes' specific UI data have to be updated before the UI is being redrawn
+
 	if (DataflowAsset.Get() && DetailsView)
 	{
 		FDataflowEditorCommands::OnSelectedNodesChanged(DetailsView, AssetOwner.Get(), DataflowAsset.Get(), NewSelection);
 	}
-
-	OnSelectionChangedMulticast.Broadcast(NewSelection);
 }
 
 FReply SDataflowGraphEditor::OnKeyDown(const FGeometry& MyGeometry, const FKeyEvent& InKeyEvent)
diff --git a/Engine/Source/Runtime/Experimental/Chaos/Private/GeometryCollection/ManagedArrayCollection.cpp b/Engine/Source/Runtime/Experimental/Chaos/Private/GeometryCollection/ManagedArrayCollection.cpp
--- a/Engine/Source/Runtime/Experimental/Chaos/Private/GeometryCollection/ManagedArrayCollection.cpp
+++ b/Engine/Source/Runtime/Experimental/Chaos/Private/GeometryCollection/ManagedArrayCollection.cpp
@@ -419,6 +419,13 @@ void FManagedArrayCollection::SetDependency(FName Name, FName Group, FName Depen
 	}
 }
 
+FName FManagedArrayCollection::GetDependency(FName Name, FName Group) const
+{
+	check(HasAttribute(Name, Group));
+	const FKeyType Key = FManagedArrayCollection::MakeMapKey(Name, Group);
+	return Map[Key].GroupIndexDependency;
+}
+
 void FManagedArrayCollection::RemoveDependencyFor(FName Group)
 {
 	ensure(HasGroup(Group));
diff --git a/Engine/Source/Runtime/Experimental/Chaos/Public/GeometryCollection/ManagedArrayCollection.h b/Engine/Source/Runtime/Experimental/Chaos/Public/GeometryCollection/ManagedArrayCollection.h
--- a/Engine/Source/Runtime/Experimental/Chaos/Public/GeometryCollection/ManagedArrayCollection.h
+++ b/Engine/Source/Runtime/Experimental/Chaos/Public/GeometryCollection/ManagedArrayCollection.h
@@ -468,6 +468,13 @@ public:
 	*/
 	CHAOS_API void SetDependency(FName Name, FName Group, FName DependencyGroup, bool bAllowCircularDependency = false);
 
+	/**
+	* Return the group index dependency for the specified attribute.
+	* @param Name - The name of the attribute
+	* @param Group - The group that manages the attribute
+	*/
+	CHAOS_API FName GetDependency(FName Name, FName Group) const;
+
 	/**
 	*
 	*/
diff --git a/Engine/Source/Runtime/Experimental/Dataflow/Core/Public/Dataflow/DataflowNode.h b/Engine/Source/Runtime/Experimental/Dataflow/Core/Public/Dataflow/DataflowNode.h
--- a/Engine/Source/Runtime/Experimental/Dataflow/Core/Public/Dataflow/DataflowNode.h
+++ b/Engine/Source/Runtime/Experimental/Dataflow/Core/Public/Dataflow/DataflowNode.h
@@ -139,6 +139,11 @@ struct FDataflowNode
 	/** Override this method to provide custom serialization for this node. */
 	virtual void Serialize(FArchive& Ar) {}
 
+	/** Called by editor toolkits when the node is selected, or already selected and invalidated. */
+	virtual void OnSelected(Dataflow::FContext& Context) {}
+	/** Called by editor toolkits when the node is deselected. */
+	virtual void OnDeselected() {}
+
 	//
 	//  Struct Support
 	//
