Chaos Cloth Asset - Fixed crash when reloading (and also when deleting) an empty cloth asset.

diff --git a/Engine/Plugins/Experimental/ChaosClothAsset/Source/ChaosClothAssetEngine/Private/ChaosClothAsset/ClothAsset.cpp b/Engine/Plugins/Experimental/ChaosClothAsset/Source/ChaosClothAssetEngine/Private/ChaosClothAsset/ClothAsset.cpp
--- a/Engine/Plugins/Experimental/ChaosClothAsset/Source/ChaosClothAssetEngine/Private/ChaosClothAsset/ClothAsset.cpp
+++ b/Engine/Plugins/Experimental/ChaosClothAsset/Source/ChaosClothAssetEngine/Private/ChaosClothAsset/ClothAsset.cpp
@@ -108,6 +108,9 @@ UChaosClothAsset::UChaosClothAsset(const FObjectInitializer& ObjectInitializer)
 	UE::Chaos::ClothAsset::FCollectionClothFacade ClothFacade(ClothCollection);
 	ClothFacade.DefineSchema();
 
+	// Add the LODInfo for the default LOD 0
+	LODInfo.SetNum(1);
+
 	// Set default skeleton (must be done after having added the LOD)
 	constexpr bool bRebuildModels = false;
 	constexpr bool bRebindMeshes = false;
@@ -212,8 +215,13 @@ void UChaosClothAsset::BeginPostLoadInternal(FSkinnedAssetPostLoadContext& Conte
 		ClothFacade.DefineSchema();
 		bAnyInvalidLods = true;
 	}
+
+	const int32 NumLods = ClothCollections.Num();
+	check(NumLods >= 1);  // The default LOD 0 should be present now if it ever was missing
+	LODInfo.SetNum(NumLods);  // Always keep a matching number of LODInfos
+
 	bool bAnyInvalidSkeletons = false;
-	for (int32 LODIndex = 0; LODIndex < ClothCollections.Num(); ++LODIndex)
+	for (int32 LODIndex = 0; LODIndex < NumLods; ++LODIndex)
 	{
 		TSharedRef<FManagedArrayCollection>& ClothCollection = ClothCollections[LODIndex];
 
