Chaos Cloth - Added IntVector3 type to sim pattern initializations.

diff --git a/Engine/Plugins/Experimental/ChaosClothAsset/Source/ChaosClothAsset/Private/ChaosClothAsset/CollectionClothSimPatternFacade.cpp b/Engine/Plugins/Experimental/ChaosClothAsset/Source/ChaosClothAsset/Private/ChaosClothAsset/CollectionClothSimPatternFacade.cpp
--- a/Engine/Plugins/Experimental/ChaosClothAsset/Source/ChaosClothAsset/Private/ChaosClothAsset/CollectionClothSimPatternFacade.cpp
+++ b/Engine/Plugins/Experimental/ChaosClothAsset/Source/ChaosClothAsset/Private/ChaosClothAsset/CollectionClothSimPatternFacade.cpp
@@ -8,6 +8,17 @@
 
 namespace UE::Chaos::ClothAsset
 {
+	namespace Private
+	{
+		template<typename IndexType> static int32 GetNumFaces(const TArray<IndexType>& Indices) { check(Indices.Num() % 3 == 0); return Indices.Num() / 3; }
+		template<> int32 GetNumFaces<FIntVector3>(const TArray<FIntVector3>& Indices) { return Indices.Num(); }
+
+		template<int32 Component, typename IndexType> static int32 GetIndex(const TArray<IndexType>& Indices, int32 Index) { return (int32)Indices[Index * 3 + Component]; }
+		template<> int32 GetIndex<0, FIntVector3>(const TArray<FIntVector3>& Indices, int32 Index) { return Indices[Index][0]; }
+		template<> int32 GetIndex<1, FIntVector3>(const TArray<FIntVector3>& Indices, int32 Index) { return Indices[Index][1]; }
+		template<> int32 GetIndex<2, FIntVector3>(const TArray<FIntVector3>& Indices, int32 Index) { return Indices[Index][2]; }
+	}
+
 	int32 FCollectionClothSimPatternConstFacade::GetNumSimVertices2D() const
 	{
 		return ClothCollection->GetNumElements(
@@ -96,7 +107,7 @@ namespace UE::Chaos::ClothAsset
 		SetDefaults();
 	}
 
-	template<typename IndexType>
+	template<typename IndexType, typename TEnableIf<TIsIndexType<IndexType>::Value, int>::type>
 	void FCollectionClothSimPatternFacade::Initialize(const TArray<FVector2f>& Positions2D, const TArray<FVector3f>& Positions3D, const TArray<IndexType>& Indices)
 	{
 		Reset();
@@ -119,8 +130,7 @@ namespace UE::Chaos::ClothAsset
 			SimNormal[SimVertexIndex + StartVertexIndex[1]] = FVector3f::ZeroVector;
 		}
 
-		const int32 NumSimFaces = Indices.Num() / 3;
-		check(NumSimFaces * 3 == Indices.Num());
+		const int32 NumSimFaces = Private::GetNumFaces(Indices);
 
 		SetNumSimFaces(NumSimFaces);
 		const TArrayView<FIntVector3> SimIndices2D = GetSimIndices2D();
@@ -129,9 +139,9 @@ namespace UE::Chaos::ClothAsset
 		for (int32 SimFaceIndex = 0; SimFaceIndex < NumSimFaces; ++SimFaceIndex)
 		{
 			// Indices from the start of the pattern
-			const int32 PattternVertexIndex0 = (int32)Indices[SimFaceIndex * 3];
-			const int32 PattternVertexIndex1 = (int32)Indices[SimFaceIndex * 3 + 1];
-			const int32 PattternVertexIndex2 = (int32)Indices[SimFaceIndex * 3 + 2];
+			const int32 PattternVertexIndex0 = Private::GetIndex<0>(Indices, SimFaceIndex);
+			const int32 PattternVertexIndex1 = Private::GetIndex<1>(Indices, SimFaceIndex);
+			const int32 PattternVertexIndex2 = Private::GetIndex<2>(Indices, SimFaceIndex);
 
 			// 2D Indices from the start of the cloth
 			const int32 Vertex2DIndex0 = PattternVertexIndex0 + BaseVertex2D;
@@ -165,6 +175,7 @@ namespace UE::Chaos::ClothAsset
 	}
 	template CHAOSCLOTHASSET_API void FCollectionClothSimPatternFacade::Initialize(const TArray<FVector2f>& Positions2D, const TArray<FVector3f>& Positions3D, const TArray<int32>& Indices);
 	template CHAOSCLOTHASSET_API void FCollectionClothSimPatternFacade::Initialize(const TArray<FVector2f>& Positions2D, const TArray<FVector3f>& Positions3D, const TArray<uint32>& Indices);
+	template CHAOSCLOTHASSET_API void FCollectionClothSimPatternFacade::Initialize(const TArray<FVector2f>& Positions2D, const TArray<FVector3f>& Positions3D, const TArray<FIntVector3>& Indices);
 
 	void FCollectionClothSimPatternFacade::Initialize(const FCollectionClothSimPatternConstFacade& Other, const int32 SimVertex3DOffset)
 	{
diff --git a/Engine/Plugins/Experimental/ChaosClothAsset/Source/ChaosClothAsset/Public/ChaosClothAsset/CollectionClothSimPatternFacade.h b/Engine/Plugins/Experimental/ChaosClothAsset/Source/ChaosClothAsset/Public/ChaosClothAsset/CollectionClothSimPatternFacade.h
--- a/Engine/Plugins/Experimental/ChaosClothAsset/Source/ChaosClothAsset/Public/ChaosClothAsset/CollectionClothSimPatternFacade.h
+++ b/Engine/Plugins/Experimental/ChaosClothAsset/Source/ChaosClothAsset/Public/ChaosClothAsset/CollectionClothSimPatternFacade.h
@@ -6,6 +6,12 @@
 
 namespace UE::Chaos::ClothAsset
 {
+	/** Integral and vector types valid as index type in initializations. */
+	template<typename T> struct TIsIndexType { static constexpr bool Value = false; };
+	template<> struct TIsIndexType<int32> { static constexpr bool Value = true; };
+	template<> struct TIsIndexType<uint32> { static constexpr bool Value = true; };
+	template<> struct TIsIndexType<FIntVector3> { static constexpr bool Value = true; };
+
 	/**
 	 * Cloth Asset collection sim pattern facade class to access cloth sim pattern data.
 	 * Constructed from FCollectionClothConstFacade.
@@ -81,7 +87,7 @@ namespace UE::Chaos::ClothAsset
 		void Reset();
 
 		/** Initialize the cloth pattern using the specified 3D and 2D positions, and topology. */
-		template<typename IndexType>
+		template<typename IndexType, TEMPLATE_REQUIRES(TIsIndexType<IndexType>::Value)>
 		void Initialize(const TArray<FVector2f>& Positions2D, const TArray<FVector3f>& Positions3D, const TArray<IndexType>& Indices);
 
 		/** Initialize this pattern using another pattern collection. */
