Chaos Cloth Asset - Fixed crash in welding when two consecutive stiches have one end point in common.

#tests Ran a reproducible that was crashing, and used other asset that weren't.

diff --git a/Engine/Plugins/Experimental/ChaosClothAsset/Source/ChaosClothAsset/Private/ChaosClothAsset/CollectionClothSeamFacade.cpp b/Engine/Plugins/Experimental/ChaosClothAsset/Source/ChaosClothAsset/Private/ChaosClothAsset/CollectionClothSeamFacade.cpp
--- a/Engine/Plugins/Experimental/ChaosClothAsset/Source/ChaosClothAsset/Private/ChaosClothAsset/CollectionClothSeamFacade.cpp
+++ b/Engine/Plugins/Experimental/ChaosClothAsset/Source/ChaosClothAsset/Private/ChaosClothAsset/CollectionClothSeamFacade.cpp
@@ -74,7 +74,7 @@ namespace UE::Chaos::ClothAsset
 					// Update group1 redirected indices with the new key
 					for (TPair<int32, int32>& IndexAndWeight : *WeldingGroup1)
 					{
-						WeldingMap[IndexAndWeight.Get<0>()] = Key0;
+						WeldingMap.FindOrAdd(IndexAndWeight.Get<0>()) = Key0;  // Could be a source index for which there isn't a WeldingMap entry and needs adding
 					}
 
 					// Merge group0 & group1
