Serialize the ClothSimulationModel data

diff --git a/Engine/Plugins/Experimental/ChaosClothAsset/Source/ChaosClothAssetEngine/Private/ChaosClothAsset/ClothAsset.cpp b/Engine/Plugins/Experimental/ChaosClothAsset/Source/ChaosClothAssetEngine/Private/ChaosClothAsset/ClothAsset.cpp
--- a/Engine/Plugins/Experimental/ChaosClothAsset/Source/ChaosClothAssetEngine/Private/ChaosClothAsset/ClothAsset.cpp
+++ b/Engine/Plugins/Experimental/ChaosClothAsset/Source/ChaosClothAssetEngine/Private/ChaosClothAsset/ClothAsset.cpp
@@ -172,6 +172,13 @@ void UChaosClothAsset::Serialize(FArchive& Ar)
 			SetResourceForRendering(MakeUnique<FSkeletalMeshRenderData>());
 		}
 		GetResourceForRendering()->Serialize(Ar, this);
+
+		if (!ClothSimulationModel.IsValid())
+		{
+			ClothSimulationModel = MakeShared<FChaosClothSimulationModel>();
+		}
+		UScriptStruct* const Struct = FChaosClothSimulationModel::StaticStruct();
+		Struct->SerializeTaggedProperties(Ar, (uint8*)ClothSimulationModel.Get(), Struct, nullptr);
 	}
 }
 
