Cloth Editor: make the SkeletalMesh and Animation properties of the Preview Scene persist between editor sessions

diff --git a/Engine/Plugins/Experimental/ChaosClothAsset/Source/ChaosClothAssetEngine/Private/ChaosClothAsset/ClothAsset.cpp b/Engine/Plugins/Experimental/ChaosClothAsset/Source/ChaosClothAssetEngine/Private/ChaosClothAsset/ClothAsset.cpp
--- a/Engine/Plugins/Experimental/ChaosClothAsset/Source/ChaosClothAssetEngine/Private/ChaosClothAsset/ClothAsset.cpp
+++ b/Engine/Plugins/Experimental/ChaosClothAsset/Source/ChaosClothAssetEngine/Private/ChaosClothAsset/ClothAsset.cpp
@@ -8,6 +8,9 @@
 #include "ChaosClothAsset/ClothAssetPrivate.h"
 #include "ChaosClothAsset/ClothSimulationModel.h"
 #include "Animation/Skeleton.h"
+#if WITH_EDITORONLY_DATA
+#include "Animation/AnimationAsset.h"
+#endif
 #include "Engine/RendererSettings.h"
 #include "Engine/SkinnedAssetAsyncCompileUtils.h"
 #include "Features/IModularFeatures.h"
@@ -837,4 +840,30 @@ void UChaosClothAsset::BindSimMeshToRootBone()
 	{
 		FClothGeometryTools::BindMeshToRootBone(ClothCollection, true, false);
 	}
-}
\ No newline at end of file
+}
+
+#if WITH_EDITORONLY_DATA
+
+void UChaosClothAsset::SetPreviewSceneSkeletalMesh(USkeletalMesh* Mesh)
+{
+	PreviewSceneSkeletalMesh = Mesh;
+}
+
+USkeletalMesh* UChaosClothAsset::GetPreviewSceneSkeletalMesh() const
+{
+	// Load the SkeletalMesh asset if it's not already loaded
+	return PreviewSceneSkeletalMesh.LoadSynchronous();
+}
+
+void UChaosClothAsset::SetPreviewSceneAnimation(UAnimationAsset* Animation)
+{
+	PreviewSceneAnimation = Animation;
+}
+
+UAnimationAsset* UChaosClothAsset::GetPreviewSceneAnimation() const
+{
+	// Load the animation asset if it's not already loaded
+	return PreviewSceneAnimation.LoadSynchronous();
+}
+
+#endif
diff --git a/Engine/Plugins/Experimental/ChaosClothAsset/Source/ChaosClothAssetEngine/Public/ChaosClothAsset/ClothAsset.h b/Engine/Plugins/Experimental/ChaosClothAsset/Source/ChaosClothAssetEngine/Public/ChaosClothAsset/ClothAsset.h
--- a/Engine/Plugins/Experimental/ChaosClothAsset/Source/ChaosClothAssetEngine/Public/ChaosClothAsset/ClothAsset.h
+++ b/Engine/Plugins/Experimental/ChaosClothAsset/Source/ChaosClothAssetEngine/Public/ChaosClothAsset/ClothAsset.h
@@ -16,6 +16,7 @@ struct FChaosClothSimulationModel;
 struct FSkeletalMeshLODInfo;
 struct FManagedArrayCollection;
 struct FChaosClothAssetLodTransitionDataCache;
+class UAnimationAsset;
 
 UENUM()
 enum class EClothAssetAsyncProperties : uint64
@@ -158,6 +159,17 @@ public:
 	UPROPERTY(EditAnywhere, Category = "Dataflow")
 	FString DataflowTerminal = "ClothAssetTerminal";
 
+
+#if WITH_EDITORONLY_DATA
+
+	void SetPreviewSceneSkeletalMesh(USkeletalMesh* Mesh);
+	USkeletalMesh* GetPreviewSceneSkeletalMesh() const;
+
+	void SetPreviewSceneAnimation(UAnimationAsset* Animation);
+	UAnimationAsset* GetPreviewSceneAnimation() const;
+
+#endif
+
 private:
 	//~ Begin USkinnedAsset interface
 	/** Initial step for the Post Load process - Can't be done in parallel. */
@@ -302,5 +314,26 @@ private:
 	/** Simulation mesh Lods as fed to the solver for constraints creation. Ownership gets transferred to the proxy when it is changed during a simulation. */
 	TSharedPtr<FChaosClothSimulationModel> ClothSimulationModel;
 
+
+#if WITH_EDITORONLY_DATA
+
+	/*
+	* The following PreviewScene properties are modeled after PreviewSkeletalMesh in USkeleton
+	*	- they are inside WITH_EDITORONLY_DATA because they are not used at game runtime
+	*	- TSoftObjectPtrs since that will make it possible to avoid loading these assets until the PreviewScene asks for them
+	*	- DuplicateTransient so that if you copy a ClothAsset it won't copy these preview properties
+	*	- AssetRegistrySearchable makes it so that if the user searches the name of a PreviewScene asset in the Asset Browser, it will return any ClothAssets that use it
+	*/
+
+	/** Optional Skeletal Mesh that the cloth asset is attached to in the Preview Scene in the Cloth Editor */
+	UPROPERTY(DuplicateTransient, AssetRegistrySearchable)
+	TSoftObjectPtr<USkeletalMesh> PreviewSceneSkeletalMesh;
+
+	/** Optional animation attached to PreviewSceneSkeletalMesh in the Preview Scene in the Cloth Editor */
+	UPROPERTY(DuplicateTransient, AssetRegistrySearchable)
+	TSoftObjectPtr<UAnimationAsset> PreviewSceneAnimation;
+
+#endif
+
 	friend class UClothAssetBuilderEditor;
 };
diff --git a/Engine/Plugins/Experimental/ChaosClothAssetEditor/Source/ChaosClothAssetEditor/Private/ChaosClothAsset/ClothEditorPreviewScene.cpp b/Engine/Plugins/Experimental/ChaosClothAssetEditor/Source/ChaosClothAssetEditor/Private/ChaosClothAsset/ClothEditorPreviewScene.cpp
--- a/Engine/Plugins/Experimental/ChaosClothAssetEditor/Source/ChaosClothAssetEditor/Private/ChaosClothAsset/ClothEditorPreviewScene.cpp
+++ b/Engine/Plugins/Experimental/ChaosClothAssetEditor/Source/ChaosClothAssetEditor/Private/ChaosClothAsset/ClothEditorPreviewScene.cpp
@@ -253,6 +253,11 @@ void FChaosClothPreviewScene::SceneDescriptionPropertyChanged(const FName& Prope
 
 		UpdateSkeletalMeshAnimation();
 		UpdateClothComponentAttachment();
+
+		if (UChaosClothAsset* const ClothAsset = ClothComponent->GetClothAsset())
+		{
+			ClothAsset->SetPreviewSceneSkeletalMesh(PreviewSceneDescription->SkeletalMeshAsset);
+		}
 	}
 
 	if (PropertyName == GET_MEMBER_NAME_CHECKED(UChaosClothPreviewSceneDescription, Translation) || 
@@ -272,6 +277,11 @@ void FChaosClothPreviewScene::SceneDescriptionPropertyChanged(const FName& Prope
 			PreviewAnimInstance = nullptr;
 		}
 		UpdateSkeletalMeshAnimation();
+
+		if (UChaosClothAsset* const ClothAsset = ClothComponent->GetClothAsset())
+		{
+			ClothAsset->SetPreviewSceneAnimation(PreviewSceneDescription->AnimationAsset);
+		}
 	}
 	
 }
@@ -338,6 +348,22 @@ void FChaosClothPreviewScene::SetClothAsset(UChaosClothAsset* Asset)
 	ClothComponent->InvalidateCachedBounds();
 	FSkinnedAssetCompilingManager::Get().FinishCompilation(TArrayView<USkinnedAsset* const>{Asset});
 	ClothComponent->UpdateBounds();
+
+	if (USkeletalMesh* const SkeletalMesh = Asset->GetPreviewSceneSkeletalMesh())
+	{
+		PreviewSceneDescription->SkeletalMeshAsset = SkeletalMesh;
+
+		SkeletalMeshComponent->SetSkeletalMeshAsset(PreviewSceneDescription->SkeletalMeshAsset);
+		UpdateSkeletalMeshAnimation();
+		UpdateClothComponentAttachment();
+	}
+	
+	if (UAnimationAsset* const Animation = Asset->GetPreviewSceneAnimation())
+	{
+		PreviewSceneDescription->AnimationAsset = Animation;
+
+		UpdateSkeletalMeshAnimation();
+	}
 }
 
 UAnimSingleNodeInstance* FChaosClothPreviewScene::GetPreviewAnimInstance()
