Chaos Cloth Asset - Moved the Cloth Asset builder outside of the editor plugin into the engine module.

diff --git a/Engine/Plugins/Experimental/ChaosClothAsset/Source/ChaosClothAssetEngine/Private/ChaosClothAsset/ClothAsset.cpp b/Engine/Plugins/Experimental/ChaosClothAsset/Source/ChaosClothAssetEngine/Private/ChaosClothAsset/ClothAsset.cpp
--- a/Engine/Plugins/Experimental/ChaosClothAsset/Source/ChaosClothAssetEngine/Private/ChaosClothAsset/ClothAsset.cpp
+++ b/Engine/Plugins/Experimental/ChaosClothAsset/Source/ChaosClothAssetEngine/Private/ChaosClothAsset/ClothAsset.cpp
@@ -485,12 +485,10 @@ void UChaosClothAsset::Build(TArray<FChaosClothAssetLodTransitionDataCache>* InO
 	// Build simulation model
 	BuildClothSimulationModel(InOutTransitionCache);
 
+#if WITH_EDITOR
 	// Rebuild LOD Model
-#if WITH_EDITORONLY_DATA
 	BuildMeshModel();
-#endif
 
-#if WITH_EDITOR
 	// Load/save render data from/to DDC
 	ExecuteBuildInternal(Context);
 #endif
@@ -553,32 +551,25 @@ void UChaosClothAsset::FinishBuildInternal(FSkinnedAssetBuildContext& Context)
 }
 #endif // #if WITH_EDITOR
 
-#if WITH_EDITORONLY_DATA
+#if WITH_EDITOR
 void UChaosClothAsset::BuildMeshModel()
 {
-	const TArray<IClothAssetBuilderClassProvider*> ClassProviders = IModularFeatures::Get().GetModularFeatureImplementations<IClothAssetBuilderClassProvider>(IClothAssetBuilderClassProvider::FeatureName);
-	if (const TSubclassOf<UClothAssetBuilder> ClothAssetBuilderClass = ClassProviders.Num() ? ClassProviders[0]->GetClothAssetBuilderClass() : nullptr)
-	{
-		if (const UClothAssetBuilder* const ClothAssetBuilder = ClothAssetBuilderClass->GetDefaultObject<UClothAssetBuilder>())
-		{
-			using namespace UE::Chaos::ClothAsset;
+	using namespace UE::Chaos::ClothAsset;
 
-			const int32 NumLods = ClothCollections.Num();
+	const int32 NumLods = ClothCollections.Num();
 
-			// Build each LOD
-			check(MeshModel);  // MeshModel should always be created in the Cloth Asset constructor WITH_EDITORONLY_DATA
-			MeshModel->LODModels.Empty();
+	// Build each LOD
+	check(MeshModel);  // MeshModel should always be created in the Cloth Asset constructor WITH_EDITORONLY_DATA
+	MeshModel->LODModels.Empty();
 
-			for (int32 LodIndex = 0; LodIndex < NumLods; ++LodIndex)
-			{
-				FSkeletalMeshLODModel* LODModel = new FSkeletalMeshLODModel();
-				ClothAssetBuilder->BuildLod(*LODModel, *this, LodIndex);
-				MeshModel->LODModels.Add(LODModel);
-			}
-		}
+	for (int32 LodIndex = 0; LodIndex < NumLods; ++LodIndex)
+	{
+		FSkeletalMeshLODModel* LODModel = new FSkeletalMeshLODModel();
+		FBuilder::BuildLod(*LODModel, *this, LodIndex);
+		MeshModel->LODModels.Add(LODModel);
 	}
 }
-#endif  // #if WITH_EDITORONLY_DATA
+#endif  // #if WITH_EDITOR
 
 void UChaosClothAsset::BuildClothSimulationModel(TArray<FChaosClothAssetLodTransitionDataCache>* InOutTransitionCache)
 {
diff --git a/Engine/Plugins/Experimental/ChaosClothAsset/Source/ChaosClothAssetEngine/Private/ChaosClothAsset/ClothAssetBuilder.cpp b/Engine/Plugins/Experimental/ChaosClothAsset/Source/ChaosClothAssetEngine/Private/ChaosClothAsset/ClothAssetBuilder.cpp
--- a/Engine/Plugins/Experimental/ChaosClothAsset/Source/ChaosClothAssetEngine/Private/ChaosClothAsset/ClothAssetBuilder.cpp
+++ b/Engine/Plugins/Experimental/ChaosClothAsset/Source/ChaosClothAssetEngine/Private/ChaosClothAsset/ClothAssetBuilder.cpp
@@ -1,30 +1,356 @@
 // Copyright Epic Games, Inc. All Rights Reserved.
 
+#if WITH_EDITOR
+
 #include "ChaosClothAsset/ClothAssetBuilder.h"
 #include "ChaosClothAsset/ClothAsset.h"
 #include "ChaosClothAsset/ClothSimulationModel.h"
+#include "ChaosClothAsset/CollectionClothFacade.h"
+#include "Chaos/CollectionPropertyFacade.h"
+#include "Engine/SkeletalMesh.h"
+#include "Engine/SkinnedAssetCommon.h"
+#include "Misc/ScopedSlowTask.h"
+#include "Modules/ModuleManager.h"
+#include "Rendering/SkeletalMeshLODModel.h"
+#include "Utils/ClothingMeshUtils.h"
+#include "BoneWeights.h"
+#include "MeshUtilities.h"
+#include "PointWeightMap.h"
 
-int32 UClothAssetBuilder::GetNumVertices(const UChaosClothAsset& ClothAsset, int32 LodIndex)
-{
-	const TSharedPtr<const FChaosClothSimulationModel> ClothSimulationModel = ClothAsset.GetClothSimulationModel();
-	return ClothSimulationModel && ClothSimulationModel->IsValidLodIndex(LodIndex) ?
-		ClothSimulationModel->GetNumVertices(LodIndex) :
-		0;
-}
+#define LOCTEXT_NAMESPACE "ClothAssetBuilderEditor"
 
-TConstArrayView<FVector3f> UClothAssetBuilder::GetSimPositions(const UChaosClothAsset& ClothAsset, int32 LodIndex)
+namespace UE::Chaos::ClothAsset::Private
 {
-	const TSharedPtr<const FChaosClothSimulationModel> ClothSimulationModel = ClothAsset.GetClothSimulationModel();
-	return ClothSimulationModel && ClothSimulationModel->IsValidLodIndex(LodIndex) ?
-		ClothSimulationModel->GetPositions(LodIndex) :
-		TConstArrayView<FVector3f>();
-}
+	int32 GetNumVertices(const UChaosClothAsset& ClothAsset, int32 LodIndex)
+	{
+		const TSharedPtr<const FChaosClothSimulationModel> ClothSimulationModel = ClothAsset.GetClothSimulationModel();
+		return ClothSimulationModel && ClothSimulationModel->IsValidLodIndex(LodIndex) ?
+			ClothSimulationModel->GetNumVertices(LodIndex) :
+			0;
+	}
+
+	TConstArrayView<FVector3f> GetSimPositions(const UChaosClothAsset& ClothAsset, int32 LodIndex)
+	{
+		const TSharedPtr<const FChaosClothSimulationModel> ClothSimulationModel = ClothAsset.GetClothSimulationModel();
+		return ClothSimulationModel && ClothSimulationModel->IsValidLodIndex(LodIndex) ?
+			ClothSimulationModel->GetPositions(LodIndex) :
+			TConstArrayView<FVector3f>();
+	}
 
-TConstArrayView<uint32> UClothAssetBuilder::GetSimIndices(const UChaosClothAsset& ClothAsset, int32 LodIndex)
+	TConstArrayView<uint32> GetSimIndices(const UChaosClothAsset& ClothAsset, int32 LodIndex)
+	{
+		const TSharedPtr<const FChaosClothSimulationModel> ClothSimulationModel = ClothAsset.GetClothSimulationModel();
+		return ClothSimulationModel && ClothSimulationModel->IsValidLodIndex(LodIndex) ?
+			ClothSimulationModel->GetIndices(LodIndex) :
+			TConstArrayView<uint32>();
+	}
+}  // End namespace UE::Chaos::ClothAsset::Private
+
+void UChaosClothAsset::FBuilder::BuildLod(FSkeletalMeshLODModel& LODModel, const UChaosClothAsset& ClothAsset, int32 LodIndex)
 {
-	const TSharedPtr<const FChaosClothSimulationModel> ClothSimulationModel = ClothAsset.GetClothSimulationModel();
-	return ClothSimulationModel && ClothSimulationModel->IsValidLodIndex(LodIndex) ?
-		ClothSimulationModel->GetIndices(LodIndex) :
-		TConstArrayView<uint32>();
+	using namespace UE::Chaos::ClothAsset;
+	using namespace ::Chaos::Softs;
+
+	// Start from an empty LODModel
+	LODModel.Empty();
+
+	// Clear the mesh infos, none are stored on this asset
+	LODModel.MaxImportVertex = 0;
+
+	// Set 1 texture coordinate
+	LODModel.NumTexCoords = 1;
+
+	// Init the size of the vertex buffer
+	LODModel.NumVertices = 0;
+
+	// Create a table to remap the LOD materials to the asset materials
+	const TArray<FSkeletalMaterial>& Materials = ClothAsset.GetMaterials();
+
+	const TSharedRef<const FManagedArrayCollection> ClothCollection = ClothAsset.GetClothCollections()[LodIndex];
+
+	const FCollectionClothConstFacade ClothFacade(ClothCollection);
+
+	// Keep track of the active bone indices for this LOD model
+	TSet<FBoneIndexType> ActiveBoneIndices;
+	ActiveBoneIndices.Reserve(ClothAsset.RefSkeleton.GetNum());
+
+	// Load the mesh utilities module used to optimized the index buffer
+	IMeshUtilities& MeshUtilities = FModuleManager::Get().LoadModuleChecked<IMeshUtilities>("MeshUtilities");
+
+	// Build the sim mesh descriptor for creation of the sections' mesh to mesh mapping data
+	const ClothingMeshUtils::ClothMeshDesc SourceMesh(
+		Private::GetSimPositions(ClothAsset, LodIndex),
+		Private::GetSimIndices(ClothAsset, LodIndex));  // Let it calculate the averaged normals as to match the simulation data output
+	const bool bIsValidSourceMesh = (SourceMesh.GetPositions().Num() > 0);
+
+	const int32 NumLodSimVertices = Private::GetNumVertices(ClothAsset, LodIndex);
+
+	// Retrieve MaxDistance information (weight map and Low/High values)
+	FCollectionPropertyConstFacade Properties(ClothCollection);
+
+	int32 MaxDistancePropertyKeyIndex;
+	FString MaxDistanceString = TEXT("MaxDistance");
+	MaxDistanceString = Properties.GetStringValue(MaxDistanceString, MaxDistanceString, &MaxDistancePropertyKeyIndex);
+	const bool bHasMaxDistanceProperty = (MaxDistancePropertyKeyIndex != INDEX_NONE);
+	const float MaxDistanceOffset = bHasMaxDistanceProperty ? Properties.GetLowValue<float>(MaxDistancePropertyKeyIndex) : TNumericLimits<float>::Max();  // Uses infinite distance when no MaxDistance properties are set
+	const float MaxDistanceScale = bHasMaxDistanceProperty ? Properties.GetHighValue<float>(MaxDistancePropertyKeyIndex) - MaxDistanceOffset : 0.f;
+	const TConstArrayView<float> MaxDistanceWeightMap = ClothFacade.GetWeightMap(FName(MaxDistanceString));
+
+	const FPointWeightMap MaxDistances = (MaxDistanceWeightMap.Num() == NumLodSimVertices) ?
+		FPointWeightMap(MaxDistanceWeightMap, MaxDistanceOffset, MaxDistanceScale) :
+		FPointWeightMap(NumLodSimVertices, MaxDistanceOffset);
+
+	const int32 NumRenderVertices = ClothFacade.GetNumRenderVertices();
+	LODModel.MeshToImportVertexMap.Reserve(NumRenderVertices);
+
+	// Populate this LOD's sections and the LOD index buffer
+	const int32 NumSections = ClothFacade.GetNumRenderPatterns();  // Cloth Render Patterns == Skeletal Mesh Sections
+	LODModel.Sections.SetNum(NumSections);
+
+	FScopedSlowTask SlowTask((float)ClothFacade.GetNumRenderFaces(), LOCTEXT("ClothAssetBuildLOD", "Building Cloth Asset LOD sections..."));
+	SlowTask.MakeDialogDelayed(1.f);
+
+	int32 BaseIndex = 0;
+	for (int32 SectionIndex = 0; SectionIndex < NumSections; ++SectionIndex)
+	{
+		FCollectionClothRenderPatternConstFacade RenderPatternFacade = ClothFacade.GetRenderPattern(SectionIndex);
+		FSkelMeshSection& Section = LODModel.Sections[SectionIndex];
+
+		Section.OriginalDataSectionIndex = SectionIndex;
+
+		const int32 MaterialIndex = SectionIndex;
+		const int32 NumFaces = RenderPatternFacade.GetNumRenderFaces();
+		const int32 NumIndices = NumFaces * 3;
+
+		// Build the section face data (indices)
+		Section.MaterialIndex = (uint16)MaterialIndex;
+
+		Section.BaseIndex = (uint32)BaseIndex;
+		BaseIndex += NumIndices;
+
+		Section.NumTriangles = (uint32)NumFaces;
+
+		TArray<uint32> Indices;
+		Indices.SetNumUninitialized(NumIndices);
+
+		const TConstArrayView<FIntVector3> PatternRenderIndices = RenderPatternFacade.GetRenderIndices();
+
+		for (int32 FaceIndex = 0; FaceIndex < NumFaces; ++FaceIndex)
+		{
+			const FIntVector3& RenderIndices = PatternRenderIndices[FaceIndex];
+			for (int32 VertexIndex = 0; VertexIndex < 3; ++VertexIndex)
+			{
+				const int32 RenderIndex = RenderIndices[VertexIndex];
+				Indices[FaceIndex * 3 + VertexIndex] = (uint32)RenderIndex;
+			}
+		}
+
+		MeshUtilities.CacheOptimizeIndexBuffer(Indices);
+
+		LODModel.IndexBuffer.Append(MoveTemp(Indices));
+
+		// Build the section vertex data 
+		const int32 NumVertices = RenderPatternFacade.GetNumRenderVertices();
+
+		Section.SoftVertices.SetNumUninitialized(NumVertices);
+		Section.NumVertices = NumVertices;
+		Section.BaseVertexIndex = LODModel.NumVertices;
+		LODModel.NumVertices += (uint32)NumVertices;
 
+		const int32 PatternVertexOffset = RenderPatternFacade.GetRenderVerticesOffset();
+		checkSlow(Section.BaseVertexIndex == PatternVertexOffset); // Otherwise Section's Indices don't match RenderPattern indices
+
+		// Map reference skeleton bone index to the index in the section's bone map
+		TMap<FBoneIndexType, FBoneIndexType> ReferenceToSectionBoneMap;
+
+		// Track how many bones we added to the section's bone map so far
+		int CurSectionBoneMapNum = 0;
+
+		const TConstArrayView<FVector3f> PatternRenderPosition = RenderPatternFacade.GetRenderPosition();
+		const TConstArrayView<FVector3f> PatternRenderTangentU = RenderPatternFacade.GetRenderTangentU();
+		const TConstArrayView<FVector3f> PatternRenderTangentV = RenderPatternFacade.GetRenderTangentV();
+		const TConstArrayView<FVector3f> PatternRenderNormal = RenderPatternFacade.GetRenderNormal();
+		const TConstArrayView<FLinearColor> PatternRenderColor = RenderPatternFacade.GetRenderColor();
+		const TConstArrayView<TArray<FVector2f>> PatternRenderUVs = RenderPatternFacade.GetRenderUVs();
+		const TConstArrayView<TArray<int32>> PatternRenderBoneIndices = RenderPatternFacade.GetRenderBoneIndices();
+		const TConstArrayView<TArray<float>> PatternRenderBoneWeights = RenderPatternFacade.GetRenderBoneWeights();
+
+		for (int32 VertexIndex = 0; VertexIndex < NumVertices; ++VertexIndex)
+		{
+			// Save the original indices for the newly added vertices
+			LODModel.MeshToImportVertexMap.Add(VertexIndex + PatternVertexOffset);
+			LODModel.MaxImportVertex = VertexIndex + PatternVertexOffset;
+
+			FSoftSkinVertex& SoftVertex = Section.SoftVertices[VertexIndex];
+
+			SoftVertex.Position = PatternRenderPosition[VertexIndex];
+			SoftVertex.TangentX = PatternRenderTangentU[VertexIndex];
+			SoftVertex.TangentY = PatternRenderTangentV[VertexIndex];
+			SoftVertex.TangentZ = PatternRenderNormal[VertexIndex];
+
+			constexpr bool bSRGB = false; // Avoid linear to srgb conversion
+			SoftVertex.Color = PatternRenderColor[VertexIndex].ToFColor(bSRGB);
+
+			const TArray<FVector2f>& RenderUVs = PatternRenderUVs[VertexIndex];
+			for (int32 TexCoord = 0; TexCoord < FMath::Min(RenderUVs.Num(), (int32)MAX_TEXCOORDS); ++TexCoord)
+			{
+				SoftVertex.UVs[TexCoord] = RenderUVs[TexCoord];
+			}
+
+			const int32 NumInfluences = PatternRenderBoneIndices[VertexIndex].Num();
+			check(NumInfluences <= MAX_TOTAL_INFLUENCES);
+
+			// Add all of the bones that have non-zero influence to the section's bone map and keep track of the order
+			// that we added the reference bone via CurSectionBoneMapNum
+			for (int32 Influence = 0; Influence < NumInfluences; ++Influence)
+			{
+				const FBoneIndexType InfluenceBone = (FBoneIndexType)PatternRenderBoneIndices[VertexIndex][Influence];
+
+				if (ReferenceToSectionBoneMap.Contains(InfluenceBone) == false)
+				{
+					ReferenceToSectionBoneMap.Add(InfluenceBone, CurSectionBoneMapNum);
+					++CurSectionBoneMapNum;
+				}
+			}
+
+			int32 Influence = 0;
+			for (; Influence < NumInfluences; ++Influence)
+			{
+				const FBoneIndexType InfluenceBone = (FBoneIndexType)PatternRenderBoneIndices[VertexIndex][Influence];
+				const float InWeight = PatternRenderBoneWeights[VertexIndex][Influence];
+				const uint16 InfluenceWeight = static_cast<uint16>(InWeight * static_cast<float>(UE::AnimationCore::MaxRawBoneWeight) + 0.5f);
+
+				// FSoftSkinVertex::InfluenceBones contain indices into the section's bone map and not the reference
+				// skeleton, so we need to remap
+				const FBoneIndexType* const MappedIndexPtr = ReferenceToSectionBoneMap.Find(InfluenceBone);
+
+				// ReferenceToSectionBoneMap should always contain InfluenceBone since it was added above
+				checkSlow(MappedIndexPtr);
+				if (MappedIndexPtr != nullptr)
+				{
+					SoftVertex.InfluenceBones[Influence] = *MappedIndexPtr;
+					SoftVertex.InfluenceWeights[Influence] = InfluenceWeight;
+				}
+			}
+
+			for (; Influence < MAX_TOTAL_INFLUENCES; ++Influence)
+			{
+				SoftVertex.InfluenceBones[Influence] = 0;
+				SoftVertex.InfluenceWeights[Influence] = 0;
+			}
+		}
+
+		// Initialize the section bone map
+		Section.BoneMap.SetNumUninitialized(ReferenceToSectionBoneMap.Num());
+		for (const TPair<FBoneIndexType, FBoneIndexType>& Pair : ReferenceToSectionBoneMap)
+		{
+			Section.BoneMap[Pair.Value] = Pair.Key;
+		}
+
+		ActiveBoneIndices.Append(Section.BoneMap);
+
+		// Update max bone influences
+		Section.CalcMaxBoneInfluences();
+		Section.CalcUse16BitBoneIndex();
+
+		// Setup clothing data
+		if (bIsValidSourceMesh)
+		{
+			Section.ClothMappingDataLODs.SetNum(1);  // TODO: LODBias maps for raytracing
+
+			Section.ClothingData.AssetLodIndex = LodIndex;
+			Section.ClothingData.AssetGuid = ClothAsset.AssetGuid;  // There is only one cloth asset,
+			Section.CorrespondClothAssetIndex = 0;       // this one
+
+			TArray<FVector3f> SectionRenderPositions;
+			TArray<FVector3f> SectionRenderNormals;
+			TArray<FVector3f> SectionRenderTangents;
+			SectionRenderPositions.Reserve(NumVertices);
+			SectionRenderNormals.Reserve(NumVertices);
+			SectionRenderTangents.Reserve(NumVertices);
+			for (const FSoftSkinVertex& SoftVert : Section.SoftVertices)
+			{
+				SectionRenderPositions.Add(SoftVert.Position);
+				SectionRenderNormals.Add(SoftVert.TangentZ);
+				SectionRenderTangents.Add(SoftVert.TangentX);
+			}
+			TArray<uint32> SectionRenderIndices;
+			SectionRenderIndices.Reserve(NumIndices);
+			const TArrayView<uint32> SectionIndexBuffer(LODModel.IndexBuffer.GetData() + Section.BaseIndex, NumIndices);
+			for (const uint32 LodModelVertIndex : SectionIndexBuffer)
+			{
+				SectionRenderIndices.Add(LodModelVertIndex - Section.BaseVertexIndex);
+			}
+
+			const ClothingMeshUtils::ClothMeshDesc TargetMesh(
+				SectionRenderPositions,
+				SectionRenderNormals,
+				SectionRenderTangents,
+				SectionRenderIndices);
+
+			ClothingMeshUtils::GenerateMeshToMeshVertData(
+				Section.ClothMappingDataLODs[0],
+				TargetMesh,
+				SourceMesh,
+				&MaxDistances,
+				ClothAsset.bSmoothTransition,
+				ClothAsset.bUseMultipleInfluences,
+				ClothAsset.SkinningKernelRadius);
+		}
+
+		// Compute the overlapping vertices map (inspired from MeshUtilities::BuildSkeletalMesh)
+		const TArray<FSoftSkinVertex>& SoftVertices = Section.SoftVertices;
+
+		typedef TPair<float, int32> FIndexAndZ;  // Acceleration structure, list of vertex Z / index pairs
+		TArray<FIndexAndZ> IndexAndZs;
+		IndexAndZs.Reserve(NumVertices);
+		for (int32 VertexIndex = 0; VertexIndex < NumVertices; ++VertexIndex)
+		{
+			const FVector3f& Position = SoftVertices[VertexIndex].Position;
+
+			const float Z = 0.30f * Position.X + 0.33f * Position.Y + 0.37f * Position.Z;
+			IndexAndZs.Emplace(Z, VertexIndex);
+		}
+		IndexAndZs.Sort([](const FIndexAndZ& A, const FIndexAndZ& B) { return A.Key < B.Key; });
+
+		for (int32 Index0 = 0; Index0 < IndexAndZs.Num(); ++Index0)
+		{
+			const float Z0 = IndexAndZs[Index0].Key;
+			const uint32 VertexIndex0 = IndexAndZs[Index0].Value;
+			const FVector3f& Position0 = SoftVertices[VertexIndex0].Position;
+
+			// Only need to search forward, since we add pairs both ways
+			for (int32 Index1 = Index0 + 1; Index1 < IndexAndZs.Num() && FMath::Abs(IndexAndZs[Index1].Key - Z0) <= THRESH_POINTS_ARE_SAME; ++Index1)
+			{
+				const uint32 VertexIndex1 = IndexAndZs[Index1].Value;
+				const FVector3f& Position1 = SoftVertices[VertexIndex1].Position;
+
+				if (PointsEqual(Position0, Position1))
+				{
+					// Add to the overlapping map
+					TArray<int32>& SrcValueArray = Section.OverlappingVertices.FindOrAdd(VertexIndex0);
+					SrcValueArray.Add(VertexIndex1);
+
+					TArray<int32>& IterValueArray = Section.OverlappingVertices.FindOrAdd(VertexIndex1);
+					IterValueArray.Add(VertexIndex0);
+				}
+			}
+		}
+
+		// Copy to user section data, otherwise the section data set above would get lost when the user section gets synced
+		FSkelMeshSourceSectionUserData::GetSourceSectionUserData(LODModel.UserSectionsData, Section);
+
+		SlowTask.EnterProgressFrame(NumFaces);
+	}
+
+	// Update the active bone indices on the LOD model
+	LODModel.ActiveBoneIndices = ActiveBoneIndices.Array();
+
+	// Ensure parent exists with incoming active bone indices, and the result should be sorted
+	ClothAsset.RefSkeleton.EnsureParentsExistAndSort(LODModel.ActiveBoneIndices);
+
+	// Compute the required bones for this model.
+	USkeletalMesh::CalculateRequiredBones(LODModel, ClothAsset.RefSkeleton, nullptr);
 }
+
+#endif  // #if WITH_EDITOR
diff --git a/Engine/Plugins/Experimental/ChaosClothAsset/Source/ChaosClothAssetEngine/Private/ChaosClothAsset/ClothAssetBuilder.h b/Engine/Plugins/Experimental/ChaosClothAsset/Source/ChaosClothAssetEngine/Private/ChaosClothAsset/ClothAssetBuilder.h
new file mode 100644
--- /dev/null
+++ b/Engine/Plugins/Experimental/ChaosClothAsset/Source/ChaosClothAssetEngine/Private/ChaosClothAsset/ClothAssetBuilder.h
@@ -0,0 +1,18 @@
+// Copyright Epic Games, Inc. All Rights Reserved.
+
+#pragma once
+
+#if WITH_EDITOR
+
+#include "ChaosClothAsset/ClothAsset.h"
+
+class FSkeletalMeshLODModel;
+
+/** Builder utility struct, nested struct of UChaosClothAsset. */
+struct UChaosClothAsset::FBuilder
+{
+	/** Build a FSkeletalMeshLODModel out of the cloth asset for the specified LOD index. */
+	static void BuildLod(FSkeletalMeshLODModel& LODModel, const UChaosClothAsset& ClothAsset, int32 LodIndex);
+};
+
+#endif  // #if WITH_EDITOR
diff --git a/Engine/Plugins/Experimental/ChaosClothAsset/Source/ChaosClothAssetEngine/Public/ChaosClothAsset/ClothAsset.h b/Engine/Plugins/Experimental/ChaosClothAsset/Source/ChaosClothAssetEngine/Public/ChaosClothAsset/ClothAsset.h
--- a/Engine/Plugins/Experimental/ChaosClothAsset/Source/ChaosClothAssetEngine/Public/ChaosClothAsset/ClothAsset.h
+++ b/Engine/Plugins/Experimental/ChaosClothAsset/Source/ChaosClothAssetEngine/Public/ChaosClothAsset/ClothAsset.h
@@ -198,7 +198,7 @@ private:
 	UE_DEPRECATED(5.3, "Use Build() instead")
 	void UpdateSkeleton(bool bRebuildClothSimulationModel = true);
 
-#if WITH_EDITORONLY_DATA
+#if WITH_EDITOR
 	/** Build the SkeletalMeshLODModel for this asset. */
 	void BuildMeshModel();
 #endif
@@ -335,5 +335,7 @@ private:
 
 #endif
 
-	friend class UClothAssetBuilderEditor;
+#if WITH_EDITOR
+	struct FBuilder;
+#endif
 };
diff --git a/Engine/Plugins/Experimental/ChaosClothAsset/Source/ChaosClothAssetEngine/Public/ChaosClothAsset/ClothAssetBuilder.h b/Engine/Plugins/Experimental/ChaosClothAsset/Source/ChaosClothAssetEngine/Public/ChaosClothAsset/ClothAssetBuilder.h
deleted file mode 100644
diff --git a/Engine/Plugins/Experimental/ChaosClothAssetEditor/Source/ChaosClothAssetTools/Private/ChaosClothAsset/ChaosClothAssetToolsModule.cpp b/Engine/Plugins/Experimental/ChaosClothAssetEditor/Source/ChaosClothAssetTools/Private/ChaosClothAsset/ChaosClothAssetToolsModule.cpp
--- a/Engine/Plugins/Experimental/ChaosClothAssetEditor/Source/ChaosClothAssetTools/Private/ChaosClothAsset/ChaosClothAssetToolsModule.cpp
+++ b/Engine/Plugins/Experimental/ChaosClothAssetEditor/Source/ChaosClothAssetTools/Private/ChaosClothAsset/ChaosClothAssetToolsModule.cpp
@@ -1,6 +1,5 @@
 // Copyright Epic Games, Inc. All Rights Reserved.
 
-#include "ChaosClothAsset/ClothAssetBuilderEditor.h"
 #include "Features/IModularFeatures.h"
 #include "Modules/ModuleManager.h"
 #include "ClothAssetEditorStyle.h"
@@ -9,7 +8,7 @@
 
 namespace UE::Chaos::ClothAsset
 {
-	class FChaosClothAssetToolsModule : public IModuleInterface, public IClothingAssetExporterClassProvider, public IClothAssetBuilderClassProvider
+	class FChaosClothAssetToolsModule : public IModuleInterface, public IClothingAssetExporterClassProvider
 	{
 	public:
 		// IModuleInterface implementation
@@ -21,7 +20,6 @@ namespace UE::Chaos::ClothAsset
 
 			// Register modular features
 			IModularFeatures::Get().RegisterModularFeature(IClothingAssetExporterClassProvider::FeatureName, static_cast<IClothingAssetExporterClassProvider*>(this));
-			IModularFeatures::Get().RegisterModularFeature(IClothAssetBuilderClassProvider::FeatureName, static_cast<IClothAssetBuilderClassProvider*>(this));
 		}
 
 		virtual void ShutdownModule() override
@@ -30,7 +28,6 @@ namespace UE::Chaos::ClothAsset
 			{
 				// Unregister modular features
 				IModularFeatures::Get().UnregisterModularFeature(IClothingAssetExporterClassProvider::FeatureName, static_cast<IClothingAssetExporterClassProvider*>(this));
-				IModularFeatures::Get().UnregisterModularFeature(IClothAssetBuilderClassProvider::FeatureName, static_cast<IClothAssetBuilderClassProvider*>(this));
 			}
 		}
 
@@ -39,12 +36,6 @@ namespace UE::Chaos::ClothAsset
 		{
 			return UClothingAssetToChaosClothAssetExporter::StaticClass();
 		}
-
-		// IClothAssetBuilderClassProvider implementation
-		virtual TSubclassOf<UClothAssetBuilder> GetClothAssetBuilderClass() const override
-		{
-			return UClothAssetBuilderEditor::StaticClass();
-		}
 	};
 }
 
diff --git a/Engine/Plugins/Experimental/ChaosClothAssetEditor/Source/ChaosClothAssetTools/Private/ChaosClothAsset/ClothAssetBuilderEditor.cpp b/Engine/Plugins/Experimental/ChaosClothAssetEditor/Source/ChaosClothAssetTools/Private/ChaosClothAsset/ClothAssetBuilderEditor.cpp
deleted file mode 100644
diff --git a/Engine/Plugins/Experimental/ChaosClothAssetEditor/Source/ChaosClothAssetTools/Private/ChaosClothAsset/ClothAssetBuilderEditor.h b/Engine/Plugins/Experimental/ChaosClothAssetEditor/Source/ChaosClothAssetTools/Private/ChaosClothAsset/ClothAssetBuilderEditor.h
deleted file mode 100644
