Fix missing undef for LOCTEXT_NAMESPACE

diff --git a/Engine/Plugins/Experimental/ChaosClothAsset/Source/ChaosClothAssetEngine/Private/ChaosClothAsset/ClothAssetBuilder.cpp b/Engine/Plugins/Experimental/ChaosClothAsset/Source/ChaosClothAssetEngine/Private/ChaosClothAsset/ClothAssetBuilder.cpp
--- a/Engine/Plugins/Experimental/ChaosClothAsset/Source/ChaosClothAssetEngine/Private/ChaosClothAsset/ClothAssetBuilder.cpp
+++ b/Engine/Plugins/Experimental/ChaosClothAsset/Source/ChaosClothAssetEngine/Private/ChaosClothAsset/ClothAssetBuilder.cpp
@@ -353,4 +353,5 @@ void UChaosClothAsset::FBuilder::BuildLod(FSkeletalMeshLODModel& LODModel, const
 	USkeletalMesh::CalculateRequiredBones(LODModel, ClothAsset.RefSkeleton, nullptr);
 }
 
+#undef LOCTEXT_NAMESPACE
 #endif  // #if WITH_EDITOR
