Chaos Cloth: Add  self collision layer editing. Also expose SelfCollisionDisableNeighborDistance distance rather than hard-coding it to 5.

The self collisions node takes in an IntArray indexed on SimFaces as the collision layers. This can be generated by making SimFace selections and converting them using the new SelectionToIntMapNode. A lot of this CL is just piping through the new FaceIntMap data.

diff --git a/Engine/Plugins/Experimental/ChaosCloth/Source/Chaos/Private/ChaosCloth/ChaosClothConstraints.cpp b/Engine/Plugins/Experimental/ChaosCloth/Source/Chaos/Private/ChaosCloth/ChaosClothConstraints.cpp
--- a/Engine/Plugins/Experimental/ChaosCloth/Source/Chaos/Private/ChaosCloth/ChaosClothConstraints.cpp
+++ b/Engine/Plugins/Experimental/ChaosCloth/Source/Chaos/Private/ChaosCloth/ChaosClothConstraints.cpp
@@ -470,7 +470,7 @@ void FClothConstraints::AddRules(
 
 	// Call new AddRules function
 	AddRules(ConfigProperties, TriangleMesh, nullptr, WeightMaps, TMap<FString, const TSet<int32>*>(),
-		Tethers, MeshScale, bEnabled);
+		TMap<FString, TConstArrayView<int32>>(), Tethers, MeshScale, bEnabled);
 }
 
 void FClothConstraints::AddRules(
@@ -482,7 +482,7 @@ void FClothConstraints::AddRules(
 	Softs::FSolverReal MeshScale, bool bEnabled)
 {
 	AddRules(ConfigProperties, TriangleMesh, PatternData, WeightMaps, TMap<FString, const TSet<int32>*>(),
-		Tethers, MeshScale, bEnabled);
+		TMap<FString, TConstArrayView<int32>>(), Tethers, MeshScale, bEnabled);
 }
 
 //only counts the number of rules that GS/CG currently cupports
@@ -797,11 +797,12 @@ void FClothConstraints::AddRules(
 	const FClothingPatternData* PatternData,
 	const TMap<FString, TConstArrayView<FRealSingle>>& WeightMaps,
 	const TMap<FString, const TSet<int32>*>& VertexSets,
+	const TMap<FString, TConstArrayView<int32>>& FaceIntMaps,
 	const TArray<TConstArrayView<TTuple<int32, int32, FRealSingle>>>& Tethers,
 	Softs::FSolverReal MeshScale, bool bEnabled)
 {
 	// Self collisions
-	CreateSelfCollisionConstraints(ConfigProperties, VertexSets, TriangleMesh);
+	CreateSelfCollisionConstraints(ConfigProperties, VertexSets, FaceIntMaps, TriangleMesh);
 
 	// Edge constraints
 	CreateStretchConstraints(ConfigProperties, WeightMaps, TriangleMesh, PatternData);
@@ -867,27 +868,13 @@ void FClothConstraints::AddRules(
 }
 
 void FClothConstraints::CreateSelfCollisionConstraints(const Softs::FCollectionPropertyConstFacade& ConfigProperties, 
-	const TMap<FString, const TSet<int32>*>& VertexSets, const FTriangleMesh& TriangleMesh)
+	const TMap<FString, const TSet<int32>*>& VertexSets,
+	const TMap<FString, TConstArrayView<int32>>& FaceIntMaps, const FTriangleMesh& TriangleMesh)
 {
 	const bool bUseSelfCollisions = ConfigProperties.GetValue<bool>(TEXT("UseSelfCollisions"));
 
 	if (bUseSelfCollisions)
 	{
-		static const int32 DisabledCollisionElementsN = 5;  // TODO: Make this a parameter?
-		TSet<TVec2<int32>> DisabledCollisionElements;  // TODO: Is this needed? Turn this into a bit array?
-
-		const TVec2<int32> Range = TriangleMesh.GetVertexRange();
-		for (int32 Index = Range[0]; Index <= Range[1]; ++Index)
-		{
-			const TSet<int32> Neighbors = TriangleMesh.GetNRing(Index, DisabledCollisionElementsN);
-			for (int32 Element : Neighbors)
-			{
-				check(Index != Element);
-				DisabledCollisionElements.Emplace(TVec2<int32>(Index, Element));
-				DisabledCollisionElements.Emplace(TVec2<int32>(Element, Index));
-			}
-		}
-
 		SelfCollisionInit = MakeShared<Softs::FPBDTriangleMeshCollisions>(
 			ParticleOffset,
 			NumParticles,
@@ -900,7 +887,7 @@ void FClothConstraints::CreateSelfCollisionConstraints(const Softs::FCollectionP
 			NumParticles,
 			TriangleMesh,
 			AnimationPositions,
-			MoveTemp(DisabledCollisionElements),
+			FaceIntMaps,
 			ConfigProperties);
 		++NumPostCollisionConstraintRules;
 
@@ -2106,6 +2093,7 @@ void FClothConstraints::Update(
 	const Softs::FCollectionPropertyConstFacade& ConfigProperties,
 	const TMap<FString, TConstArrayView<FRealSingle>>& WeightMaps,
 	const TMap<FString, const TSet<int32>*>& VertexSets,
+	const TMap<FString, TConstArrayView<int32>>& FaceIntMaps,
 	Softs::FSolverReal MeshScale,
 	Softs::FSolverReal MaxDistancesScale)
 {
@@ -2171,7 +2159,7 @@ void FClothConstraints::Update(
 	}
 	if (SelfCollisionConstraints)
 	{
-		SelfCollisionConstraints->SetProperties(ConfigProperties);
+		SelfCollisionConstraints->SetProperties(ConfigProperties, FaceIntMaps);
 	}
 	if (SelfCollisionInit)
 	{
@@ -2217,7 +2205,7 @@ void FClothConstraints::Update(
 	Softs::FSolverReal MeshScale,
 	Softs::FSolverReal MaxDistancesScale)
 {
-	Update(ConfigProperties, WeightMaps, TMap<FString, const TSet<int32>*>(), MeshScale, MaxDistancesScale);
+	Update(ConfigProperties, WeightMaps, TMap<FString, const TSet<int32>*>(), TMap<FString, TConstArrayView<int32>>(), MeshScale, MaxDistancesScale);
 }
 
 // Deprecated
@@ -2226,7 +2214,7 @@ void FClothConstraints::Update(
 	Softs::FSolverReal MeshScale,
 	Softs::FSolverReal MaxDistancesScale)
 {
-	Update(ConfigProperties, TMap<FString, TConstArrayView<FRealSingle>>(), TMap<FString, const TSet<int32>*>(), MeshScale, MaxDistancesScale);
+	Update(ConfigProperties, TMap<FString, TConstArrayView<FRealSingle>>(), TMap<FString, const TSet<int32>*>(), TMap<FString, TConstArrayView<int32>>(), MeshScale, MaxDistancesScale);
 }
 
 }  // End namespace Chaos
diff --git a/Engine/Plugins/Experimental/ChaosCloth/Source/Chaos/Private/ChaosCloth/ChaosClothingSimulationCloth.cpp b/Engine/Plugins/Experimental/ChaosCloth/Source/Chaos/Private/ChaosCloth/ChaosClothingSimulationCloth.cpp
--- a/Engine/Plugins/Experimental/ChaosCloth/Source/Chaos/Private/ChaosCloth/ChaosClothingSimulationCloth.cpp
+++ b/Engine/Plugins/Experimental/ChaosCloth/Source/Chaos/Private/ChaosCloth/ChaosClothingSimulationCloth.cpp
@@ -57,6 +57,7 @@ struct FClothingSimulationCloth::FLODData
 	const int32 NumParticles;
 	const TMap<FString, TConstArrayView<FRealSingle>> WeightMaps;
 	const TMap<FString, const TSet<int32>*> VertexSets;
+	const TMap<FString, TConstArrayView<int32>> FaceIntMaps;
 	const TArray<TConstArrayView<TTuple<int32, int32, FRealSingle>>> Tethers;
 
 	const FClothingPatternData PatternData;
@@ -83,6 +84,7 @@ struct FClothingSimulationCloth::FLODData
 		const TConstArrayView<uint32>& InPatternToWeldedIndices,
 		TMap<FString, TConstArrayView<FRealSingle>>&& InWeightMaps,
 		TMap<FString, const TSet<int32>*>&& InVertexSets,
+		TMap<FString, TConstArrayView<int32>>&& InFaceIntMaps,
 		const TArray<TConstArrayView<TTuple<int32, int32, FRealSingle>>>& InTethers);
 
 	static FTriangleMesh BuildTriangleMesh(const TConstArrayView<uint32>& Indices, const int32 NumParticles);
@@ -107,10 +109,12 @@ FClothingSimulationCloth::FLODData::FLODData(
 	const TConstArrayView<uint32>& InPatternToWeldedIndices,
 	TMap<FString, TConstArrayView<FRealSingle>>&& InWeightMaps,
 	TMap<FString, const TSet<int32>*>&& InVertexSets,
+	TMap<FString, TConstArrayView<int32>>&& InFaceIntMaps,
 	const TArray<TConstArrayView<TTuple<int32, int32, FRealSingle>>>& InTethers)
 	: NumParticles(InNumParticles)
 	, WeightMaps(MoveTemp(InWeightMaps))
 	, VertexSets(MoveTemp(InVertexSets))
+	, FaceIntMaps(MoveTemp(InFaceIntMaps))
 	, Tethers(InTethers)
 	, PatternData(NumParticles, InIndices, InPatternPositions, InPatternIndices, InPatternToWeldedIndices)
 	, NoOffsetTriangleMesh(BuildTriangleMesh(InIndices, NumParticles))
@@ -249,7 +253,7 @@ PRAGMA_ENABLE_DEPRECATION_WARNINGS
 
 	// Create constraints
 	const bool bEnabled = false;  // Set constraint disabled by default
-	ClothConstraints.AddRules(ConfigProperties, TriangleMesh, &PatternData, WeightMaps, VertexSets, Tethers, MeshScale, bEnabled);
+	ClothConstraints.AddRules(ConfigProperties, TriangleMesh, &PatternData, WeightMaps, VertexSets, FaceIntMaps, Tethers, MeshScale, bEnabled);
 
 	// Update LOD stats
 	const TConstArrayView<Softs::FSolverReal> InvMasses(Solver->GetParticleInvMasses(ParticleRangeId), NumParticles);
@@ -290,7 +294,7 @@ PRAGMA_DISABLE_DEPRECATION_WARNINGS  // TODO: CHAOS_IS_CLOTHINGSIMULATIONMESH_AB
 	const Softs::FSolverReal MeshScale = Cloth->Mesh->GetScale();
 PRAGMA_ENABLE_DEPRECATION_WARNINGS
 	const Softs::FSolverReal MaxDistancesScale = (Softs::FSolverReal)Cloth->MaxDistancesMultiplier;
-	ClothConstraints.Update(Cloth->Config->GetProperties(SolverDatum.LODIndex), WeightMaps, VertexSets, MeshScale, MaxDistancesScale);
+	ClothConstraints.Update(Cloth->Config->GetProperties(SolverDatum.LODIndex), WeightMaps, VertexSets, FaceIntMaps, MeshScale, MaxDistancesScale);
 }
 
 void FClothingSimulationCloth::FLODData::Enable(FClothingSimulationSolver* Solver, bool bEnable) const
@@ -604,6 +608,7 @@ void FClothingSimulationCloth::SetMesh(FClothingSimulationMesh* InMesh)
 		}
 
 		TMap<FString, const TSet<int32>*> VertexSets = Mesh->GetVertexSets(LODIndex);
+		TMap<FString, TConstArrayView<int32>> FaceIntMaps = Mesh->GetFaceIntMaps(LODIndex);
 
 		const bool bUseGeodesicTethers = Config->GetProperties(LODIndex).GetValue<bool>(TEXT("UseGeodesicTethers"), ClothingSimulationClothDefault::bUseGeodesicTethers);
 
@@ -616,6 +621,7 @@ void FClothingSimulationCloth::SetMesh(FClothingSimulationMesh* InMesh)
 			Mesh->GetPatternToWeldedIndices(LODIndex),
 			MoveTemp(WeightMaps),
 			MoveTemp(VertexSets),
+			MoveTemp(FaceIntMaps),
 			Mesh->GetTethers(LODIndex, bUseGeodesicTethers)));
 	}
 
diff --git a/Engine/Plugins/Experimental/ChaosCloth/Source/Chaos/Private/ChaosCloth/ChaosClothingSimulationMesh.cpp b/Engine/Plugins/Experimental/ChaosCloth/Source/Chaos/Private/ChaosCloth/ChaosClothingSimulationMesh.cpp
--- a/Engine/Plugins/Experimental/ChaosCloth/Source/Chaos/Private/ChaosCloth/ChaosClothingSimulationMesh.cpp
+++ b/Engine/Plugins/Experimental/ChaosCloth/Source/Chaos/Private/ChaosCloth/ChaosClothingSimulationMesh.cpp
@@ -206,6 +206,12 @@ TMap<FString, const TSet<int32>*> FClothingSimulationMesh::GetVertexSets(int32 L
 	return TMap<FString, const TSet<int32>*>();
 }
 
+TMap<FString, TConstArrayView<int32>> FClothingSimulationMesh::GetFaceIntMaps(int32 LODIndex) const
+{
+	// Not supported
+	return TMap<FString, TConstArrayView<int32>>();
+}
+
 TArray<TConstArrayView<TTuple<int32, int32, float>>> FClothingSimulationMesh::GetTethers(int32 LODIndex, bool bUseGeodesicTethers) const
 {
 	TArray<TConstArrayView<TTuple<int32, int32, float>>> Tethers;
diff --git a/Engine/Plugins/Experimental/ChaosCloth/Source/Chaos/Private/ChaosCloth/ChaosClothingSimulationSkeletalMesh.cpp b/Engine/Plugins/Experimental/ChaosCloth/Source/Chaos/Private/ChaosCloth/ChaosClothingSimulationSkeletalMesh.cpp
--- a/Engine/Plugins/Experimental/ChaosCloth/Source/Chaos/Private/ChaosCloth/ChaosClothingSimulationSkeletalMesh.cpp
+++ b/Engine/Plugins/Experimental/ChaosCloth/Source/Chaos/Private/ChaosCloth/ChaosClothingSimulationSkeletalMesh.cpp
@@ -173,6 +173,12 @@ PRAGMA_ENABLE_DEPRECATION_WARNINGS
 		return VertexSets;
 	}
 
+	TMap<FString, TConstArrayView<int32>> FClothingSimulationSkeletalMesh::GetFaceIntMaps(int32 LODIndex) const
+	{
+		// Not supported
+		return TMap<FString, TConstArrayView<int32>>();
+	}
+
 	TArray<TConstArrayView<TTuple<int32, int32, float>>> FClothingSimulationSkeletalMesh::GetTethers(int32 LODIndex, bool bUseGeodesicTethers) const
 	{
 		TArray<TConstArrayView<TTuple<int32, int32, float>>> Tethers;
diff --git a/Engine/Plugins/Experimental/ChaosCloth/Source/Chaos/Public/ChaosCloth/ChaosClothConstraints.h b/Engine/Plugins/Experimental/ChaosCloth/Source/Chaos/Public/ChaosCloth/ChaosClothConstraints.h
--- a/Engine/Plugins/Experimental/ChaosCloth/Source/Chaos/Public/ChaosCloth/ChaosClothConstraints.h
+++ b/Engine/Plugins/Experimental/ChaosCloth/Source/Chaos/Public/ChaosCloth/ChaosClothConstraints.h
@@ -67,11 +67,12 @@ PRAGMA_ENABLE_DEPRECATION_WARNINGS
 			const FClothingPatternData* PatternData,
 			const TMap<FString, TConstArrayView<FRealSingle>>& WeightMaps,
 			const TMap<FString, const TSet<int32>*>& VertexSets,
+			const TMap<FString, TConstArrayView<int32>>& FaceIntMaps,
 			const TArray<TConstArrayView<TTuple<int32, int32, FRealSingle>>>& Tethers,
 			Softs::FSolverReal MeshScale,
 			bool bEnabled);
 
-		UE_DEPRECATED(5.4, "Use AddRules() with WeightMaps, VertexSets, and optional PatternData instead.")
+		UE_DEPRECATED(5.4, "Use AddRules() with WeightMaps, VertexSets, FaceIntMaps, and optional PatternData instead.")
 		void AddRules(
 			const Softs::FCollectionPropertyConstFacade& ConfigProperties,
 			const FTriangleMesh& TriangleMesh,
@@ -81,7 +82,7 @@ PRAGMA_ENABLE_DEPRECATION_WARNINGS
 			Softs::FSolverReal MeshScale,
 			bool bEnabled);
 
-		UE_DEPRECATED(5.3, "Use AddRules() with WeightMaps, VertexSets, and optional PatternData instead.")
+		UE_DEPRECATED(5.3, "Use AddRules() with WeightMaps, VertexSets, FaceIntMap, and optional PatternData instead.")
 		void AddRules(
 			const Softs::FCollectionPropertyConstFacade& ConfigProperties,
 			const FTriangleMesh& TriangleMesh,
@@ -94,17 +95,18 @@ PRAGMA_ENABLE_DEPRECATION_WARNINGS
 			const Softs::FCollectionPropertyConstFacade& ConfigProperties,
 			const TMap<FString, TConstArrayView<FRealSingle>>& WeightMaps,
 			const TMap<FString, const TSet<int32>*>& VertexSets,
+			const TMap<FString, TConstArrayView<int32>>& FaceIntMaps,
 			Softs::FSolverReal MeshScale,
 			Softs::FSolverReal MaxDistancesScale = (Softs::FSolverReal)1.);
 
-		UE_DEPRECATED(5.4, "Use Update() with WeightMaps and VertexSets instead.")
+		UE_DEPRECATED(5.4, "Use Update() with WeightMaps, VertexSets, and FaceIntMaps instead.")
 		void Update(
 			const Softs::FCollectionPropertyConstFacade& ConfigProperties,
 			const TMap<FString, TConstArrayView<FRealSingle>>& WeightMaps,
 			Softs::FSolverReal MeshScale,
 			Softs::FSolverReal MaxDistancesScale = (Softs::FSolverReal)1.);
 
-		UE_DEPRECATED(5.3, "Use Update() with WeightMaps instead.")
+		UE_DEPRECATED(5.3, "Use Update() with WeightMaps, VertexSets, and FaceIntMaps instead.")
 		void Update(
 			const Softs::FCollectionPropertyConstFacade& ConfigProperties,
 			Softs::FSolverReal MeshScale,
@@ -148,6 +150,7 @@ PRAGMA_ENABLE_DEPRECATION_WARNINGS
 		void CreateSelfCollisionConstraints(
 			const Softs::FCollectionPropertyConstFacade& ConfigProperties,
 			const TMap<FString, const TSet<int32>*>& VertexSets,
+			const TMap<FString, TConstArrayView<int32>>& FaceIntMaps,
 			const FTriangleMesh& TriangleMesh);
 		void CreateStretchConstraints(
 			const Softs::FCollectionPropertyConstFacade& ConfigProperties,
diff --git a/Engine/Plugins/Experimental/ChaosCloth/Source/Chaos/Public/ChaosCloth/ChaosClothingSimulationMesh.h b/Engine/Plugins/Experimental/ChaosCloth/Source/Chaos/Public/ChaosCloth/ChaosClothingSimulationMesh.h
--- a/Engine/Plugins/Experimental/ChaosCloth/Source/Chaos/Public/ChaosCloth/ChaosClothingSimulationMesh.h
+++ b/Engine/Plugins/Experimental/ChaosCloth/Source/Chaos/Public/ChaosCloth/ChaosClothingSimulationMesh.h
@@ -73,6 +73,7 @@ namespace Chaos
 		virtual TMap<FString, int32> GetWeightMapIndices() const { return GetWeightMapIndices(0); }
 		virtual TArray<TConstArrayView<FRealSingle>> GetWeightMaps(int32 LODIndex) const;
 		virtual TMap<FString, const TSet<int32>*> GetVertexSets(int32 LODIndex) const;
+		virtual TMap<FString, TConstArrayView<int32>> GetFaceIntMaps(int32 LODIndex) const;
 		virtual TArray<TConstArrayView<TTuple<int32, int32, float>>> GetTethers(int32 LODIndex, bool bUseGeodesicTethers) const;
 		virtual int32 GetReferenceBoneIndex() const;
 		virtual FTransform GetReferenceBoneTransform() const;
@@ -136,6 +137,9 @@ namespace Chaos
 		/* Return the specified LOD's vertex sets. */
 		virtual const TMap<FString, TSet<int32>*> GetVertexSets(int32 LODIndex) const = 0;
 
+		/* Return the specified LOD's face int maps. */
+		virtual TMap<FString, TConstArrayView<int32>> GetFaceIntMaps(int32 LODIndex) const = 0;
+
 		/* Return the tethers connections for the long range attachment into convenient parallel friendly batches. */
 		virtual TArray<TConstArrayView<TTuple<int32, int32, float>>> GetTethers(int32 LODIndex, bool bUseGeodesicTethers) const = 0;
 
diff --git a/Engine/Plugins/Experimental/ChaosCloth/Source/Chaos/Public/ChaosCloth/ChaosClothingSimulationSkeletalMesh.h b/Engine/Plugins/Experimental/ChaosCloth/Source/Chaos/Public/ChaosCloth/ChaosClothingSimulationSkeletalMesh.h
--- a/Engine/Plugins/Experimental/ChaosCloth/Source/Chaos/Public/ChaosCloth/ChaosClothingSimulationSkeletalMesh.h
+++ b/Engine/Plugins/Experimental/ChaosCloth/Source/Chaos/Public/ChaosCloth/ChaosClothingSimulationSkeletalMesh.h
@@ -53,6 +53,7 @@ PRAGMA_ENABLE_DEPRECATION_WARNINGS
 		virtual TMap<FString, int32> GetWeightMapIndices() const override { return GetWeightMapIndices(0); }
 		virtual TArray<TConstArrayView<FRealSingle>> GetWeightMaps(int32 LODIndex) const override;
 		virtual TMap<FString, const TSet<int32>*> GetVertexSets(int32 LODIndex) const override;
+		virtual TMap<FString, TConstArrayView<int32>> GetFaceIntMaps(int32 LODIndex) const override;
 		virtual TArray<TConstArrayView<TTuple<int32, int32, float>>> GetTethers(int32 LODIndex, bool bUseGeodesicTethers) const override;
 		virtual int32 GetReferenceBoneIndex() const override;
 		virtual FTransform GetReferenceBoneTransform() const override;
diff --git a/Engine/Plugins/Experimental/ChaosClothAsset/Source/ChaosClothAsset/Private/ChaosClothAsset/ClothCollection.cpp b/Engine/Plugins/Experimental/ChaosClothAsset/Source/ChaosClothAsset/Private/ChaosClothAsset/ClothCollection.cpp
--- a/Engine/Plugins/Experimental/ChaosClothAsset/Source/ChaosClothAsset/Private/ChaosClothAsset/ClothCollection.cpp
+++ b/Engine/Plugins/Experimental/ChaosClothAsset/Source/ChaosClothAsset/Private/ChaosClothAsset/ClothCollection.cpp
@@ -602,24 +602,50 @@ namespace UE::Chaos::ClothAsset
 	template CHAOSCLOTHASSET_API TArray<FName> FClothCollection::GetUserDefinedAttributeNames<FVector3f>(const FName& GroupName) const;
 
 	template<typename T, typename TEnableIf<TIsUserAttributeType<T>::Value, int>::type>
-	void FClothCollection::AddUserDefinedAttribute(const FName& Name, const FName& GroupName)
+	TManagedArray<T>* FClothCollection::FindOrAddUserDefinedAttribute(const FName& Name, const FName& GroupName)
 	{
-		ManagedArrayCollection->AddAttribute<T>(Name, GroupName);
+		using namespace UE::Chaos::ClothAsset::Private;
+		if (const TArray<FName>* const FixedAttributeNames = FixedAttributeNamesMap.Find(GroupName))
+		{
+			// Only allow adding to know groups. Do not allow adding a name reserved as a fixed attribute.
+			if (!FixedAttributeNames->Contains(Name))
+			{
+				return ManagedArrayCollection->FindOrAddAttributeTyped<T>(Name, GroupName);
+			}
+		}
+		return nullptr;
 	}
-	template CHAOSCLOTHASSET_API void FClothCollection::AddUserDefinedAttribute<bool>(const FName& Name, const FName& GroupName);
-	template CHAOSCLOTHASSET_API void FClothCollection::AddUserDefinedAttribute<int32>(const FName& Name, const FName& GroupName);
-	template CHAOSCLOTHASSET_API void FClothCollection::AddUserDefinedAttribute<float>(const FName& Name, const FName& GroupName);
-	template CHAOSCLOTHASSET_API void FClothCollection::AddUserDefinedAttribute<FVector3f>(const FName& Name, const FName& GroupName);
+	template CHAOSCLOTHASSET_API TManagedArray<bool>* FClothCollection::FindOrAddUserDefinedAttribute<bool>(const FName& Name, const FName& GroupName);
+	template CHAOSCLOTHASSET_API TManagedArray<int32>* FClothCollection::FindOrAddUserDefinedAttribute<int32>(const FName& Name, const FName& GroupName);
+	template CHAOSCLOTHASSET_API TManagedArray<float>* FClothCollection::FindOrAddUserDefinedAttribute<float>(const FName& Name, const FName& GroupName);
+	template CHAOSCLOTHASSET_API TManagedArray<FVector3f>* FClothCollection::FindOrAddUserDefinedAttribute<FVector3f>(const FName& Name, const FName& GroupName);
 
 	void FClothCollection::RemoveUserDefinedAttribute(const FName& Name, const FName& GroupName)
 	{
-		ManagedArrayCollection->RemoveAttribute(Name, GroupName);
+		using namespace UE::Chaos::ClothAsset::Private;
+		if (const TArray<FName>* const FixedAttributeNames = FixedAttributeNamesMap.Find(GroupName))
+		{
+			// User defined attributes are only allowed in known group names. Do not allow removing fixed attributes through this method.
+			if (!FixedAttributeNames->Contains(Name))
+			{
+				ManagedArrayCollection->RemoveAttribute(Name, GroupName);
+			}
+		}
 	}
 
 	template<typename T, typename TEnableIf<TIsUserAttributeType<T>::Value, int>::type>
 	bool FClothCollection::HasUserDefinedAttribute(const FName& Name, const FName& GroupName) const
 	{
-		return ManagedArrayCollection->FindAttributeTyped<T>(Name, GroupName) != nullptr;
+		using namespace UE::Chaos::ClothAsset::Private;
+		if (const TArray<FName>* const FixedAttributeNames = FixedAttributeNamesMap.Find(GroupName))
+		{
+			// User defined attributes are only allowed in known group names.
+			if (!FixedAttributeNames->Contains(Name))
+			{
+				return ManagedArrayCollection->FindAttributeTyped<T>(Name, GroupName) != nullptr;
+			}
+		}
+		return false;
 	}
 	template CHAOSCLOTHASSET_API bool FClothCollection::HasUserDefinedAttribute<bool>(const FName& Name, const FName& GroupName) const;
 	template CHAOSCLOTHASSET_API bool FClothCollection::HasUserDefinedAttribute<int32>(const FName& Name, const FName& GroupName) const;
@@ -629,7 +655,16 @@ namespace UE::Chaos::ClothAsset
 	template<typename T, typename TEnableIf<TIsUserAttributeType<T>::Value, int>::type>
 	const TManagedArray<T>* FClothCollection::GetUserDefinedAttribute(const FName& Name, const FName& GroupName) const
 	{
-		return ManagedArrayCollection->FindAttribute<T>(Name, GroupName);
+		using namespace UE::Chaos::ClothAsset::Private;
+		if (const TArray<FName>* const FixedAttributeNames = FixedAttributeNamesMap.Find(GroupName))
+		{
+			// User defined attributes are only allowed in known group names.
+			if (!FixedAttributeNames->Contains(Name))
+			{
+				return ManagedArrayCollection->FindAttributeTyped<T>(Name, GroupName);
+			}
+		}
+		return nullptr;
 	}
 	template CHAOSCLOTHASSET_API const TManagedArray<bool>* FClothCollection::GetUserDefinedAttribute<bool>(const FName& Name, const FName& GroupName) const;
 	template CHAOSCLOTHASSET_API const TManagedArray<int32>* FClothCollection::GetUserDefinedAttribute<int32>(const FName& Name, const FName& GroupName) const;
@@ -639,13 +674,42 @@ namespace UE::Chaos::ClothAsset
 	template<typename T, typename TEnableIf<TIsUserAttributeType<T>::Value, int>::type>
 	TManagedArray<T>* FClothCollection::GetUserDefinedAttribute(const FName& Name, const FName& GroupName)
 	{
-		return ManagedArrayCollection->FindAttribute<T>(Name, GroupName);
+		using namespace UE::Chaos::ClothAsset::Private;
+		if (const TArray<FName>* const FixedAttributeNames = FixedAttributeNamesMap.Find(GroupName))
+		{
+			// User defined attributes are only allowed in known group names.
+			if (!FixedAttributeNames->Contains(Name))
+			{
+				return ManagedArrayCollection->FindAttributeTyped<T>(Name, GroupName);
+			}
+		}
+		return nullptr;
 	}
 	template CHAOSCLOTHASSET_API TManagedArray<bool>* FClothCollection::GetUserDefinedAttribute<bool>(const FName& Name, const FName& GroupName);
 	template CHAOSCLOTHASSET_API TManagedArray<int32>* FClothCollection::GetUserDefinedAttribute<int32>(const FName& Name, const FName& GroupName);
 	template CHAOSCLOTHASSET_API TManagedArray<float>* FClothCollection::GetUserDefinedAttribute<float>(const FName& Name, const FName& GroupName);
 	template CHAOSCLOTHASSET_API TManagedArray<FVector3f>* FClothCollection::GetUserDefinedAttribute<FVector3f>(const FName& Name, const FName& GroupName);
 
+	bool FClothCollection::IsValidClothCollectionGroupName(const FName& GroupName)
+	{
+		using namespace UE::Chaos::ClothAsset::Private;
+		return FixedAttributeNamesMap.Contains(GroupName);
+	}
+
+	bool FClothCollection::IsValidUserDefinedAttributeName(const FName& Name, const FName& GroupName)
+	{
+		using namespace UE::Chaos::ClothAsset::Private;
+		if (const TArray<FName>* const FixedAttributeNames = FixedAttributeNamesMap.Find(GroupName))
+		{
+			// User defined attributes are only allowed in known group names and cannot be reserved by FixedAttributeNames.
+			if (!FixedAttributeNames->Contains(Name))
+			{
+				return true;
+			}
+		}
+		return false;
+	}
+
 	void FClothCollection::CopyArrayViewDataAndApplyOffset(const TArrayView<TArray<int32>>& To, const TConstArrayView<TArray<int32>>& From, const int32 Offset)
 	{
 		check(To.Num() == From.Num());
diff --git a/Engine/Plugins/Experimental/ChaosClothAsset/Source/ChaosClothAsset/Private/ChaosClothAsset/ClothCollection.h b/Engine/Plugins/Experimental/ChaosClothAsset/Source/ChaosClothAsset/Private/ChaosClothAsset/ClothCollection.h
--- a/Engine/Plugins/Experimental/ChaosClothAsset/Source/ChaosClothAsset/Private/ChaosClothAsset/ClothCollection.h
+++ b/Engine/Plugins/Experimental/ChaosClothAsset/Source/ChaosClothAsset/Private/ChaosClothAsset/ClothCollection.h
@@ -3,18 +3,12 @@
 #pragma once
 
 #include "GeometryCollection/ManagedArray.h"
+#include "ChaosClothAsset/IsUserAttributeType.h"
 
 struct FManagedArrayCollection;
 
 namespace UE::Chaos::ClothAsset
 {
-	/** User defined attribute types. */
-	template<typename T> struct TIsUserAttributeType { static constexpr bool Value = false; };
-	template<> struct TIsUserAttributeType<bool> { static constexpr bool Value = true; };
-	template<> struct TIsUserAttributeType<int32> { static constexpr bool Value = true; };
-	template<> struct TIsUserAttributeType<float> { static constexpr bool Value = true; };
-	template<> struct TIsUserAttributeType<FVector3f> { static constexpr bool Value = true; };
-
 	/**
 	 * Cloth collection facade data.
 	 */
@@ -136,7 +130,13 @@ namespace UE::Chaos::ClothAsset
 		TArray<FName> GetUserDefinedAttributeNames(const FName& GroupName) const;
 
 		template<typename T, TEMPLATE_REQUIRES(TIsUserAttributeType<T>::Value)>
-		void AddUserDefinedAttribute(const FName& Name, const FName& GroupName);
+		TManagedArray<T>* FindOrAddUserDefinedAttribute(const FName& Name, const FName& GroupName);
+
+		template<typename T, TEMPLATE_REQUIRES(TIsUserAttributeType<T>::Value)>
+		void AddUserDefinedAttribute(const FName& Name, const FName& GroupName)
+		{
+			FindOrAddUserDefinedAttribute<T>(Name, GroupName);
+		}
 
 		void RemoveUserDefinedAttribute(const FName& Name, const FName& GroupName);
 
@@ -149,6 +149,9 @@ namespace UE::Chaos::ClothAsset
 		template<typename T, TEMPLATE_REQUIRES(TIsUserAttributeType<T>::Value)>
 		TManagedArray<T>* GetUserDefinedAttribute(const FName& Name, const FName& GroupName);
 
+		static bool IsValidClothCollectionGroupName(const FName& GroupName);
+		static bool IsValidUserDefinedAttributeName(const FName& Name, const FName& GroupName);
+
 		//~ LODs Group (There should be only one LOD per ClothCollection)
 		const TManagedArray<FString>* GetPhysicsAssetPathName() const { return PhysicsAssetPathName; }
 		const TManagedArray<FString>* GetSkeletalMeshPathName() const { return SkeletalMeshPathName; }
diff --git a/Engine/Plugins/Experimental/ChaosClothAsset/Source/ChaosClothAsset/Private/ChaosClothAsset/CollectionClothFacade.cpp b/Engine/Plugins/Experimental/ChaosClothAsset/Source/ChaosClothAsset/Private/ChaosClothAsset/CollectionClothFacade.cpp
--- a/Engine/Plugins/Experimental/ChaosClothAsset/Source/ChaosClothAsset/Private/ChaosClothAsset/CollectionClothFacade.cpp
+++ b/Engine/Plugins/Experimental/ChaosClothAsset/Source/ChaosClothAsset/Private/ChaosClothAsset/CollectionClothFacade.cpp
@@ -116,6 +116,22 @@ namespace UE::Chaos::ClothAsset
 		return ResultHash;
 	}
 
+	template<typename T, typename TEnableIf<TIsUserAttributeType<T>::Value, int>::type>
+	uint32 FCollectionClothConstFacade::CalculateUserDefinedAttributesTypeHash(const FName& GroupName, uint32 PreviousHash) const
+	{
+		uint32 ResultHash = PreviousHash;
+		const TArray<FName> AttributesNames = GetUserDefinedAttributeNames<T>(GroupName);
+		for (const FName& AttributesName : AttributesNames)
+		{
+			ResultHash = HashCombineFast(ResultHash, GetTypeHash(AttributesName));
+			ResultHash = HashCombineFast(ResultHash, ClothCollection->GetElementsTypeHash(ClothCollection->GetUserDefinedAttribute<T>(AttributesName, GroupName)));
+		}
+		return ResultHash;
+	}
+	template CHAOSCLOTHASSET_API uint32 FCollectionClothConstFacade::CalculateUserDefinedAttributesTypeHash<int32>(const FName& GroupName, uint32 PreviousHash) const;
+	template CHAOSCLOTHASSET_API uint32 FCollectionClothConstFacade::CalculateUserDefinedAttributesTypeHash<float>(const FName& GroupName, uint32 PreviousHash) const;
+	template CHAOSCLOTHASSET_API uint32 FCollectionClothConstFacade::CalculateUserDefinedAttributesTypeHash<FVector3f>(const FName& GroupName, uint32 PreviousHash) const;
+
 	const FString& FCollectionClothConstFacade::GetPhysicsAssetPathName() const
 	{
 		static const FString EmptyString;
@@ -328,6 +344,34 @@ namespace UE::Chaos::ClothAsset
 		return ClothCollection->GetElements(ClothCollection->GetUserDefinedAttribute<float>(Name, ClothCollectionGroup::SimVertices3D));
 	}
 
+
+	template<typename T, typename TEnableIf<TIsUserAttributeType<T>::Value, int>::type>
+	bool FCollectionClothConstFacade::HasUserDefinedAttribute(const FName& Name, const FName& GroupName) const
+	{
+		return ClothCollection->HasUserDefinedAttribute<T>(Name, GroupName);
+	}
+	template CHAOSCLOTHASSET_API bool FCollectionClothConstFacade::HasUserDefinedAttribute<int32>(const FName& Name, const FName& GroupName) const;
+	template CHAOSCLOTHASSET_API bool FCollectionClothConstFacade::HasUserDefinedAttribute<float>(const FName& Name, const FName& GroupName) const;
+	template CHAOSCLOTHASSET_API bool FCollectionClothConstFacade::HasUserDefinedAttribute<FVector3f>(const FName& Name, const FName& GroupName) const;
+
+	template<typename T, typename TEnableIf<TIsUserAttributeType<T>::Value, int>::type>
+	TArray<FName> FCollectionClothConstFacade::GetUserDefinedAttributeNames(const FName& GroupName) const
+	{
+		return ClothCollection->GetUserDefinedAttributeNames<T>(GroupName);
+	}
+	template CHAOSCLOTHASSET_API TArray<FName> FCollectionClothConstFacade::GetUserDefinedAttributeNames<int32>(const FName& GroupName) const;
+	template CHAOSCLOTHASSET_API TArray<FName> FCollectionClothConstFacade::GetUserDefinedAttributeNames<float>(const FName& GroupName) const;
+	template CHAOSCLOTHASSET_API TArray<FName> FCollectionClothConstFacade::GetUserDefinedAttributeNames<FVector3f>(const FName& GroupName) const;
+
+	template<typename T, typename TEnableIf<TIsUserAttributeType<T>::Value, int>::type>
+	TConstArrayView<T> FCollectionClothConstFacade::GetUserDefinedAttribute(const FName& Name, const FName& GroupName) const
+	{
+		return ClothCollection->GetElements(ClothCollection->GetUserDefinedAttribute<T>(Name, GroupName));
+	}
+	template CHAOSCLOTHASSET_API TConstArrayView<int32> FCollectionClothConstFacade::GetUserDefinedAttribute<int32>(const FName& Name, const FName& GroupName) const;
+	template CHAOSCLOTHASSET_API TConstArrayView<float> FCollectionClothConstFacade::GetUserDefinedAttribute<float>(const FName& Name, const FName& GroupName) const;
+	template CHAOSCLOTHASSET_API TConstArrayView<FVector3f> FCollectionClothConstFacade::GetUserDefinedAttribute<FVector3f>(const FName& Name, const FName& GroupName) const;
+
 	void FCollectionClothConstFacade::BuildSimulationMesh(TArray<FVector3f>& Positions, TArray<FVector3f>& Normals, TArray<uint32>& Indices, TArray<FVector2f>& PatternsPositions, TArray<uint32>& PatternsIndices,
 		TArray<uint32>& PatternToWeldedIndices, TArray<TArray<int32>>* OptionalWeldedToPatternIndices) const
 	{
@@ -431,7 +475,9 @@ namespace UE::Chaos::ClothAsset
 		// Sim Patterns Group
 		const int32 StartNumSimVertices2D = GetNumSimVertices2D();
 		const int32 StartNumSimPatterns = GetNumSimPatterns();
+		const int32 StartNumSimFaces = GetNumSimFaces();
 		const int32 OtherNumSimPatterns = Other.GetNumSimPatterns();
+		const int32 OtherNumSimFaces = Other.GetNumSimFaces();
 		SetNumSimPatterns(StartNumSimPatterns + OtherNumSimPatterns);
 		for (int32 PatternIndex = 0; PatternIndex < OtherNumSimPatterns; ++PatternIndex)
 		{
@@ -471,6 +517,14 @@ namespace UE::Chaos::ClothAsset
 			AddWeightMap(WeightMapName);
 			FClothCollection::CopyArrayViewData(GetWeightMap(WeightMapName).Right(OtherNumSimVertices3D), Other.GetWeightMap(WeightMapName));
 		}
+
+		// Face int maps (self collision layers)
+		const TArray<FName> FaceIntMapNames = Other.GetUserDefinedAttributeNames<int32>(ClothCollectionGroup::SimFaces);
+		for (const FName& IntMapName : FaceIntMapNames)
+		{
+			AddUserDefinedAttribute<int32>(IntMapName, ClothCollectionGroup::SimFaces);
+			FClothCollection::CopyArrayViewData(GetUserDefinedAttribute<int32>(IntMapName, ClothCollectionGroup::SimFaces).Right(OtherNumSimFaces), Other.GetUserDefinedAttribute<int32>(IntMapName, ClothCollectionGroup::SimFaces));
+		}
 	}
 
 	void FCollectionClothFacade::SetPhysicsAssetPathName(const FString& PathName)
@@ -760,6 +814,36 @@ namespace UE::Chaos::ClothAsset
 		return GetClothCollection()->GetElements(GetClothCollection()->GetUserDefinedAttribute<float>(Name, ClothCollectionGroup::SimVertices3D));
 	}
 
+	template<typename T, typename TEnableIf<TIsUserAttributeType<T>::Value, int>::type>
+	bool FCollectionClothFacade::AddUserDefinedAttribute(const FName& Name, const FName& GroupName)
+	{
+		check(IsValid());
+		return GetClothCollection()->FindOrAddUserDefinedAttribute<T>(Name, GroupName) != nullptr;
+	}
+	template CHAOSCLOTHASSET_API bool FCollectionClothFacade::AddUserDefinedAttribute<int32>(const FName& Name, const FName& GroupName);
+	template CHAOSCLOTHASSET_API bool FCollectionClothFacade::AddUserDefinedAttribute<float>(const FName& Name, const FName& GroupName);
+	template CHAOSCLOTHASSET_API bool FCollectionClothFacade::AddUserDefinedAttribute<FVector3f>(const FName& Name, const FName& GroupName);
+
+	void FCollectionClothFacade::RemoveUserDefinedAttribute(const FName& Name, const FName& GroupName)
+	{
+		check(IsValid());
+		GetClothCollection()->RemoveUserDefinedAttribute(Name, GroupName);
+	}
+
+	template<typename T, typename TEnableIf<TIsUserAttributeType<T>::Value, int>::type>
+	TArrayView<T> FCollectionClothFacade::GetUserDefinedAttribute(const FName& Name, const FName& GroupName)
+	{
+		return GetClothCollection()->GetElements(GetClothCollection()->GetUserDefinedAttribute<T>(Name, GroupName));
+	}
+	template CHAOSCLOTHASSET_API TArrayView<int32> FCollectionClothFacade::GetUserDefinedAttribute<int32>(const FName& Name, const FName& GroupName);
+	template CHAOSCLOTHASSET_API TArrayView<float> FCollectionClothFacade::GetUserDefinedAttribute<float>(const FName& Name, const FName& GroupName);
+	template CHAOSCLOTHASSET_API TArrayView<FVector3f> FCollectionClothFacade::GetUserDefinedAttribute<FVector3f>(const FName& Name, const FName& GroupName);
+
+	bool FCollectionClothFacade::IsValidClothCollectionGroupName(const FName& GroupName)
+	{
+		return FClothCollection::IsValidClothCollectionGroupName(GroupName);
+	}
+
 	void FCollectionClothFacade::SetDefaults()
 	{
 		GetClothCollection()->SetNumElements(1, ClothCollectionGroup::Lods);
diff --git a/Engine/Plugins/Experimental/ChaosClothAsset/Source/ChaosClothAsset/Public/ChaosClothAsset/CollectionClothFacade.h b/Engine/Plugins/Experimental/ChaosClothAsset/Source/ChaosClothAsset/Public/ChaosClothAsset/CollectionClothFacade.h
--- a/Engine/Plugins/Experimental/ChaosClothAsset/Source/ChaosClothAsset/Public/ChaosClothAsset/CollectionClothFacade.h
+++ b/Engine/Plugins/Experimental/ChaosClothAsset/Source/ChaosClothAsset/Public/ChaosClothAsset/CollectionClothFacade.h
@@ -5,6 +5,7 @@
 #include "ChaosClothAsset/CollectionClothSimPatternFacade.h"
 #include "ChaosClothAsset/CollectionClothRenderPatternFacade.h"
 #include "ChaosClothAsset/CollectionClothSeamFacade.h"
+#include "ChaosClothAsset/IsUserAttributeType.h"
 
 namespace Chaos
 {
@@ -41,8 +42,10 @@ namespace UE::Chaos::ClothAsset
 		bool HasValidData() const;
 
 		uint32 CalculateTypeHash(bool bIncludeWeightMaps, uint32 PreviousHash = 0) const;
-
 		uint32 CalculateWeightMapTypeHash(uint32 PreviousHash = 0) const;
+		template<typename T, TEMPLATE_REQUIRES(TIsUserAttributeType<T>::Value)>
+		uint32 CalculateUserDefinedAttributesTypeHash(const FName& GroupName, uint32 PreviousHash = 0) const;
+		
 
 		//~ LOD (single per collection) Group
 		/** Return the physics asset path names used for this collection. */
@@ -119,14 +122,21 @@ namespace UE::Chaos::ClothAsset
 		int32 GetNumRenderFaces() const;
 		TConstArrayView<FIntVector3> GetRenderIndices() const;
 
+		//~ Weight Maps
 		/** Return whether this cloth collection has the specified weight map. */
 		bool HasWeightMap(const FName& Name) const;
-
 		/** Return the name of all user weight maps on this cloth collection. */
 		TArray<FName> GetWeightMapNames() const;
-
 		TConstArrayView<float> GetWeightMap(const FName& Name) const;
 
+		//~ Other User-Defined Attributes (not instantiated for bools)
+		template<typename T, TEMPLATE_REQUIRES(TIsUserAttributeType<T>::Value)>
+		bool HasUserDefinedAttribute(const FName& Name, const FName& GroupName) const;
+		template<typename T, TEMPLATE_REQUIRES(TIsUserAttributeType<T>::Value)>
+		TArray<FName> GetUserDefinedAttributeNames(const FName& GroupName) const;
+		template<typename T, TEMPLATE_REQUIRES(TIsUserAttributeType<T>::Value)>
+		TConstArrayView<T> GetUserDefinedAttribute(const FName& Name, const FName& GroupName) const;
+
 		void BuildSimulationMesh(TArray<FVector3f>& Positions, TArray<FVector3f>& Normals, TArray<uint32>& Indices, TArray<FVector2f>& PatternsPositions, TArray<uint32>& PatternsIndices, 
 			TArray<uint32>& PatternToWeldedIndices, TArray<TArray<int32>>* OptionalWeldedToPatternIndices = nullptr) const;
 
@@ -254,14 +264,22 @@ namespace UE::Chaos::ClothAsset
 		/** SetNumRenderFaces per pattern within pattern facade. */
 		TArrayView<FIntVector3> GetRenderIndices();
 
+		//~ Weight Maps
 		/** Add a new weight map to this cloth. Access is then done per pattern. */
 		void AddWeightMap(const FName& Name);
-
 		/** Remove a weight map from this cloth. */
 		void RemoveWeightMap(const FName& Name);
-
 		TArrayView<float> GetWeightMap(const FName& Name);
 
+		//~ Other User-Defined Attributes (not instantiated for bools)
+		/** GroupName must be an existing group as defined in ClothCollectionGroup. Returns success */
+		template<typename T, TEMPLATE_REQUIRES(TIsUserAttributeType<T>::Value)>
+		bool AddUserDefinedAttribute(const FName& Name, const FName& GroupName);
+		void RemoveUserDefinedAttribute(const FName& Name, const FName& GroupName);
+		template<typename T, TEMPLATE_REQUIRES(TIsUserAttributeType<T>::Value)>
+		TArrayView<T> GetUserDefinedAttribute(const FName& Name, const FName& GroupName);
+		static bool IsValidClothCollectionGroupName(const FName& GroupName);
+
 	private:
 		void SetDefaults();
 
diff --git a/Engine/Plugins/Experimental/ChaosClothAsset/Source/ChaosClothAsset/Public/ChaosClothAsset/IsUserAttributeType.h b/Engine/Plugins/Experimental/ChaosClothAsset/Source/ChaosClothAsset/Public/ChaosClothAsset/IsUserAttributeType.h
new file mode 100644
--- /dev/null
+++ b/Engine/Plugins/Experimental/ChaosClothAsset/Source/ChaosClothAsset/Public/ChaosClothAsset/IsUserAttributeType.h
@@ -0,0 +1,13 @@
+// Copyright Epic Games, Inc. All Rights Reserved.
+
+#pragma once
+
+namespace UE::Chaos::ClothAsset
+{
+	/** User defined attribute types. */
+	template<typename T> struct TIsUserAttributeType { static constexpr bool Value = false; };
+	template<> struct TIsUserAttributeType<bool> { static constexpr bool Value = true; };
+	template<> struct TIsUserAttributeType<int32> { static constexpr bool Value = true; };
+	template<> struct TIsUserAttributeType<float> { static constexpr bool Value = true; };
+	template<> struct TIsUserAttributeType<FVector3f> { static constexpr bool Value = true; };
+}  // End namespace UE::Chaos::ClothAsset
diff --git a/Engine/Plugins/Experimental/ChaosClothAsset/Source/ChaosClothAssetEngine/Private/ChaosClothAsset/ClothSimulationMesh.cpp b/Engine/Plugins/Experimental/ChaosClothAsset/Source/ChaosClothAssetEngine/Private/ChaosClothAsset/ClothSimulationMesh.cpp
--- a/Engine/Plugins/Experimental/ChaosClothAsset/Source/ChaosClothAssetEngine/Private/ChaosClothAsset/ClothSimulationMesh.cpp
+++ b/Engine/Plugins/Experimental/ChaosClothAsset/Source/ChaosClothAssetEngine/Private/ChaosClothAsset/ClothSimulationMesh.cpp
@@ -137,6 +137,20 @@ PRAGMA_ENABLE_DEPRECATION_WARNINGS
 		return VertexSets;
 	}
 
+	TMap<FString, TConstArrayView<int32>> FClothSimulationMesh::GetFaceIntMaps(int32 LODIndex) const
+	{
+		TMap<FString, TConstArrayView<int32>> FaceIntMaps;
+		if (ClothSimulationModel.ClothSimulationLodModels.IsValidIndex(LODIndex))
+		{
+			FaceIntMaps.Reserve(ClothSimulationModel.ClothSimulationLodModels[LODIndex].FaceIntMaps.Num());
+			for (TMap<FName, TArray<int32>>::TConstIterator Iter = ClothSimulationModel.ClothSimulationLodModels[LODIndex].FaceIntMaps.CreateConstIterator(); Iter; ++Iter)
+			{
+				FaceIntMaps.Emplace(Iter.Key().ToString(), TConstArrayView<int32>(Iter.Value()));
+			}
+		}
+		return FaceIntMaps;
+	}
+
 	TArray<TConstArrayView<TTuple<int32, int32, float>>> FClothSimulationMesh::GetTethers(int32 LODIndex, bool /*bUseGeodesicTethers*/) const
 	{
 		return ClothSimulationModel.GetTethers(LODIndex);
diff --git a/Engine/Plugins/Experimental/ChaosClothAsset/Source/ChaosClothAssetEngine/Private/ChaosClothAsset/ClothSimulationMesh.h b/Engine/Plugins/Experimental/ChaosClothAsset/Source/ChaosClothAssetEngine/Private/ChaosClothAsset/ClothSimulationMesh.h
--- a/Engine/Plugins/Experimental/ChaosClothAsset/Source/ChaosClothAssetEngine/Private/ChaosClothAsset/ClothSimulationMesh.h
+++ b/Engine/Plugins/Experimental/ChaosClothAsset/Source/ChaosClothAssetEngine/Private/ChaosClothAsset/ClothSimulationMesh.h
@@ -46,6 +46,7 @@ PRAGMA_ENABLE_DEPRECATION_WARNINGS
 		virtual TMap<FString, int32> GetWeightMapIndices() const override { return GetWeightMapIndices(0); }
 		virtual TArray<TConstArrayView<::Chaos::FRealSingle>> GetWeightMaps(int32 LODIndex) const override;
 		virtual TMap<FString, const TSet<int32>*> GetVertexSets(int32 LODIndex) const override;
+		virtual TMap<FString, TConstArrayView<int32>> GetFaceIntMaps(int32 LODIndex) const override;
 		// Note: there is only one set of tethers stored on ClothSimulationMesh assets
 		virtual TArray<TConstArrayView<TTuple<int32, int32, float>>> GetTethers(int32 LODIndex, bool /*bUseGeodesicTethers*/) const override;
 		virtual int32 GetReferenceBoneIndex() const override;
diff --git a/Engine/Plugins/Experimental/ChaosClothAsset/Source/ChaosClothAssetEngine/Private/ChaosClothAsset/ClothSimulationModel.cpp b/Engine/Plugins/Experimental/ChaosClothAsset/Source/ChaosClothAssetEngine/Private/ChaosClothAsset/ClothSimulationModel.cpp
--- a/Engine/Plugins/Experimental/ChaosClothAsset/Source/ChaosClothAssetEngine/Private/ChaosClothAsset/ClothSimulationModel.cpp
+++ b/Engine/Plugins/Experimental/ChaosClothAsset/Source/ChaosClothAssetEngine/Private/ChaosClothAsset/ClothSimulationModel.cpp
@@ -146,6 +146,7 @@ bool FChaosClothSimulationLodModel::Serialize(FArchive& Ar)
 	Ar << LODTransitionDownData;
 
 	Ar << VertexSets;
+	Ar << FaceIntMaps;
 
 	// Return true to confirm that serialization has already been taken care of
 	return true;
@@ -193,6 +194,14 @@ FChaosClothSimulationModel::FChaosClothSimulationModel(const TArray<TSharedRef<c
 			}
 		}
 
+		// Copy face int maps
+		const TArray<FName> FaceIntMapNames = Cloth.GetUserDefinedAttributeNames<int32>(ClothCollectionGroup::SimFaces);
+		LodModel.FaceIntMaps.Reserve(FaceIntMapNames.Num());
+		for (const FName& FaceIntMapName : FaceIntMapNames)
+		{
+			LodModel.FaceIntMaps.Add(FaceIntMapName) = Cloth.GetUserDefinedAttribute<int32>(FaceIntMapName, ClothCollectionGroup::SimFaces);
+		}
+
 		// Copy bone influences and gather tether data
 		const int32 NumSimVertices3D = Cloth.GetNumSimVertices3D();
 		TConstArrayView<TArray<int32>> SimBoneIndices = Cloth.GetSimBoneIndices();
diff --git a/Engine/Plugins/Experimental/ChaosClothAsset/Source/ChaosClothAssetEngine/Private/ChaosClothAsset/ClothSimulationModel.h b/Engine/Plugins/Experimental/ChaosClothAsset/Source/ChaosClothAssetEngine/Private/ChaosClothAsset/ClothSimulationModel.h
--- a/Engine/Plugins/Experimental/ChaosClothAsset/Source/ChaosClothAssetEngine/Private/ChaosClothAsset/ClothSimulationModel.h
+++ b/Engine/Plugins/Experimental/ChaosClothAsset/Source/ChaosClothAssetEngine/Private/ChaosClothAsset/ClothSimulationModel.h
@@ -65,6 +65,9 @@ struct FChaosClothSimulationLodModel
 	/** Vertex sets */
 	TMap<FName, TSet<int32>> VertexSets;
 
+	/** Face int maps (currently used by cloth collision layers)*/
+	TMap<FName, TArray<int32>> FaceIntMaps;
+
 	// Custom serialize for weight maps
 	bool Serialize(FArchive& Ar);
 };
diff --git a/Engine/Plugins/Experimental/ChaosClothAssetEditor/Source/ChaosClothAssetDataflowNodes/Private/ChaosClothAsset/AddWeightMapNode.cpp b/Engine/Plugins/Experimental/ChaosClothAssetEditor/Source/ChaosClothAssetDataflowNodes/Private/ChaosClothAsset/AddWeightMapNode.cpp
--- a/Engine/Plugins/Experimental/ChaosClothAssetEditor/Source/ChaosClothAssetDataflowNodes/Private/ChaosClothAsset/AddWeightMapNode.cpp
+++ b/Engine/Plugins/Experimental/ChaosClothAssetEditor/Source/ChaosClothAssetDataflowNodes/Private/ChaosClothAsset/AddWeightMapNode.cpp
@@ -232,6 +232,15 @@ void FChaosClothAssetAddWeightMapNode::Evaluate(Dataflow::FContext& Context, con
 			ClothFacade.AddWeightMap(InName);		// Does nothing if weight map already exists
 
 			TArrayView<float> ClothWeights = ClothFacade.GetWeightMap(InName);
+			if (ClothWeights.Num() != ClothFacade.GetNumSimVertices3D())
+			{
+				check(ClothWeights.Num() == 0);
+				FClothDataflowTools::LogAndToastWarning(*this,
+					LOCTEXT("InvalidWeightMapNameHeadline", "Invalid weight map name."),
+					FText::Format(LOCTEXT("InvalidWeightMapNameDetails", "Could not create a weight map with name \"{0}\" (reserved name? wrong type?)."),
+						FText::FromName(InName)));
+			}
+
 			const int32 MaxWeightIndex = FMath::Min(GetVertexWeights().Num(), ClothWeights.Num());
 			if (GetVertexWeights().Num() > 0 && GetVertexWeights().Num() != ClothWeights.Num())
 			{
diff --git a/Engine/Plugins/Experimental/ChaosClothAssetEditor/Source/ChaosClothAssetDataflowNodes/Private/ChaosClothAsset/ChaosClothAssetDataflowNodesModule.cpp b/Engine/Plugins/Experimental/ChaosClothAssetEditor/Source/ChaosClothAssetDataflowNodes/Private/ChaosClothAsset/ChaosClothAssetDataflowNodesModule.cpp
--- a/Engine/Plugins/Experimental/ChaosClothAssetEditor/Source/ChaosClothAssetDataflowNodes/Private/ChaosClothAsset/ChaosClothAssetDataflowNodesModule.cpp
+++ b/Engine/Plugins/Experimental/ChaosClothAssetEditor/Source/ChaosClothAssetDataflowNodes/Private/ChaosClothAsset/ChaosClothAssetDataflowNodesModule.cpp
@@ -13,6 +13,7 @@
 #include "ChaosClothAsset/MergeClothCollectionsNode.h"
 #include "ChaosClothAsset/ReverseNormalsNode.h"
 #include "ChaosClothAsset/SelectionNode.h"
+#include "ChaosClothAsset/SelectionToIntMapNode.h"
 #include "ChaosClothAsset/SelectionToWeightMapNode.h"
 #include "ChaosClothAsset/SetPhysicsAssetNode.h"
 #include "ChaosClothAsset/SimulationAerodynamicsConfigNode.h"
@@ -69,6 +70,7 @@ namespace UE::Chaos::ClothAsset
 			DATAFLOW_NODE_REGISTER_CREATION_FACTORY(FChaosClothAssetRemeshNode);
 			DATAFLOW_NODE_REGISTER_CREATION_FACTORY(FChaosClothAssetReverseNormalsNode);
 			DATAFLOW_NODE_REGISTER_CREATION_FACTORY(FChaosClothAssetSelectionNode);
+			DATAFLOW_NODE_REGISTER_CREATION_FACTORY(FChaosClothAssetSelectionToIntMapNode);
 			DATAFLOW_NODE_REGISTER_CREATION_FACTORY(FChaosClothAssetSelectionToWeightMapNode);
 			DATAFLOW_NODE_REGISTER_CREATION_FACTORY(FChaosClothAssetSetPhysicsAssetNode);
 			DATAFLOW_NODE_REGISTER_CREATION_FACTORY(FChaosClothAssetSimulationAerodynamicsConfigNode);
@@ -123,6 +125,8 @@ namespace UE::Chaos::ClothAsset
 				PropertyModule->RegisterCustomPropertyTypeLayout(FChaosClothAssetWeightedValue::StaticStruct()->GetFName(), FOnGetPropertyTypeCustomizationInstance::CreateStatic(&FWeightedValueCustomization::MakeInstance));
 				PropertyModule->RegisterCustomPropertyTypeLayout(FChaosClothAssetWeightedValueNonAnimatable::StaticStruct()->GetFName(), FOnGetPropertyTypeCustomizationInstance::CreateStatic(&FWeightedValueCustomization::MakeInstance));
 				PropertyModule->RegisterCustomPropertyTypeLayout(FChaosClothAssetWeightedValueNonAnimatableNoLowHighRange::StaticStruct()->GetFName(), FOnGetPropertyTypeCustomizationInstance::CreateStatic(&FWeightedValueCustomization::MakeInstance));
+				PropertyModule->RegisterCustomPropertyTypeLayout(FChaosClothAssetConnectableIStringValue::StaticStruct()->GetFName(), FOnGetPropertyTypeCustomizationInstance::CreateStatic(&FWeightedValueCustomization::MakeInstance));
+				PropertyModule->RegisterCustomPropertyTypeLayout(FChaosClothAssetConnectableIOStringValue::StaticStruct()->GetFName(), FOnGetPropertyTypeCustomizationInstance::CreateStatic(&FWeightedValueCustomization::MakeInstance));
 				PropertyModule->RegisterCustomPropertyTypeLayout(FChaosClothAssetNodeSelectionGroup::StaticStruct()->GetFName(), FOnGetPropertyTypeCustomizationInstance::CreateStatic(&FSelectionGroupCustomization::MakeInstance));
 				PropertyModule->RegisterCustomPropertyTypeLayout(FChaosClothAssetImportFilePath::StaticStruct()->GetFName(), FOnGetPropertyTypeCustomizationInstance::CreateStatic(&FImportFilePathCustomization::MakeInstance));
 			}
diff --git a/Engine/Plugins/Experimental/ChaosClothAssetEditor/Source/ChaosClothAssetDataflowNodes/Private/ChaosClothAsset/SelectionToIntMapNode.cpp b/Engine/Plugins/Experimental/ChaosClothAssetEditor/Source/ChaosClothAssetDataflowNodes/Private/ChaosClothAsset/SelectionToIntMapNode.cpp
new file mode 100644
--- /dev/null
+++ b/Engine/Plugins/Experimental/ChaosClothAssetEditor/Source/ChaosClothAssetDataflowNodes/Private/ChaosClothAsset/SelectionToIntMapNode.cpp
@@ -0,0 +1,99 @@
+// Copyright Epic Games, Inc. All Rights Reserved.
+
+#include "ChaosClothAsset/SelectionToIntMapNode.h"
+#include "ChaosClothAsset/ClothCollectionGroup.h"
+#include "ChaosClothAsset/ClothDataflowTools.h"
+#include "ChaosClothAsset/CollectionClothFacade.h"
+#include "ChaosClothAsset/CollectionClothSelectionFacade.h"
+#include "ChaosClothAsset/WeightedValue.h"
+#include "Dataflow/DataflowInputOutput.h"
+
+#include UE_INLINE_GENERATED_CPP_BY_NAME(SelectionToIntMapNode)
+#define LOCTEXT_NAMESPACE "FChaosClothAssetSelectionToIntMapNode"
+
+FChaosClothAssetSelectionToIntMapNode::FChaosClothAssetSelectionToIntMapNode(const Dataflow::FNodeParameters& InParam, FGuid InGuid)
+	: FDataflowNode(InParam, InGuid)
+{
+	RegisterInputConnection(&Collection);
+	RegisterOutputConnection(&Collection, &Collection);
+	RegisterInputConnection(&SelectionName.StringValue, GET_MEMBER_NAME_CHECKED(FChaosClothAssetConnectableIStringValue, StringValue));
+	RegisterInputConnection(&IntMapName.StringValue, GET_MEMBER_NAME_CHECKED(FChaosClothAssetConnectableIOStringValue, StringValue));
+	RegisterOutputConnection(&IntMapName.StringValue, &IntMapName.StringValue, GET_MEMBER_NAME_CHECKED(FChaosClothAssetConnectableIOStringValue, StringValue), GET_MEMBER_NAME_CHECKED(FChaosClothAssetConnectableIOStringValue, StringValue));
+}
+
+void FChaosClothAssetSelectionToIntMapNode::Evaluate(Dataflow::FContext& Context, const FDataflowOutput* Out) const
+{
+	using namespace UE::Chaos::ClothAsset;
+
+	if (Out->IsA<FManagedArrayCollection>(&Collection))
+	{
+		FManagedArrayCollection InCollection = GetValue<FManagedArrayCollection>(Context, &Collection);
+		const TSharedRef<FManagedArrayCollection> ClothCollection = MakeShared<FManagedArrayCollection>(MoveTemp(InCollection));
+		FCollectionClothFacade ClothFacade(ClothCollection);
+
+		const FCollectionClothSelectionConstFacade SelectionFacade(ClothCollection);
+		const FName InSelectionName(*GetValue<FString>(Context, &SelectionName.StringValue));
+		SelectionName.StringValue_Override = GetValue<FString>(Context, &SelectionName.StringValue, UE::Chaos::ClothAsset::FWeightMapTools::NotOverridden);
+		const FString& InIntMapNameString = GetValue<FString>(Context, &IntMapName.StringValue);
+		IntMapName.StringValue_Override = GetValue<FString>(Context, &IntMapName.StringValue, UE::Chaos::ClothAsset::FWeightMapTools::NotOverridden);
+		const FName InIntMapName = InIntMapNameString.IsEmpty() ? InSelectionName : FName(*InIntMapNameString);
+		if (SelectionFacade.IsValid() && ClothFacade.IsValid() && InSelectionName != NAME_None && InIntMapName != NAME_None)
+		{
+			if (const TSet<int32>* const SelectionSet = SelectionFacade.FindSelectionSet(InSelectionName))
+			{
+				const FName& SelectionGroup = SelectionFacade.GetSelectionGroup(InSelectionName);
+
+				bool bIsNewMap = false;
+				if (!ClothFacade.HasUserDefinedAttribute<int32>(InIntMapName, SelectionGroup))
+				{
+					if (ClothFacade.AddUserDefinedAttribute<int32>(InIntMapName, SelectionGroup))
+					{
+						bIsNewMap = true;
+					}
+					else
+					{
+						FClothDataflowTools::LogAndToastWarning(*this,
+							LOCTEXT("IntMapNameInvalidHeadline", "Invalid Int Map Name."),
+							FText::Format(LOCTEXT("IntMapNameInvalidDetails", "Could not create (or find) int map with name \"{0}\" in group \"{1}\""),
+								FText::FromName(InIntMapName),
+								FText::FromName(SelectionGroup)));
+					}
+				}
+
+				TArrayView<int32> IntMap = ClothFacade.GetUserDefinedAttribute<int32>(InIntMapName, SelectionGroup);
+
+				if (bIsNewMap || !bKeepExistingUnselectedValues)
+				{
+					for (int32& MapValue : IntMap)
+					{
+						MapValue = UnselectedValue;
+					}
+				}
+				for (const int32 Index : *SelectionSet)
+				{
+					if (IntMap.IsValidIndex(Index))
+					{
+						IntMap[Index] = SelectedValue;
+					}
+				}
+			}
+			else
+			{
+				FClothDataflowTools::LogAndToastWarning(*this,
+					LOCTEXT("SelectionNameNotFoundHeadline", "Invalid selection name."),
+					FText::Format(LOCTEXT("SelectionNameNotFoundDetails", "Selection \"{0}\" was not found in the collection."),
+						FText::FromName(InSelectionName)));
+			}
+		}
+
+		SetValue(Context, MoveTemp(*ClothCollection), &Collection);
+	}
+	else if (Out->IsA<FString>(&IntMapName.StringValue))
+	{
+		const FString& InIntMapNameString = GetValue<FString>(Context, &IntMapName.StringValue);
+		SetValue(Context, InIntMapNameString.IsEmpty() ? GetValue<FString>(Context, &SelectionName.StringValue) : InIntMapNameString, &IntMapName.StringValue);
+	}
+}
+
+
+#undef LOCTEXT_NAMESPACE
diff --git a/Engine/Plugins/Experimental/ChaosClothAssetEditor/Source/ChaosClothAssetDataflowNodes/Private/ChaosClothAsset/SelectionToWeightMapNode.cpp b/Engine/Plugins/Experimental/ChaosClothAssetEditor/Source/ChaosClothAssetDataflowNodes/Private/ChaosClothAsset/SelectionToWeightMapNode.cpp
--- a/Engine/Plugins/Experimental/ChaosClothAssetEditor/Source/ChaosClothAssetDataflowNodes/Private/ChaosClothAsset/SelectionToWeightMapNode.cpp
+++ b/Engine/Plugins/Experimental/ChaosClothAssetEditor/Source/ChaosClothAssetDataflowNodes/Private/ChaosClothAsset/SelectionToWeightMapNode.cpp
@@ -45,50 +45,60 @@ void FChaosClothAssetSelectionToWeightMapNode::Evaluate(Dataflow::FContext& Cont
 					const FName InMapName(WeightMapName.IsEmpty() ? SelectionName : WeightMapName);
 					ClothFacade.AddWeightMap(InMapName);
 					TArrayView<float> OutClothWeights = ClothFacade.GetWeightMap(InMapName);
-
-					if (SelectionGroup == ClothCollectionGroup::SimVertices3D)
+					if (OutClothWeights.Num() == ClothFacade.GetNumSimVertices3D())
 					{
-						for (int32 VertexIndex = 0; VertexIndex < OutClothWeights.Num(); ++VertexIndex)
+						if (SelectionGroup == ClothCollectionGroup::SimVertices3D)
 						{
-							OutClothWeights[VertexIndex] = SelectionSet->Contains(VertexIndex) ? SelectedValue : UnselectedValue;
+							for (int32 VertexIndex = 0; VertexIndex < OutClothWeights.Num(); ++VertexIndex)
+							{
+								OutClothWeights[VertexIndex] = SelectionSet->Contains(VertexIndex) ? SelectedValue : UnselectedValue;
+							}
 						}
-					}
-					else if (SelectionGroup == ClothCollectionGroup::SimVertices2D)
-					{
-						// We are given a selection over the set of 2D vertices, but weight maps only exist for 3D vertices, so
-						// we need a bit of translation
-						const TConstArrayView<TArray<int32>> Vertex3DTo2D = ClothFacade.GetSimVertex2DLookup();
-
-						for (int32 Vertex3DIndex = 0; Vertex3DIndex < OutClothWeights.Num(); ++Vertex3DIndex)
+						else if (SelectionGroup == ClothCollectionGroup::SimVertices2D)
 						{
-							// If any corresponding 2D vertex is selected, set the 3D weight map value to one, otherwise it gets zero
-							OutClothWeights[Vertex3DIndex] = UnselectedValue;
-							for (const int32 Vertex2DIndex : Vertex3DTo2D[Vertex3DIndex])
+							// We are given a selection over the set of 2D vertices, but weight maps only exist for 3D vertices, so
+							// we need a bit of translation
+							const TConstArrayView<TArray<int32>> Vertex3DTo2D = ClothFacade.GetSimVertex2DLookup();
+
+							for (int32 Vertex3DIndex = 0; Vertex3DIndex < OutClothWeights.Num(); ++Vertex3DIndex)
 							{
-								if (SelectionSet->Contains(Vertex2DIndex))
+								// If any corresponding 2D vertex is selected, set the 3D weight map value to one, otherwise it gets zero
+								OutClothWeights[Vertex3DIndex] = UnselectedValue;
+								for (const int32 Vertex2DIndex : Vertex3DTo2D[Vertex3DIndex])
 								{
-									OutClothWeights[Vertex3DIndex] = SelectedValue;
-									break;
+									if (SelectionSet->Contains(Vertex2DIndex))
+									{
+										OutClothWeights[Vertex3DIndex] = SelectedValue;
+										break;
+									}
 								}
 							}
 						}
+						else
+						{
+							check(SelectionGroup == ClothCollectionGroup::SimFaces);
+
+							// Fill with unselected value and then set all vertices in the selected faces to the selected value.
+							for (int32 Vertex3DIndex = 0; Vertex3DIndex < OutClothWeights.Num(); ++Vertex3DIndex)
+							{
+								OutClothWeights[Vertex3DIndex] = UnselectedValue;
+							}
+							const TConstArrayView<FIntVector3> SimIndices3D = ClothFacade.GetSimIndices3D();
+							for (const int32 FaceIndex : *SelectionSet)
+							{
+								OutClothWeights[SimIndices3D[FaceIndex][0]] = SelectedValue;
+								OutClothWeights[SimIndices3D[FaceIndex][1]] = SelectedValue;
+								OutClothWeights[SimIndices3D[FaceIndex][2]] = SelectedValue;
+							}
+						}
 					}
 					else
 					{
-						check(SelectionGroup == ClothCollectionGroup::SimFaces);
-
-						// Fill with unselected value and then set all vertices in the selected faces to the selected value.
-						for (int32 Vertex3DIndex = 0; Vertex3DIndex < OutClothWeights.Num(); ++Vertex3DIndex)
-						{
-							OutClothWeights[Vertex3DIndex] = UnselectedValue;
-						}
-						const TConstArrayView<FIntVector3> SimIndices3D = ClothFacade.GetSimIndices3D();
-						for (const int32 FaceIndex : *SelectionSet)
-						{
-							OutClothWeights[SimIndices3D[FaceIndex][0]] = SelectedValue;
-							OutClothWeights[SimIndices3D[FaceIndex][1]] = SelectedValue;
-							OutClothWeights[SimIndices3D[FaceIndex][2]] = SelectedValue;
-						}
+						check(OutClothWeights.Num() == 0);
+						FClothDataflowTools::LogAndToastWarning(*this,
+							LOCTEXT("InvalidWeightMapNameHeadline", "Invalid weight map name."),
+							FText::Format(LOCTEXT("InvalidWeightMapNameDetails", "Could not create a weight map with name \"{0}\" (reserved name? wrong type?)."),
+								FText::FromName(InMapName)));
 					}
 				}
 				else
diff --git a/Engine/Plugins/Experimental/ChaosClothAssetEditor/Source/ChaosClothAssetDataflowNodes/Private/ChaosClothAsset/SimulationBaseConfigNode.cpp b/Engine/Plugins/Experimental/ChaosClothAssetEditor/Source/ChaosClothAssetDataflowNodes/Private/ChaosClothAsset/SimulationBaseConfigNode.cpp
--- a/Engine/Plugins/Experimental/ChaosClothAssetEditor/Source/ChaosClothAssetDataflowNodes/Private/ChaosClothAsset/SimulationBaseConfigNode.cpp
+++ b/Engine/Plugins/Experimental/ChaosClothAssetEditor/Source/ChaosClothAssetDataflowNodes/Private/ChaosClothAsset/SimulationBaseConfigNode.cpp
@@ -201,4 +201,16 @@ int32 FChaosClothAssetSimulationBaseConfigNode::FPropertyHelper::SetPropertyWeig
 	return PropertyKeyIndex;
 }
 
+int32 FChaosClothAssetSimulationBaseConfigNode::FPropertyHelper::SetPropertyString(
+	const FName& PropertyName,
+	const FChaosClothAssetConnectableIStringValue& PropertyValue,
+	const TArray<FName>& SimilarPropertyNames,
+	ECollectionPropertyFlags PropertyFlags)
+{
+	const int32 PropertyKeyIndex = ConfigNode.AddPropertyHelper(Properties, PropertyName, SimilarPropertyNames, PropertyFlags);
+	Properties.SetStringValue(PropertyKeyIndex, ConfigNode.GetValue<FString>(Context, &PropertyValue.StringValue));
+	PropertyValue.StringValue_Override = ConfigNode.GetValue<FString>(Context, &PropertyValue.StringValue, UE::Chaos::ClothAsset::FWeightMapTools::NotOverridden);
+	return PropertyKeyIndex;
+}
+
 #undef LOCTEXT_NAMESPACE
diff --git a/Engine/Plugins/Experimental/ChaosClothAssetEditor/Source/ChaosClothAssetDataflowNodes/Private/ChaosClothAsset/SimulationSelfCollisionConfigNode.cpp b/Engine/Plugins/Experimental/ChaosClothAssetEditor/Source/ChaosClothAssetDataflowNodes/Private/ChaosClothAsset/SimulationSelfCollisionConfigNode.cpp
--- a/Engine/Plugins/Experimental/ChaosClothAssetEditor/Source/ChaosClothAssetDataflowNodes/Private/ChaosClothAsset/SimulationSelfCollisionConfigNode.cpp
+++ b/Engine/Plugins/Experimental/ChaosClothAssetEditor/Source/ChaosClothAssetDataflowNodes/Private/ChaosClothAsset/SimulationSelfCollisionConfigNode.cpp
@@ -1,6 +1,7 @@
 // Copyright Epic Games, Inc. All Rights Reserved.
 
 #include "ChaosClothAsset/SimulationSelfCollisionConfigNode.h"
+#include "Dataflow/DataflowInputOutput.h"
 #include "Chaos/CollectionPropertyFacade.h"
 
 #include UE_INLINE_GENERATED_CPP_BY_NAME(SimulationSelfCollisionConfigNode)
@@ -9,6 +10,7 @@ FChaosClothAssetSimulationSelfCollisionConfigNode::FChaosClothAssetSimulationSel
 	: FChaosClothAssetSimulationBaseConfigNode(InParam, InGuid)
 {
 	RegisterCollectionConnections();
+	RegisterInputConnection(&SelfCollisionLayers.StringValue, GET_MEMBER_NAME_CHECKED(FChaosClothAssetConnectableIStringValue, StringValue));
 }
 
 void FChaosClothAssetSimulationSelfCollisionConfigNode::AddProperties(FPropertyHelper& PropertyHelper) const
@@ -17,6 +19,8 @@ void FChaosClothAssetSimulationSelfCollisionConfigNode::AddProperties(FPropertyH
 	PropertyHelper.SetProperty(this, &SelfCollisionThickness);
 	PropertyHelper.SetProperty(this, &SelfCollisionStiffness);
 	PropertyHelper.SetProperty(this, &SelfCollisionFriction);
+	PropertyHelper.SetProperty(this, &SelfCollisionDisableNeighborDistance, {}, ECollectionPropertyFlags::None); // Non animatable
+	PropertyHelper.SetPropertyString(this, &SelfCollisionLayers);
 
 	PropertyHelper.SetPropertyBool(this, &bUseSelfIntersections);
 	PropertyHelper.SetPropertyBool(this, &bUseGlobalIntersectionAnalysis);
diff --git a/Engine/Plugins/Experimental/ChaosClothAssetEditor/Source/ChaosClothAssetDataflowNodes/Private/ChaosClothAsset/TerminalNode.cpp b/Engine/Plugins/Experimental/ChaosClothAssetEditor/Source/ChaosClothAssetDataflowNodes/Private/ChaosClothAsset/TerminalNode.cpp
--- a/Engine/Plugins/Experimental/ChaosClothAssetEditor/Source/ChaosClothAssetDataflowNodes/Private/ChaosClothAsset/TerminalNode.cpp
+++ b/Engine/Plugins/Experimental/ChaosClothAssetEditor/Source/ChaosClothAssetDataflowNodes/Private/ChaosClothAsset/TerminalNode.cpp
@@ -31,6 +31,7 @@ namespace UE::Chaos::ClothAsset::Private
 			if (Cloth.HasValidData())
 			{
 				Checksum = Cloth.CalculateTypeHash(bIncludeWeightMapsTrue, Checksum);
+				Checksum = Cloth.CalculateUserDefinedAttributesTypeHash<int32>(ClothCollectionGroup::SimFaces, Checksum);
 			}
 			FCollectionClothSelectionConstFacade Selection(ClothCollection);
 			if (Selection.IsValid())
diff --git a/Engine/Plugins/Experimental/ChaosClothAssetEditor/Source/ChaosClothAssetDataflowNodes/Private/ChaosClothAsset/WeightedValueCustomization.cpp b/Engine/Plugins/Experimental/ChaosClothAssetEditor/Source/ChaosClothAssetDataflowNodes/Private/ChaosClothAsset/WeightedValueCustomization.cpp
--- a/Engine/Plugins/Experimental/ChaosClothAssetEditor/Source/ChaosClothAssetDataflowNodes/Private/ChaosClothAsset/WeightedValueCustomization.cpp
+++ b/Engine/Plugins/Experimental/ChaosClothAssetEditor/Source/ChaosClothAssetDataflowNodes/Private/ChaosClothAsset/WeightedValueCustomization.cpp
@@ -161,7 +161,7 @@ namespace UE::Chaos::ClothAsset
 			TWeakPtr<IPropertyHandle> HandleWeakPtr = PropertyHandle;
 			return
 				SNew(SEditableTextBox)
-				.ToolTipText(LOCTEXT("WeightMap", "The name of the weight map for this property."))
+				.ToolTipText(PropertyHandle->GetToolTipText())
 				.Text_Lambda([HandleWeakPtr, OverrideHandleWeakPtr]() -> FText
 					{
 						FString Text;
diff --git a/Engine/Plugins/Experimental/ChaosClothAssetEditor/Source/ChaosClothAssetDataflowNodes/Public/ChaosClothAsset/ConnectableValue.h b/Engine/Plugins/Experimental/ChaosClothAssetEditor/Source/ChaosClothAssetDataflowNodes/Public/ChaosClothAsset/ConnectableValue.h
new file mode 100644
--- /dev/null
+++ b/Engine/Plugins/Experimental/ChaosClothAssetEditor/Source/ChaosClothAssetDataflowNodes/Public/ChaosClothAsset/ConnectableValue.h
@@ -0,0 +1,34 @@
+// Copyright Epic Games, Inc. All Rights Reserved.
+#pragma once
+
+#include "UObject/ObjectMacros.h"
+#include "Containers/UnrealString.h"
+#include "ConnectableValue.generated.h"
+
+USTRUCT()
+struct FChaosClothAssetConnectableIStringValue
+{
+	GENERATED_BODY()
+
+	/** The value for this property. */
+	UPROPERTY(EditAnywhere, Category = "Value", Meta = (DataflowInput))
+	FString StringValue = TEXT("StringValue");
+
+	/** The string value override value for when the StringValue has a connection that replaces the provided string value. */
+	UPROPERTY(VisibleAnywhere, Category = "Value", Transient)
+	mutable FString StringValue_Override;  // _Override has a special meaning to the property customization, mutable because this property is set while getting the original value
+};
+
+USTRUCT()
+struct FChaosClothAssetConnectableIOStringValue
+{
+	GENERATED_BODY()
+
+	/** The value for this property. */
+	UPROPERTY(EditAnywhere, Category = "Value", Meta = (DataflowInput, DataflowOutput, DataflowPassthrough = "StringValue"))
+	FString StringValue = TEXT("StringValue");
+
+	/** The string value override value for when the StringValue has a connection that replaces the provided string value. */
+	UPROPERTY(VisibleAnywhere, Category = "Value", Transient)
+	mutable FString StringValue_Override;  // _Override has a special meaning to the property customization, mutable because this property is set while getting the original value
+};
diff --git a/Engine/Plugins/Experimental/ChaosClothAssetEditor/Source/ChaosClothAssetDataflowNodes/Public/ChaosClothAsset/SelectionToIntMapNode.h b/Engine/Plugins/Experimental/ChaosClothAssetEditor/Source/ChaosClothAssetDataflowNodes/Public/ChaosClothAsset/SelectionToIntMapNode.h
new file mode 100644
--- /dev/null
+++ b/Engine/Plugins/Experimental/ChaosClothAssetEditor/Source/ChaosClothAssetDataflowNodes/Public/ChaosClothAsset/SelectionToIntMapNode.h
@@ -0,0 +1,51 @@
+// Copyright Epic Games, Inc. All Rights Reserved.
+
+#pragma once 
+
+#include "Dataflow/DataflowNode.h"
+#include "GeometryCollection/ManagedArrayCollection.h"
+#include "ChaosClothAsset/ConnectableValue.h"
+#include "SelectionToIntMapNode.generated.h"
+
+/** Convert an integer index selection to an integer map. Map type will match the selection type.*/
+USTRUCT(Meta = (DataflowCloth))
+struct FChaosClothAssetSelectionToIntMapNode : public FDataflowNode
+{
+	GENERATED_USTRUCT_BODY()
+	DATAFLOW_NODE_DEFINE_INTERNAL(FChaosClothAssetSelectionToIntMapNode, "SelectionToIntMap", "Cloth", "Cloth Selection To Int Map")
+
+public:
+
+	UPROPERTY(Meta = (Dataflowinput, DataflowOutput, DataflowPassthrough = "Collection"))
+	FManagedArrayCollection Collection;
+
+	/** The name of the selection to convert. */
+	UPROPERTY(EditAnywhere, Category = "Selection To Int Map")
+	FChaosClothAssetConnectableIStringValue SelectionName;
+
+	/**
+	 * The name of the integer map attribute that will be added to the collection.
+	 * If left empty the same name as the selection name will be used instead.
+	 */
+	UPROPERTY(EditAnywhere, Category = "Selection To Int Map")
+	FChaosClothAssetConnectableIOStringValue IntMapName;
+
+
+	/** If the IntMapName already exists, keep existing values rather than overwriting them with 'Unselected Value'.*/
+	UPROPERTY(EditAnywhere, Category = "Selection To Int Map")
+	bool bKeepExistingUnselectedValues = true;
+
+	/** The value unselected elements receive. Unselected existing values can be kept by setting 'Keep Existing Unselected Values'*/
+	UPROPERTY(EditAnywhere, Category = "Selection To Int Map")
+	int32 UnselectedValue = 0;
+
+	/** The value selected elements receive. */
+	UPROPERTY(EditAnywhere, Category = "Selection To Int Map")
+	int32 SelectedValue = 1;
+
+	FChaosClothAssetSelectionToIntMapNode(const Dataflow::FNodeParameters& InParam, FGuid InGuid = FGuid::NewGuid());
+
+private:
+
+	virtual void Evaluate(Dataflow::FContext& Context, const FDataflowOutput* Out) const override;
+};
diff --git a/Engine/Plugins/Experimental/ChaosClothAssetEditor/Source/ChaosClothAssetDataflowNodes/Public/ChaosClothAsset/SimulationBaseConfigNode.h b/Engine/Plugins/Experimental/ChaosClothAssetEditor/Source/ChaosClothAssetDataflowNodes/Public/ChaosClothAsset/SimulationBaseConfigNode.h
--- a/Engine/Plugins/Experimental/ChaosClothAssetEditor/Source/ChaosClothAssetDataflowNodes/Public/ChaosClothAsset/SimulationBaseConfigNode.h
+++ b/Engine/Plugins/Experimental/ChaosClothAssetEditor/Source/ChaosClothAssetDataflowNodes/Public/ChaosClothAsset/SimulationBaseConfigNode.h
@@ -3,6 +3,7 @@
 #pragma once 
 
 #include "Chaos/CollectionPropertyFacade.h"
+#include "ChaosClothAsset/ConnectableValue.h"
 #include "ChaosClothAsset/WeightedValue.h"
 #include "Dataflow/DataflowNode.h"
 #include "GeometryCollection/ManagedArrayCollection.h"
@@ -52,6 +53,10 @@ protected:
 
 		int32 SetPropertyString(const FName& PropertyName, const FString& PropertyValue, const TArray<FName>& SimilarPropertyNames = {}, ECollectionPropertyFlags PropertyFlags = ECollectionPropertyFlags::Animatable);
 
+		int32 SetPropertyString(const FName& PropertyName, const FChaosClothAssetConnectableIStringValue& PropertyValue, const TArray<FName>& SimilarPropertyNames = {}, ECollectionPropertyFlags PropertyFlags = ECollectionPropertyFlags::Animatable);
+
+		template<typename T, typename ConnectableStringValueType, TEMPLATE_REQUIRES(TIsDerivedFrom<T, FChaosClothAssetSimulationBaseConfigNode>::Value)>
+		inline int32 SetPropertyString(const T* ConfigStruct, const ConnectableStringValueType* PropertyValue, const TArray<FName>& SimilarPropertyNames = {}, ECollectionPropertyFlags PropertyFlags = ECollectionPropertyFlags::Animatable);
 		template<typename T, TEMPLATE_REQUIRES(TIsDerivedFrom<T, FChaosClothAssetSimulationBaseConfigNode>::Value)>
 		inline int32 SetPropertyString(const T* ConfigStruct, const FString* PropertyValue, const TArray<FName>& SimilarPropertyNames = {}, ECollectionPropertyFlags PropertyFlags = ECollectionPropertyFlags::Animatable);
 
@@ -205,6 +210,19 @@ inline int32 FChaosClothAssetSimulationBaseConfigNode::FPropertyHelper::SetPrope
 	return SetPropertyWeighted(PropertyName, *PropertyValue, SimilarPropertyNames, PropertyFlags);
 }
 
+template<typename T, typename ConnectableStringValueType, typename TEnableIf<TIsDerivedFrom<T, FChaosClothAssetSimulationBaseConfigNode>::Value, int>::type>
+inline int32 FChaosClothAssetSimulationBaseConfigNode::FPropertyHelper::SetPropertyString(
+	const T* ConfigStruct, 
+	const ConnectableStringValueType* PropertyValue,
+	const TArray<FName>& SimilarPropertyNames, 
+	ECollectionPropertyFlags PropertyFlags)
+{
+	check(PropertyValue);
+	const FName PropertyName = FindPropertyNameByAddress(ConfigStruct, PropertyValue);
+	checkf(PropertyName != NAME_None, TEXT("Unknown property."));
+	return SetPropertyString(PropertyName, *PropertyValue, SimilarPropertyNames, PropertyFlags);
+}
+
 template<typename T, typename TEnableIf<TIsDerivedFrom<T, FChaosClothAssetSimulationBaseConfigNode>::Value, int>::type>
 inline FName FChaosClothAssetSimulationBaseConfigNode::FPropertyHelper::FindPropertyNameByAddress(const T* ConfigStruct, const void* PropertyAddress)
 {
diff --git a/Engine/Plugins/Experimental/ChaosClothAssetEditor/Source/ChaosClothAssetDataflowNodes/Public/ChaosClothAsset/SimulationSelfCollisionConfigNode.h b/Engine/Plugins/Experimental/ChaosClothAssetEditor/Source/ChaosClothAssetDataflowNodes/Public/ChaosClothAsset/SimulationSelfCollisionConfigNode.h
--- a/Engine/Plugins/Experimental/ChaosClothAssetEditor/Source/ChaosClothAssetDataflowNodes/Public/ChaosClothAsset/SimulationSelfCollisionConfigNode.h
+++ b/Engine/Plugins/Experimental/ChaosClothAssetEditor/Source/ChaosClothAssetDataflowNodes/Public/ChaosClothAsset/SimulationSelfCollisionConfigNode.h
@@ -3,6 +3,7 @@
 #pragma once 
 
 #include "ChaosClothAsset/SimulationBaseConfigNode.h"
+#include "ChaosClothAsset/ConnectableValue.h"
 #include "SimulationSelfCollisionConfigNode.generated.h"
 
 /** Self-collision repulsion forces (point-face) properties configuration node. */
@@ -25,6 +26,17 @@ public:
 	UPROPERTY(EditAnywhere, Category = "Self-Collision Properties", meta = (UIMin = "0", UIMax = "1", ClampMin = "0", ClampMax = "1"))
 	float SelfCollisionFriction = 0.0f;
 
+	/** Disabled neighbor collision ring. Collisions are disabled between vertices within this N-ring connectivity distance.*/
+	UPROPERTY(EditAnywhere, Category = "Self-Collision Properties", meta = (UIMin = "1", UIMax = "5", ClampMin = "1"))
+	int32 SelfCollisionDisableNeighborDistance = 5;
+
+	/** Self collision layers face int map. Generate this map using the SelectionsToIntMap node with SimFace Selections.
+	* Faces labeled with -1 will collide normally without any layering behavior.
+	* Faces labeled with any other number will keep higher layer numbers outside lower layer numbers (outside = front facing normal direction).
+	*/
+	UPROPERTY(EditAnywhere, Category = "Self-Collision Properties")
+	FChaosClothAssetConnectableIStringValue SelfCollisionLayers = { TEXT("SelfCollisionLayers") };
+
 	/** Enable self intersection resolution. This will try to fix any cloth intersections that are not handled by collision repulsions. */
 	UPROPERTY(EditAnywhere, Category = "Experimental")
 	bool bUseSelfIntersections = false;
diff --git a/Engine/Source/Runtime/Experimental/Chaos/Private/Chaos/PBDCollisionSpringConstraintsBase.cpp b/Engine/Source/Runtime/Experimental/Chaos/Private/Chaos/PBDCollisionSpringConstraintsBase.cpp
--- a/Engine/Source/Runtime/Experimental/Chaos/Private/Chaos/PBDCollisionSpringConstraintsBase.cpp
+++ b/Engine/Source/Runtime/Experimental/Chaos/Private/Chaos/PBDCollisionSpringConstraintsBase.cpp
@@ -38,6 +38,7 @@ FPBDCollisionSpringConstraintsBase::FPBDCollisionSpringConstraintsBase(
 	const FTriangleMesh& InTriangleMesh,
 	const TArray<FSolverVec3>* InReferencePositions,
 	TSet<TVec2<int32>>&& InDisabledCollisionElements,
+	const TConstArrayView<int32>& InSelfCollisionLayers,
 	const FSolverReal InThickness,
 	const FSolverReal InStiffness,
 	const FSolverReal InFrictionCoefficient,
@@ -54,6 +55,41 @@ FPBDCollisionSpringConstraintsBase::FPBDCollisionSpringConstraintsBase(
 	, NumParticles(InNumParticles)
 	, bGlobalIntersectionAnalysis(false)
 {
+	UpdateCollisionLayers(InSelfCollisionLayers);
+}
+
+void FPBDCollisionSpringConstraintsBase::UpdateCollisionLayers(const TConstArrayView<int32>& InFaceCollisionLayers)
+{
+	if (InFaceCollisionLayers.Num() != Elements.Num())
+	{
+		// Reset collision layers
+		FaceCollisionLayers = TConstArrayView<int32>();
+		VertexCollisionLayers.Reset();
+	}
+	else
+	{
+		FaceCollisionLayers = InFaceCollisionLayers;
+		VertexCollisionLayers.SetNumUninitialized(NumParticles);
+
+		TConstArrayView<TArray<int32>> PointToTriangle = TriangleMesh.GetPointToTriangleMap();
+		for (int32 ParticleIndexNoOffset = 0; ParticleIndexNoOffset < NumParticles; ++ParticleIndexNoOffset)
+		{
+			const int32 ParticleIndex = ParticleIndexNoOffset + Offset;
+			TVec2<int32>& VertexCollisionLayer = VertexCollisionLayers[ParticleIndexNoOffset];
+			VertexCollisionLayer = TVec2<int32>(INDEX_NONE);
+			for (const int32 FaceIndex : PointToTriangle[ParticleIndex])
+			{
+				if (FaceCollisionLayers[FaceIndex] != INDEX_NONE)
+				{
+					VertexCollisionLayer[0] = VertexCollisionLayer[0] == INDEX_NONE ? 
+						FaceCollisionLayers[FaceIndex] : FMath::Min(FaceCollisionLayers[FaceIndex], VertexCollisionLayer[0]);
+
+					VertexCollisionLayer[1] = VertexCollisionLayer[1] == INDEX_NONE ?
+						FaceCollisionLayers[FaceIndex] : FMath::Max(FaceCollisionLayers[FaceIndex], VertexCollisionLayer[1]);
+				}
+			}
+		}
+	}
 }
 
 void FPBDCollisionSpringConstraintsBase::Init(const FSolverParticles& Particles)
@@ -112,12 +148,24 @@ void FPBDCollisionSpringConstraintsBase::Init(const SolverParticlesOrRange& Part
 				}
 				constexpr FSolverReal ExtraThicknessMult = 1.5f;
 
+				const bool bVertexHasCollisionLayers = VertexCollisionLayers.IsValidIndex(i) && VertexCollisionLayers[i][0] != INDEX_NONE;
+				check(!bVertexHasCollisionLayers || VertexCollisionLayers[i][0] <= VertexCollisionLayers[i][1]);
+
 				TArray< TTriangleCollisionPoint<FSolverReal> > Result;
 				if (TriangleMesh.PointProximityQuery(Spatial, static_cast<const TArrayView<const FSolverVec3>&>(Particles.XArray()), Index, Particles.X(Index), Thickness * ExtraThicknessMult, Thickness * ExtraThicknessMult,
-					[this, &VertexGIAColors, &TriangleGIAColors](const int32 PointIndex, const int32 TriangleIndex)->bool
+					[this, bVertexHasCollisionLayers, &VertexGIAColors, &TriangleGIAColors](const int32 PointIndex, const int32 TriangleIndex)->bool
 					{
-						const TVector<int32, 3>& Elem = Elements[TriangleIndex];
-						if (bGlobalIntersectionAnalysis)
+						const TVector<int32, 3>& Elem = Elements[TriangleIndex];						
+
+						bool bUseCollisionLayerOverride = false;
+						if (bVertexHasCollisionLayers && FaceCollisionLayers[TriangleIndex] != INDEX_NONE)
+						{
+							if (FaceCollisionLayers[TriangleIndex] < VertexCollisionLayers[PointIndex - Offset][0] || FaceCollisionLayers[TriangleIndex] > VertexCollisionLayers[PointIndex - Offset][1])
+							{
+								bUseCollisionLayerOverride = true;
+							}
+						}
+						if (!bUseCollisionLayerOverride && bGlobalIntersectionAnalysis)
 						{
 							const bool bIsAnyBoundary = VertexGIAColors[PointIndex].IsBoundary()
 								|| VertexGIAColors[Elem[0]].IsBoundary() 
@@ -180,20 +228,44 @@ void FPBDCollisionSpringConstraintsBase::Init(const SolverParticlesOrRange& Part
 							}
 						}
 
-						// NOTE: CollisionPoint.Normal has already been flipped to point toward the Point, so need to recalculate here.
-						const TTriangle<FSolverReal> Triangle(Particles.X(Elem[0]), Particles.X(Elem[1]), Particles.X(Elem[2]));
-						bool bFlipNormal = (Particles.X(Index) - CollisionPoint.Location).Dot(Triangle.GetNormal()) < 0; // Is Point currently behind Triangle?
-
-						// Doing a check against ANY (plus the TriangleGIAColors which captures sub-triangle intersections) seems to work better than checking against ALL vertex colors where the triangle must agree.
-						// In particular, it's better at handling thin regions of intersection where a single vertex or line of vertices intersect through faces.
-						if (bGlobalIntersectionAnalysis &&
-							(FPBDTriangleMeshCollisions::FGIAColor::ShouldFlipNormal(VertexGIAColors[Index], VertexGIAColors[Elem[0]]) ||
-								FPBDTriangleMeshCollisions::FGIAColor::ShouldFlipNormal(VertexGIAColors[Index], VertexGIAColors[Elem[1]]) ||
-								FPBDTriangleMeshCollisions::FGIAColor::ShouldFlipNormal(VertexGIAColors[Index], VertexGIAColors[Elem[2]]) ||
-								FPBDTriangleMeshCollisions::FGIAColor::ShouldFlipNormal(VertexGIAColors[Index], TriangleGIAColors[CollisionPoint.Indices[1]])))
+						bool bFlipNormal = false;
+						// Check collision layers
+						bool bUseCollisionLayerOverride = false;
+						if (bVertexHasCollisionLayers && FaceCollisionLayers[CollisionPoint.Indices[1]] != INDEX_NONE)
+						{
+							if (FaceCollisionLayers[CollisionPoint.Indices[1]] < VertexCollisionLayers[i][0])
+							{
+								// Face is lower layer than the vertex. Vertex should always be in front of face (as UE sees it).
+								// NOTE: Chaos internal winding order for normals is reversed, so flip normal in this case.
+								bFlipNormal = true;
+								bUseCollisionLayerOverride = true;
+							}
+							else if (FaceCollisionLayers[CollisionPoint.Indices[1]] > VertexCollisionLayers[i][1])
+							{
+								// Face is higher layer than the vertex. Vertex should always be behind face (as UE sees it).
+								// NOTE: Chaos internal winding order for normals is reversed, so don't flip normal in this case.
+								bFlipNormal = false;
+								bUseCollisionLayerOverride = true;
+							}
+						}
+
+						if (!bUseCollisionLayerOverride)
 						{
+							// NOTE: CollisionPoint.Normal has already been flipped to point toward the Point, so need to recalculate here.
+							const TTriangle<FSolverReal> Triangle(Particles.X(Elem[0]), Particles.X(Elem[1]), Particles.X(Elem[2]));
+							bFlipNormal = (Particles.X(Index) - CollisionPoint.Location).Dot(Triangle.GetNormal()) < 0; // Is Point currently behind Triangle?
+							// Doing a check against ANY (plus the TriangleGIAColors which captures sub-triangle intersections) seems to work better than checking against ALL vertex colors where the triangle must agree.
+							// In particular, it's better at handling thin regions of intersection where a single vertex or line of vertices intersect through faces.
 							// Want Point to push to opposite side of triangle
-							bFlipNormal = !bFlipNormal;
+							if (bGlobalIntersectionAnalysis &&
+								(FPBDTriangleMeshCollisions::FGIAColor::ShouldFlipNormal(VertexGIAColors[Index], VertexGIAColors[Elem[0]]) ||
+									FPBDTriangleMeshCollisions::FGIAColor::ShouldFlipNormal(VertexGIAColors[Index], VertexGIAColors[Elem[1]]) ||
+									FPBDTriangleMeshCollisions::FGIAColor::ShouldFlipNormal(VertexGIAColors[Index], VertexGIAColors[Elem[2]]) ||
+									FPBDTriangleMeshCollisions::FGIAColor::ShouldFlipNormal(VertexGIAColors[Index], TriangleGIAColors[CollisionPoint.Indices[1]])))
+							{
+
+								bFlipNormal = !bFlipNormal;
+							}
 						}
 						const int32 IndexToWrite = ConstraintIndex.fetch_add(1);
 
diff --git a/Engine/Source/Runtime/Experimental/Chaos/Public/Chaos/CollectionPropertyFacade.h b/Engine/Source/Runtime/Experimental/Chaos/Public/Chaos/CollectionPropertyFacade.h
--- a/Engine/Source/Runtime/Experimental/Chaos/Public/Chaos/CollectionPropertyFacade.h
+++ b/Engine/Source/Runtime/Experimental/Chaos/Public/Chaos/CollectionPropertyFacade.h
@@ -537,27 +537,27 @@ PRAGMA_ENABLE_DEPRECATION_WARNINGS
 	{ \
 		return PropertyCollection.IsAnimatable(PropertyName##Name.ToString(), bDefault); \
 	} \
-	Type GetLow##PropertyName(const FCollectionPropertyConstFacade& PropertyCollection, const Type& Default) \
+	static Type GetLow##PropertyName(const FCollectionPropertyConstFacade& PropertyCollection, const Type& Default) \
 	{ \
 		return PropertyCollection.GetLowValue<Type>(PropertyName##Name.ToString(), Default); \
 	} \
-	Type GetHigh##PropertyName(const FCollectionPropertyConstFacade& PropertyCollection, const Type& Default) \
+	static Type GetHigh##PropertyName(const FCollectionPropertyConstFacade& PropertyCollection, const Type& Default) \
 	{ \
 		return PropertyCollection.GetHighValue<Type>(PropertyName##Name.ToString(), Default); \
 	} \
-	TPair<Type, Type> GetWeighted##PropertyName(const FCollectionPropertyConstFacade& PropertyCollection, const Type& Default) \
+	static TPair<Type, Type> GetWeighted##PropertyName(const FCollectionPropertyConstFacade& PropertyCollection, const Type& Default) \
 	{ \
 		return PropertyCollection.GetWeightedValue<Type>(PropertyName##Name.ToString(), Default); \
 	} \
-	FVector2f GetWeightedFloat##PropertyName(const FCollectionPropertyConstFacade& PropertyCollection, const float& Default) \
+	static FVector2f GetWeightedFloat##PropertyName(const FCollectionPropertyConstFacade& PropertyCollection, const float& Default) \
 	{ \
 		return PropertyCollection.GetWeightedFloatValue(PropertyName##Name.ToString(), Default); \
 	} \
-	Type Get##PropertyName(const FCollectionPropertyConstFacade& PropertyCollection, const Type& Default) \
+	static Type Get##PropertyName(const FCollectionPropertyConstFacade& PropertyCollection, const Type& Default) \
 	{ \
 		return PropertyCollection.GetValue<Type>(PropertyName##Name.ToString(), Default); \
 	} \
-	FString Get##PropertyName##String(const FCollectionPropertyConstFacade& PropertyCollection, const FString& Default) \
+	static FString Get##PropertyName##String(const FCollectionPropertyConstFacade& PropertyCollection, const FString& Default) \
 	{ \
 		return PropertyCollection.GetStringValue(PropertyName##Name.ToString(), Default); \
 	} \
diff --git a/Engine/Source/Runtime/Experimental/Chaos/Public/Chaos/PBDCollisionSpringConstraints.h b/Engine/Source/Runtime/Experimental/Chaos/Public/Chaos/PBDCollisionSpringConstraints.h
--- a/Engine/Source/Runtime/Experimental/Chaos/Public/Chaos/PBDCollisionSpringConstraints.h
+++ b/Engine/Source/Runtime/Experimental/Chaos/Public/Chaos/PBDCollisionSpringConstraints.h
@@ -17,12 +17,38 @@ class FPBDCollisionSpringConstraints : public FPBDCollisionSpringConstraintsBase
 public:
 	static constexpr FSolverReal MinFrictionCoefficient = (FSolverReal)0.;
 	static constexpr FSolverReal MaxFrictionCoefficient = (FSolverReal)1.;
+	static constexpr int32 DefaultSelfCollisionDisableNeighborDistance = 5;
 
 	static bool IsEnabled(const FCollectionPropertyConstFacade& PropertyCollection)
 	{
 		return IsSelfCollisionStiffnessEnabled(PropertyCollection, false);
 	}
+	FPBDCollisionSpringConstraints(
+		const int32 InOffset,
+		const int32 InNumParticles,
+		const FTriangleMesh& InTriangleMesh,
+		const TArray<FSolverVec3>* InRestPositions,
+		const TMap<FString, TConstArrayView<int32>>& FaceIntMaps,
+		const FCollectionPropertyConstFacade& PropertyCollection)
+		: Base(
+			InOffset,
+			InNumParticles,
+			InTriangleMesh,
+			InRestPositions,
+			GenerateDisabledCollisionElements(InTriangleMesh, PropertyCollection),
+			FaceIntMaps.FindRef(GetSelfCollisionLayersString(PropertyCollection, SelfCollisionLayersName.ToString())),
+			(FSolverReal)FMath::Max(GetSelfCollisionThickness(PropertyCollection, Base::BackCompatThickness), 0.f),
+			(FSolverReal)FMath::Clamp(GetSelfCollisionStiffness(PropertyCollection, Base::BackCompatStiffness), 0.f, 1.f),
+			FMath::Clamp((FSolverReal)GetSelfCollisionFriction(PropertyCollection, Base::BackCompatFrictionCoefficient), MinFrictionCoefficient, MaxFrictionCoefficient),
+			(FSolverReal)GetSelfCollisionProximityStiffness(PropertyCollection, Base::DefaultProximityStiffness))
+		, SelfCollisionThicknessIndex(PropertyCollection)
+		, SelfCollisionStiffnessIndex(PropertyCollection)
+		, SelfCollisionFrictionIndex(PropertyCollection)
+		, SelfCollisionProximityStiffnessIndex(PropertyCollection)
+		, SelfCollisionLayersIndex(PropertyCollection)
+	{}
 
+	UE_DEPRECATED(5.4, "Use constructor with FaceIntMaps")
 	FPBDCollisionSpringConstraints(
 		const int32 InOffset,
 		const int32 InNumParticles,
@@ -36,6 +62,7 @@ public:
 			InTriangleMesh,
 			InRestPositions,
 			MoveTemp(InDisabledCollisionElements),
+			TConstArrayView<int32>(),
 			(FSolverReal)FMath::Max(GetSelfCollisionThickness(PropertyCollection, Base::BackCompatThickness), 0.f),
 			(FSolverReal)FMath::Clamp(GetSelfCollisionStiffness(PropertyCollection, Base::BackCompatStiffness), 0.f, 1.f),
 			FMath::Clamp((FSolverReal)GetSelfCollisionFriction(PropertyCollection, Base::BackCompatFrictionCoefficient), MinFrictionCoefficient, MaxFrictionCoefficient),
@@ -44,6 +71,7 @@ public:
 		, SelfCollisionStiffnessIndex(PropertyCollection)
 		, SelfCollisionFrictionIndex(PropertyCollection)
 		, SelfCollisionProximityStiffnessIndex(PropertyCollection)
+		, SelfCollisionLayersIndex(PropertyCollection)
 	{}
 
 	FPBDCollisionSpringConstraints(
@@ -61,6 +89,7 @@ public:
 			InTriangleMesh,
 			InRestPositions,
 			MoveTemp(InDisabledCollisionElements),
+			TConstArrayView<int32>(),
 			InThickness,
 			InStiffness,
 			InFrictionCoefficient)
@@ -68,13 +97,15 @@ public:
 		, SelfCollisionStiffnessIndex(ForceInit)
 		, SelfCollisionFrictionIndex(ForceInit)
 		, SelfCollisionProximityStiffnessIndex(ForceInit)
+		, SelfCollisionLayersIndex(ForceInit)
 	{}
 
 	virtual ~FPBDCollisionSpringConstraints() override {}
 
 	using Base::Init;
 
-	void SetProperties(const FCollectionPropertyConstFacade& PropertyCollection)
+	void SetProperties(const FCollectionPropertyConstFacade& PropertyCollection,
+		const TMap<FString, TConstArrayView<int32>>& FaceIntMaps)
 	{
 		if (IsSelfCollisionThicknessMutable(PropertyCollection))
 		{
@@ -92,6 +123,20 @@ public:
 		{
 			ProximityStiffness = GetSelfCollisionProximityStiffness(PropertyCollection);
 		}
+		if (IsSelfCollisionLayersMutable(PropertyCollection))
+		{
+			if (IsSelfCollisionLayersStringDirty(PropertyCollection))
+			{
+				const FString& FaceIntMapName = GetSelfCollisionLayersString(PropertyCollection);
+				Base::UpdateCollisionLayers(FaceIntMaps.FindRef(FaceIntMapName));
+			}
+		}
+	}
+
+	UE_DEPRECATED(5.4, "Use SetProperties with FaceIntMaps")
+	void SetProperties(const FCollectionPropertyConstFacade& PropertyCollection)
+	{
+		SetProperties(PropertyCollection, TMap<FString, TConstArrayView<int32>>());
 	}
 
 private:
@@ -100,10 +145,33 @@ private:
 	using Base::FrictionCoefficient;
 	using Base::ProximityStiffness;
 
+	static TSet<TVec2<int32>> GenerateDisabledCollisionElements(
+		const FTriangleMesh& TriangleMesh, const FCollectionPropertyConstFacade& PropertyCollection)
+	{
+		TRACE_CPUPROFILER_EVENT_SCOPE(FPBDCollisionSpringConstraints_GenerateDisabledCollisionElements);
+		const int32 DisabledCollisionElementsN = GetSelfCollisionDisableNeighborDistance(PropertyCollection, DefaultSelfCollisionDisableNeighborDistance);
+		TSet<TVec2<int32>> DisabledCollisionElements;  // TODO: Is this needed? Turn this into a bit array?
+
+		const TVec2<int32> Range = TriangleMesh.GetVertexRange();
+		for (int32 Index = Range[0]; Index <= Range[1]; ++Index)
+		{
+			const TSet<int32> Neighbors = TriangleMesh.GetNRing(Index, DisabledCollisionElementsN);
+			for (int32 Element : Neighbors)
+			{
+				check(Index != Element);
+				DisabledCollisionElements.Emplace(TVec2<int32>(Index, Element));
+				DisabledCollisionElements.Emplace(TVec2<int32>(Element, Index));
+			}
+		}
+		return DisabledCollisionElements;
+	}
+
 	UE_CHAOS_DECLARE_PROPERTYCOLLECTION_NAME(SelfCollisionThickness, float);
 	UE_CHAOS_DECLARE_PROPERTYCOLLECTION_NAME(SelfCollisionStiffness, float);
 	UE_CHAOS_DECLARE_PROPERTYCOLLECTION_NAME(SelfCollisionFriction, float);
 	UE_CHAOS_DECLARE_PROPERTYCOLLECTION_NAME(SelfCollisionProximityStiffness, float);
+	UE_CHAOS_DECLARE_INDEXLESS_PROPERTYCOLLECTION_NAME(SelfCollisionDisableNeighborDistance, int32);
+	UE_CHAOS_DECLARE_PROPERTYCOLLECTION_NAME(SelfCollisionLayers, float);
 };
 
 }  // End namespace Chaos::Softs
diff --git a/Engine/Source/Runtime/Experimental/Chaos/Public/Chaos/PBDCollisionSpringConstraintsBase.h b/Engine/Source/Runtime/Experimental/Chaos/Public/Chaos/PBDCollisionSpringConstraintsBase.h
--- a/Engine/Source/Runtime/Experimental/Chaos/Public/Chaos/PBDCollisionSpringConstraintsBase.h
+++ b/Engine/Source/Runtime/Experimental/Chaos/Public/Chaos/PBDCollisionSpringConstraintsBase.h
@@ -32,11 +32,27 @@ public:
 		const FTriangleMesh& InTriangleMesh,
 		const TArray<FSolverVec3>* InReferencePositions,
 		TSet<TVec2<int32>>&& InDisabledCollisionElements,
+		const TConstArrayView<int32>& InSelfCollisionLayers,
 		const FSolverReal InThickness = BackCompatThickness,
 		const FSolverReal InStiffness = BackCompatStiffness,
 		const FSolverReal InFrictionCoefficient = BackCompatFrictionCoefficient,
 		const FSolverReal InProximityStiffness = DefaultProximityStiffness);
 
+	UE_DEPRECATED(5.4, "Use constructor with SelfCollisionLayers")
+	FPBDCollisionSpringConstraintsBase(
+		const int32 InOffset,
+		const int32 InNumParticles,
+		const FTriangleMesh& InTriangleMesh,
+		const TArray<FSolverVec3>* InReferencePositions,
+		TSet<TVec2<int32>>&& InDisabledCollisionElements,
+		const FSolverReal InThickness = BackCompatThickness,
+		const FSolverReal InStiffness = BackCompatStiffness,
+		const FSolverReal InFrictionCoefficient = BackCompatFrictionCoefficient,
+		const FSolverReal InProximityStiffness = DefaultProximityStiffness)
+		: FPBDCollisionSpringConstraintsBase(InOffset, InNumParticles, InTriangleMesh, InReferencePositions, MoveTemp(InDisabledCollisionElements),
+			TConstArrayView<int32>(), InThickness, InStiffness, InFrictionCoefficient, InProximityStiffness)
+	{}
+
 	virtual ~FPBDCollisionSpringConstraintsBase() {}
 
 	UE_DEPRECATED(5.0, "Use Init(Particles, Spatial, GIAColors) instead.")
@@ -102,6 +118,9 @@ public:
 	}
 
 	CHAOS_API void UpdateLinearSystem(const FSolverParticlesRange& Particles, const FSolverReal Dt, FEvolutionLinearSystem& LinearSystem) const;
+	
+	TConstArrayView<int32> GetFaceCollisionLayers() const { return FaceCollisionLayers; }
+	const TArray<TVector<int32, 2>>& GetVertexCollisionLayers() const { return VertexCollisionLayers; }
 
 protected:
 	TArray<TVec4<int32>> Constraints;
@@ -112,11 +131,15 @@ protected:
 	FSolverReal FrictionCoefficient;
 	FSolverReal ProximityStiffness; // (actual spring stiffness for force-based solver)
 
+	CHAOS_API void UpdateCollisionLayers(const TConstArrayView<int32>& InFaceCollisionLayers);
+
 private:
 	const FTriangleMesh& TriangleMesh;
 	const TArray<TVec3<int32>>& Elements;
 	const TArray<FSolverVec3>* ReferencePositions;
 	const TSet<TVec2<int32>> DisabledCollisionElements;  // TODO: Make this a bitarray
+	TConstArrayView<int32> FaceCollisionLayers;
+	TArray<TVector<int32, 2>> VertexCollisionLayers; // Only non-empty if FaceCollisionLayers is non-empty. Values are Min and Max layers for that vertex
 
 	int32 Offset;
 	int32 NumParticles;
