Chaos Cloth Asset - Fixed cooking with less than 12 skin weight influences.

diff --git a/Engine/Plugins/Experimental/ChaosClothAsset/Source/ChaosClothAssetEngine/Private/ChaosClothAsset/ClothAsset.cpp b/Engine/Plugins/Experimental/ChaosClothAsset/Source/ChaosClothAssetEngine/Private/ChaosClothAsset/ClothAsset.cpp
--- a/Engine/Plugins/Experimental/ChaosClothAsset/Source/ChaosClothAssetEngine/Private/ChaosClothAsset/ClothAsset.cpp
+++ b/Engine/Plugins/Experimental/ChaosClothAsset/Source/ChaosClothAssetEngine/Private/ChaosClothAsset/ClothAsset.cpp
@@ -558,15 +558,19 @@ void UChaosClothAsset::BuildMeshModel()
 
 	const int32 NumLods = ClothCollections.Num();
 
-	// Build each LOD
+	// Clear current LOD models
 	check(MeshModel);  // MeshModel should always be created in the Cloth Asset constructor WITH_EDITORONLY_DATA
-	MeshModel->LODModels.Empty();
+	MeshModel->LODModels.Reset(NumLods);
 
+	// Get the running platform
+	ITargetPlatformManagerModule& TargetPlatformManager = GetTargetPlatformManagerRef();
+	ITargetPlatform* const TargetPlatform = TargetPlatformManager.GetRunningTargetPlatform();
+
+	// Rebuild each LOD models
 	for (int32 LodIndex = 0; LodIndex < NumLods; ++LodIndex)
 	{
-		FSkeletalMeshLODModel* LODModel = new FSkeletalMeshLODModel();
-		FBuilder::BuildLod(*LODModel, *this, LodIndex);
-		MeshModel->LODModels.Add(LODModel);
+		MeshModel->LODModels.Add(new FSkeletalMeshLODModel());
+		BuildLODModel(TargetPlatform, LodIndex);
 	}
 }
 #endif  // #if WITH_EDITOR
@@ -630,6 +634,12 @@ void UChaosClothAsset::CacheDerivedData(FSkinnedAssetCompilationContext* Context
 
 }
 
+void UChaosClothAsset::BuildLODModel(const ITargetPlatform* TargetPlatform, int32 LODIndex)
+{
+	check(MeshModel && MeshModel->LODModels.IsValidIndex(LODIndex));
+	FBuilder::BuildLod(MeshModel->LODModels[LODIndex], *this, LODIndex, TargetPlatform);
+}
+
 FString UChaosClothAsset::BuildDerivedDataKey(const ITargetPlatform* TargetPlatform)
 {
 	FString KeySuffix(TEXT(""));
diff --git a/Engine/Plugins/Experimental/ChaosClothAsset/Source/ChaosClothAssetEngine/Private/ChaosClothAsset/ClothAssetBuilder.cpp b/Engine/Plugins/Experimental/ChaosClothAsset/Source/ChaosClothAssetEngine/Private/ChaosClothAsset/ClothAssetBuilder.cpp
--- a/Engine/Plugins/Experimental/ChaosClothAsset/Source/ChaosClothAssetEngine/Private/ChaosClothAsset/ClothAssetBuilder.cpp
+++ b/Engine/Plugins/Experimental/ChaosClothAsset/Source/ChaosClothAssetEngine/Private/ChaosClothAsset/ClothAssetBuilder.cpp
@@ -7,6 +7,7 @@
 #include "ChaosClothAsset/ClothSimulationModel.h"
 #include "ChaosClothAsset/CollectionClothFacade.h"
 #include "Chaos/CollectionPropertyFacade.h"
+#include "Engine/RendererSettings.h"
 #include "Engine/SkeletalMesh.h"
 #include "Engine/SkinnedAssetCommon.h"
 #include "Misc/ScopedSlowTask.h"
@@ -44,13 +45,69 @@ namespace UE::Chaos::ClothAsset::Private
 			ClothSimulationModel->GetIndices(LodIndex) :
 			TConstArrayView<uint32>();
 	}
+
+	int32 ConformSkinWeightsToMaxInfluences(
+		const TArray<int32>& InBoneIndices,
+		const TArray<float>& InBoneWeights,
+		TArray<int32>& OutBoneIndices,
+		TArray<float>& OutBoneWeights,
+		int32 MaxInfluences)
+	{
+		check(InBoneIndices.Num() == InBoneWeights.Num());
+		const int32 NumInfluences = InBoneIndices.Num();
+
+		// Sort the influences by bone weight
+		TArray<int32, TInlineAllocator<MAX_TOTAL_INFLUENCES>> SortedInfluences;
+		SortedInfluences.Reserve(NumInfluences);
+		for (int32 Index = 0; Index < NumInfluences; ++Index)
+		{
+			SortedInfluences.Add(Index);
+		}
+		SortedInfluences.Sort([&InBoneWeights](int32 Index0, int32 Index1) { return InBoneWeights[Index0] > InBoneWeights[Index1]; });
+
+		// Copy the weights by order of the most influential to the less
+		OutBoneIndices.Reset(MAX_TOTAL_INFLUENCES);
+		OutBoneWeights.Reset(MAX_TOTAL_INFLUENCES);
+
+		float BoneWeightsSum = 0.f;
+		
+		for (int32 Index = 0; Index < FMath::Min(NumInfluences, MaxInfluences); ++Index)
+		{
+			const float BoneWeight = InBoneWeights[SortedInfluences[Index]];
+
+			if (!FMath::IsNearlyZero(BoneWeight))
+			{
+				OutBoneIndices.Add(InBoneIndices[SortedInfluences[Index]]);
+				OutBoneWeights.Add(BoneWeight);
+				BoneWeightsSum += BoneWeight;
+			}
+			else
+			{
+				break;  // Found zero weights, early exit
+			}
+		}
+		check(OutBoneIndices.Num() == OutBoneWeights.Num());
+
+		if (BoneWeightsSum != 0.f && !FMath::IsNearlyEqual(BoneWeightsSum, 1.f))
+		{
+			const float BoneWeightsSumRecip = 1.f / BoneWeightsSum;
+			for (float& OutBoneWeight : OutBoneWeights)
+			{
+				OutBoneWeight *= BoneWeightsSumRecip;
+			}
+		}
+
+		return OutBoneWeights.Num();
+	}
 }  // End namespace UE::Chaos::ClothAsset::Private
 
-void UChaosClothAsset::FBuilder::BuildLod(FSkeletalMeshLODModel& LODModel, const UChaosClothAsset& ClothAsset, int32 LodIndex)
+void UChaosClothAsset::FBuilder::BuildLod(FSkeletalMeshLODModel& LODModel, const UChaosClothAsset& ClothAsset, int32 LodIndex, const ITargetPlatform* TargetPlatform)
 {
 	using namespace UE::Chaos::ClothAsset;
 	using namespace ::Chaos::Softs;
 
+	check(TargetPlatform);
+
 	// Start from an empty LODModel
 	LODModel.Empty();
 
@@ -197,14 +254,25 @@ void UChaosClothAsset::FBuilder::BuildLod(FSkeletalMeshLODModel& LODModel, const
 				SoftVertex.UVs[TexCoord] = RenderUVs[TexCoord];
 			}
 
-			const int32 NumInfluences = PatternRenderBoneIndices[VertexIndex].Num();
-			check(NumInfluences <= MAX_TOTAL_INFLUENCES);
+			// Conform to the platform max number of influences
+			const int32 MaxBoneInfluencesFromUnlimitedBoneInfluences = FGPUBaseSkinVertexFactory::GetUnlimitedBoneInfluences(TargetPlatform) ? MAX_TOTAL_INFLUENCES : EXTRA_BONE_INFLUENCES;
+			const int32 MaxBoneInfluencesFromPlatformProjectSettings = FGPUBaseSkinVertexFactory::GetBoneInfluenceLimitForAsset(0, TargetPlatform);
+			const int32 MaxNumInfluences = FMath::Min(MaxBoneInfluencesFromUnlimitedBoneInfluences, MaxBoneInfluencesFromPlatformProjectSettings);
+
+			TArray<int32> BoneIndices;
+			TArray<float> BoneWeights;
+			const int32 NumInfluences = Private::ConformSkinWeightsToMaxInfluences(
+				PatternRenderBoneIndices[VertexIndex],
+				PatternRenderBoneWeights[VertexIndex],
+				BoneIndices,
+				BoneWeights,
+				MaxNumInfluences);
 
 			// Add all of the bones that have non-zero influence to the section's bone map and keep track of the order
 			// that we added the reference bone via CurSectionBoneMapNum
 			for (int32 Influence = 0; Influence < NumInfluences; ++Influence)
 			{
-				const FBoneIndexType InfluenceBone = (FBoneIndexType)PatternRenderBoneIndices[VertexIndex][Influence];
+				const FBoneIndexType InfluenceBone = (FBoneIndexType)BoneIndices[Influence];
 
 				if (ReferenceToSectionBoneMap.Contains(InfluenceBone) == false)
 				{
@@ -216,8 +284,8 @@ void UChaosClothAsset::FBuilder::BuildLod(FSkeletalMeshLODModel& LODModel, const
 			int32 Influence = 0;
 			for (; Influence < NumInfluences; ++Influence)
 			{
-				const FBoneIndexType InfluenceBone = (FBoneIndexType)PatternRenderBoneIndices[VertexIndex][Influence];
-				const float InWeight = PatternRenderBoneWeights[VertexIndex][Influence];
+				const FBoneIndexType InfluenceBone = (FBoneIndexType)BoneIndices[Influence];
+				const float InWeight = BoneWeights[Influence];
 				const uint16 InfluenceWeight = static_cast<uint16>(InWeight * static_cast<float>(UE::AnimationCore::MaxRawBoneWeight) + 0.5f);
 
 				// FSoftSkinVertex::InfluenceBones contain indices into the section's bone map and not the reference
diff --git a/Engine/Plugins/Experimental/ChaosClothAsset/Source/ChaosClothAssetEngine/Private/ChaosClothAsset/ClothAssetBuilder.h b/Engine/Plugins/Experimental/ChaosClothAsset/Source/ChaosClothAssetEngine/Private/ChaosClothAsset/ClothAssetBuilder.h
--- a/Engine/Plugins/Experimental/ChaosClothAsset/Source/ChaosClothAssetEngine/Private/ChaosClothAsset/ClothAssetBuilder.h
+++ b/Engine/Plugins/Experimental/ChaosClothAsset/Source/ChaosClothAssetEngine/Private/ChaosClothAsset/ClothAssetBuilder.h
@@ -11,8 +11,8 @@ class FSkeletalMeshLODModel;
 /** Builder utility struct, nested struct of UChaosClothAsset. */
 struct UChaosClothAsset::FBuilder
 {
-	/** Build a FSkeletalMeshLODModel out of the cloth asset for the specified LOD index. */
-	static void BuildLod(FSkeletalMeshLODModel& LODModel, const UChaosClothAsset& ClothAsset, int32 LodIndex);
+	/** Build a FSkeletalMeshLODModel out of the cloth asset for the specified LOD index and specified platform. */
+	static void BuildLod(FSkeletalMeshLODModel& LODModel, const UChaosClothAsset& ClothAsset, int32 LodIndex, const ITargetPlatform* TargetPlatform);
 };
 
 #endif  // #if WITH_EDITOR
diff --git a/Engine/Plugins/Experimental/ChaosClothAsset/Source/ChaosClothAssetEngine/Public/ChaosClothAsset/ClothAsset.h b/Engine/Plugins/Experimental/ChaosClothAsset/Source/ChaosClothAssetEngine/Public/ChaosClothAsset/ClothAsset.h
--- a/Engine/Plugins/Experimental/ChaosClothAsset/Source/ChaosClothAssetEngine/Public/ChaosClothAsset/ClothAsset.h
+++ b/Engine/Plugins/Experimental/ChaosClothAsset/Source/ChaosClothAssetEngine/Public/ChaosClothAsset/ClothAsset.h
@@ -91,6 +91,8 @@ public:
 	virtual const FPerPlatformBool& GetDisableBelowMinLodStripping() const override { return DisableBelowMinLodStripping; }
 	virtual const FPerPlatformInt& GetMinLod() const override					{ return MinLod; }
 #if WITH_EDITOR
+	/* Build a LOD model for the targeted platform. */
+	virtual void BuildLODModel(const ITargetPlatform* TargetPlatform, int32 LODIndex) override;
 	virtual FString BuildDerivedDataKey(const ITargetPlatform* TargetPlatform) override;
 	virtual bool IsInitialBuildDone() const override;
 	virtual bool GetEnableLODStreaming(const class ITargetPlatform* TargetPlatform) const override	{ return false; }
