Chaos Cloth Asset: enable debug draw that matches SKM cloth in PIE.

diff --git a/Engine/Plugins/Experimental/ChaosClothAsset/Source/ChaosClothAssetEngine/Private/ChaosClothAsset/ClothSimulationProxy.cpp b/Engine/Plugins/Experimental/ChaosClothAsset/Source/ChaosClothAssetEngine/Private/ChaosClothAsset/ClothSimulationProxy.cpp
--- a/Engine/Plugins/Experimental/ChaosClothAsset/Source/ChaosClothAssetEngine/Private/ChaosClothAsset/ClothSimulationProxy.cpp
+++ b/Engine/Plugins/Experimental/ChaosClothAsset/Source/ChaosClothAssetEngine/Private/ChaosClothAsset/ClothSimulationProxy.cpp
@@ -287,9 +287,47 @@ PRAGMA_ENABLE_DEPRECATION_WARNINGS
 		}
 
 		// Visualization
-		Visualization->DrawPhysMeshWired();
-		Visualization->DrawCollision();
-		Visualization->DrawOpenEdges();
+#if CHAOS_DEBUG_DRAW
+		static const TConsoleVariableData<bool>* const DebugDrawLocalSpaceCVar = IConsoleManager::Get().FindTConsoleVariableDataBool(TEXT("p.ChaosCloth.DebugDrawLocalSpace"));
+		static const TConsoleVariableData<bool>* const DebugDrawBoundsCVar = IConsoleManager::Get().FindTConsoleVariableDataBool(TEXT("p.ChaosCloth.DebugDrawBounds"));
+		static const TConsoleVariableData<bool>* const DebugDrawGravityCVar = IConsoleManager::Get().FindTConsoleVariableDataBool(TEXT("p.ChaosCloth.DebugDrawGravity"));
+		static const TConsoleVariableData<bool>* const DebugDrawPhysMeshWiredCVar = IConsoleManager::Get().FindTConsoleVariableDataBool(TEXT("p.ChaosCloth.DebugDrawPhysMeshWired"));
+		static const TConsoleVariableData<bool>* const DebugDrawAnimMeshWiredCVar = IConsoleManager::Get().FindTConsoleVariableDataBool(TEXT("p.ChaosCloth.DebugDrawAnimMeshWired"));
+		static const TConsoleVariableData<bool>* const DebugDrawPointVelocitiesCVar = IConsoleManager::Get().FindTConsoleVariableDataBool(TEXT("p.ChaosCloth.DebugDrawPointVelocities"));
+		static const TConsoleVariableData<bool>* const DebugDrawAnimNormalsCVar = IConsoleManager::Get().FindTConsoleVariableDataBool(TEXT("p.ChaosCloth.DebugDrawAnimNormals"));
+		static const TConsoleVariableData<bool>* const DebugDrawPointNormalsCVar = IConsoleManager::Get().FindTConsoleVariableDataBool(TEXT("p.ChaosCloth.DebugDrawPointNormals"));
+		static const TConsoleVariableData<bool>* const DebugDrawCollisionCVar = IConsoleManager::Get().FindTConsoleVariableDataBool(TEXT("p.ChaosCloth.DebugDrawCollision"));
+		static const TConsoleVariableData<bool>* const DebugDrawBackstopsCVar = IConsoleManager::Get().FindTConsoleVariableDataBool(TEXT("p.ChaosCloth.DebugDrawBackstops"));
+		static const TConsoleVariableData<bool>* const DebugDrawBackstopDistancesCVar = IConsoleManager::Get().FindTConsoleVariableDataBool(TEXT("p.ChaosCloth.DebugDrawBackstopDistances"));
+		static const TConsoleVariableData<bool>* const DebugDrawMaxDistancesCVar = IConsoleManager::Get().FindTConsoleVariableDataBool(TEXT("p.ChaosCloth.DebugDrawMaxDistances"));
+		static const TConsoleVariableData<bool>* const DebugDrawAnimDriveCVar = IConsoleManager::Get().FindTConsoleVariableDataBool(TEXT("p.ChaosCloth.DebugDrawAnimDrive"));
+		static const TConsoleVariableData<bool>* const DebugDrawEdgeConstraintCVar = IConsoleManager::Get().FindTConsoleVariableDataBool(TEXT("p.ChaosCloth.DebugDrawEdgeConstraint"));
+		static const TConsoleVariableData<bool>* const DebugDrawBendingConstraintCVar = IConsoleManager::Get().FindTConsoleVariableDataBool(TEXT("p.ChaosCloth.DebugDrawBendingConstraint"));
+		static const TConsoleVariableData<bool>* const DebugDrawLongRangeConstraintCVar = IConsoleManager::Get().FindTConsoleVariableDataBool(TEXT("p.ChaosCloth.DebugDrawLongRangeConstraint"));
+		static const TConsoleVariableData<bool>* const DebugDrawWindForcesCVar = IConsoleManager::Get().FindTConsoleVariableDataBool(TEXT("p.ChaosCloth.DebugDrawWindForces"));
+		static const TConsoleVariableData<bool>* const DebugDrawSelfCollisionCVar = IConsoleManager::Get().FindTConsoleVariableDataBool(TEXT("p.ChaosCloth.DebugDrawSelfCollision"));
+		static const TConsoleVariableData<bool>* const DebugDrawSelfIntersectionCVar = IConsoleManager::Get().FindTConsoleVariableDataBool(TEXT("p.ChaosCloth.DebugDrawSelfIntersection"));
+
+		if (DebugDrawLocalSpaceCVar && DebugDrawLocalSpaceCVar->GetValueOnAnyThread()) { Visualization->DrawLocalSpace(); }
+		if (DebugDrawBoundsCVar && DebugDrawBoundsCVar->GetValueOnAnyThread()) { Visualization->DrawBounds(); }
+		if (DebugDrawGravityCVar && DebugDrawGravityCVar->GetValueOnAnyThread()) { Visualization->DrawGravity(); }
+		if (DebugDrawPhysMeshWiredCVar && DebugDrawPhysMeshWiredCVar->GetValueOnAnyThread()) { Visualization->DrawPhysMeshWired(); }
+		if (DebugDrawAnimMeshWiredCVar && DebugDrawAnimMeshWiredCVar->GetValueOnAnyThread()) { Visualization->DrawAnimMeshWired(); }
+		if (DebugDrawPointVelocitiesCVar && DebugDrawPointVelocitiesCVar->GetValueOnAnyThread()) { Visualization->DrawPointVelocities(); }
+		if (DebugDrawAnimNormalsCVar && DebugDrawAnimNormalsCVar->GetValueOnAnyThread()) { Visualization->DrawAnimNormals(); }
+		if (DebugDrawPointNormalsCVar && DebugDrawPointNormalsCVar->GetValueOnAnyThread()) { Visualization->DrawPointNormals(); }
+		if (DebugDrawCollisionCVar && DebugDrawCollisionCVar->GetValueOnAnyThread()) { Visualization->DrawCollision(); }
+		if (DebugDrawBackstopsCVar && DebugDrawBackstopsCVar->GetValueOnAnyThread()) { Visualization->DrawBackstops(); }
+		if (DebugDrawBackstopDistancesCVar && DebugDrawBackstopDistancesCVar->GetValueOnAnyThread()) { Visualization->DrawMaxDistances(); }
+		if (DebugDrawMaxDistancesCVar && DebugDrawMaxDistancesCVar->GetValueOnAnyThread()) { Visualization->DrawMaxDistances(); }
+		if (DebugDrawAnimDriveCVar && DebugDrawAnimDriveCVar->GetValueOnAnyThread()) { Visualization->DrawAnimDrive(); }
+		if (DebugDrawEdgeConstraintCVar && DebugDrawEdgeConstraintCVar->GetValueOnAnyThread()) { Visualization->DrawEdgeConstraint(); }
+		if (DebugDrawBendingConstraintCVar && DebugDrawBendingConstraintCVar->GetValueOnAnyThread()) { Visualization->DrawBendingConstraint(); }
+		if (DebugDrawLongRangeConstraintCVar && DebugDrawLongRangeConstraintCVar->GetValueOnAnyThread()) { Visualization->DrawLongRangeConstraint(); }
+		if (DebugDrawWindForcesCVar && DebugDrawWindForcesCVar->GetValueOnAnyThread()) { Visualization->DrawWindAndPressureForces(); }
+		if (DebugDrawSelfCollisionCVar && DebugDrawSelfCollisionCVar->GetValueOnAnyThread()) { Visualization->DrawSelfCollision(); }
+		if (DebugDrawSelfIntersectionCVar && DebugDrawSelfIntersectionCVar->GetValueOnAnyThread()) { Visualization->DrawSelfIntersection(); }
+#endif  // #if CHAOS_DEBUG_DRAW
 	}
 
 	void FClothSimulationProxy::CompleteParallelSimulation_GameThread()
