Fix material index mapping when turning multiple cloth collections into LODs for one asset

diff --git a/Engine/Plugins/Experimental/ChaosClothAsset/Source/ChaosClothAssetEngine/Private/ChaosClothAsset/ClothAssetBuilder.cpp b/Engine/Plugins/Experimental/ChaosClothAsset/Source/ChaosClothAssetEngine/Private/ChaosClothAsset/ClothAssetBuilder.cpp
--- a/Engine/Plugins/Experimental/ChaosClothAsset/Source/ChaosClothAssetEngine/Private/ChaosClothAsset/ClothAssetBuilder.cpp
+++ b/Engine/Plugins/Experimental/ChaosClothAsset/Source/ChaosClothAssetEngine/Private/ChaosClothAsset/ClothAssetBuilder.cpp
@@ -120,8 +120,14 @@ void UChaosClothAsset::FBuilder::BuildLod(FSkeletalMeshLODModel& LODModel, const
 	// Init the size of the vertex buffer
 	LODModel.NumVertices = 0;
 
-	// Create a table to remap the LOD materials to the asset materials
-	const TArray<FSkeletalMaterial>& Materials = ClothAsset.GetMaterials();
+	// Offset to remap the LOD materials to the asset materials
+	int32 MaterialOffset = 0;
+	for (int32 CollectionIndex = 0; CollectionIndex < LodIndex; ++CollectionIndex)
+	{
+		const TSharedRef<const FManagedArrayCollection> ClothCollection = ClothAsset.GetClothCollections()[CollectionIndex];
+		const FCollectionClothConstFacade ClothFacade(ClothCollection);
+		MaterialOffset += ClothFacade.GetNumRenderPatterns();
+	}
 
 	const TSharedRef<const FManagedArrayCollection> ClothCollection = ClothAsset.GetClothCollections()[LodIndex];
 
@@ -175,7 +181,9 @@ void UChaosClothAsset::FBuilder::BuildLod(FSkeletalMeshLODModel& LODModel, const
 
 		Section.OriginalDataSectionIndex = SectionIndex;
 
-		const int32 MaterialIndex = SectionIndex;
+		const int32 MaterialIndex = MaterialOffset + SectionIndex;
+		check(MaterialIndex < ClothAsset.GetMaterials().Num());
+
 		const int32 NumFaces = RenderPatternFacade.GetNumRenderFaces();
 		const int32 NumIndices = NumFaces * 3;
 
