Chaos Cloth Asset:
- Added a new Cloth Component property to enable the Cloth Asset simulation in editor.
- Added simulation controls to the Cloth Component Details panel.
- Added a new Cloth Sim section filter at the top of the Cloth Component UI.
- Tidied up the Cloth Component categories.

diff --git a/Engine/Plugins/Experimental/ChaosClothAsset/Source/ChaosClothAssetEngine/ChaosClothAssetEngine.Build.cs b/Engine/Plugins/Experimental/ChaosClothAsset/Source/ChaosClothAssetEngine/ChaosClothAssetEngine.Build.cs
--- a/Engine/Plugins/Experimental/ChaosClothAsset/Source/ChaosClothAssetEngine/ChaosClothAssetEngine.Build.cs
+++ b/Engine/Plugins/Experimental/ChaosClothAsset/Source/ChaosClothAssetEngine/ChaosClothAssetEngine.Build.cs
@@ -26,6 +26,11 @@ public class ChaosClothAssetEngine : ModuleRules
 			}
 		);
 
+		if (Target.bBuildEditor || Target.bCompileAgainstEditor)
+		{
+			PrivateDependencyModuleNames.Add("PropertyEditor");  // For adding the Cloth Component "Cloth Sim" section to the Details panel UI
+		}
+
 		PrivateIncludePathModuleNames.AddRange(
 			new string[] {
 				"DerivedDataCache",
diff --git a/Engine/Plugins/Experimental/ChaosClothAsset/Source/ChaosClothAssetEngine/Private/ChaosClothAsset/ChaosClothAssetEngineModule.cpp b/Engine/Plugins/Experimental/ChaosClothAsset/Source/ChaosClothAssetEngine/Private/ChaosClothAsset/ChaosClothAssetEngineModule.cpp
--- a/Engine/Plugins/Experimental/ChaosClothAsset/Source/ChaosClothAssetEngine/Private/ChaosClothAsset/ChaosClothAssetEngineModule.cpp
+++ b/Engine/Plugins/Experimental/ChaosClothAsset/Source/ChaosClothAssetEngine/Private/ChaosClothAsset/ChaosClothAssetEngineModule.cpp
@@ -1,12 +1,18 @@
 // Copyright Epic Games, Inc. All Rights Reserved.
 
 #include "ChaosClothAsset/ClothAssetPrivate.h"
+#include "ChaosClothAsset/ClothComponent.h"
+#include "ChaosClothAsset/ClothComponentCacheAdapter.h"
 #include "Modules/ModuleInterface.h"
 #include "Modules/ModuleManager.h"
-#include "ChaosClothAsset/ClothComponentCacheAdapter.h"
+#if WITH_EDITOR
+#include "PropertyEditorModule.h"
+#endif
 
 DEFINE_LOG_CATEGORY(LogChaosClothAsset);
 
+#define LOCTEXT_NAMESPACE "ClothAssetEngineModule"
+
 namespace UE::Chaos::ClothAsset
 {
 	class FClothAssetEngineModule : public IModuleInterface
@@ -16,13 +22,31 @@ namespace UE::Chaos::ClothAsset
 		{
 			ClothComponentCacheAdapter = MakeUnique<FClothComponentCacheAdapter>();
 			RegisterAdapter(ClothComponentCacheAdapter.Get());
+
+#if WITH_EDITOR
+			// Add the Cloth Component "Cloth Sim" section to the Details panel UI
+			FPropertyEditorModule& PropertyModule = FModuleManager::GetModuleChecked<FPropertyEditorModule>("PropertyEditor");
+			const TSharedPtr<FPropertySection> Section = PropertyModule.FindOrCreateSection(UChaosClothComponent::StaticClass()->GetFName(), TEXT("ClothSim"), LOCTEXT("ClothSim", "Cloth Sim"));
+			Section->AddCategory("Cloth Component");
+			PropertyModule.NotifyCustomizationModuleChanged();
+#endif
 		}
 
 		virtual void ShutdownModule() override
 		{
+#if WITH_EDITOR
+			if (UObjectInitialized())
+			{
+				// Remove sections
+				FPropertyEditorModule& PropertyModule = FModuleManager::GetModuleChecked<FPropertyEditorModule>("PropertyEditor");
+				PropertyModule.RemoveSection(UChaosClothComponent::StaticClass()->GetFName(), TEXT("ClothSim"));
+			}
+#endif
 		}
 	private:
 		TUniquePtr<FClothComponentCacheAdapter> ClothComponentCacheAdapter;
 	};
 }
 IMPLEMENT_MODULE(UE::Chaos::ClothAsset::FClothAssetEngineModule, ChaosClothAssetEngine);
+
+#undef LOCTEXT_NAMESPACE
diff --git a/Engine/Plugins/Experimental/ChaosClothAsset/Source/ChaosClothAssetEngine/Private/ChaosClothAsset/ClothComponent.cpp b/Engine/Plugins/Experimental/ChaosClothAsset/Source/ChaosClothAssetEngine/Private/ChaosClothAsset/ClothComponent.cpp
--- a/Engine/Plugins/Experimental/ChaosClothAsset/Source/ChaosClothAssetEngine/Private/ChaosClothAsset/ClothComponent.cpp
+++ b/Engine/Plugins/Experimental/ChaosClothAsset/Source/ChaosClothAssetEngine/Private/ChaosClothAsset/ClothComponent.cpp
@@ -17,6 +17,9 @@ CSV_DECLARE_CATEGORY_MODULE_EXTERN(ENGINE_API, Animation);
 
 UChaosClothComponent::UChaosClothComponent(const FObjectInitializer& ObjectInitializer)
 	: Super(ObjectInitializer)
+#if WITH_EDITORONLY_DATA
+	, bSimulateInEditor(0)
+#endif
 	, bUseAttachedParentAsPoseComponent(1)  // By default use the parent component as leader pose component
 	, bWaitForParallelTask(0)
 	, bEnableSimulation(1)
@@ -164,12 +167,16 @@ void UChaosClothComponent::PostEditChangeProperty(FPropertyChangedEvent& Propert
 	// Set the skinned asset pointer with the alias pointer (must happen before the call to Super::PostEditChangeProperty)
 	if (const FProperty* const Property = PropertyChangedEvent.Property)
 	{
-		PRAGMA_DISABLE_DEPRECATION_WARNINGS
-			if (Property->GetFName() == GET_MEMBER_NAME_CHECKED(UChaosClothComponent, ClothAsset))
-			{
-				SetClothAsset(ClothAsset);
-			}
-		PRAGMA_ENABLE_DEPRECATION_WARNINGS
+PRAGMA_DISABLE_DEPRECATION_WARNINGS
+		if (Property->GetFName() == GET_MEMBER_NAME_CHECKED(UChaosClothComponent, ClothAsset))
+		{
+			SetClothAsset(ClothAsset);
+		}
+PRAGMA_ENABLE_DEPRECATION_WARNINGS
+		if (Property->GetFName() == GET_MEMBER_NAME_CHECKED(UChaosClothComponent, bSimulateInEditor))
+		{
+			bTickInEditor = bSimulateInEditor;
+		}
 	}
 
 	Super::PostEditChangeProperty(PropertyChangedEvent);
@@ -242,7 +249,7 @@ void UChaosClothComponent::TickComponent(float DeltaTime, enum ELevelTick TickTy
 	}
 
 #if WITH_EDITOR
-	if (TickType == LEVELTICK_ViewportsOnly && bTickOnceInEditor)
+	if (TickType == LEVELTICK_ViewportsOnly && bTickOnceInEditor && !bSimulateInEditor)
 	{
 		// Only tick once in editor when requested. This is used to update from caches by the Chaos Cache Manager.
 		bTickInEditor = false;
diff --git a/Engine/Plugins/Experimental/ChaosClothAsset/Source/ChaosClothAssetEngine/Public/ChaosClothAsset/ClothComponent.h b/Engine/Plugins/Experimental/ChaosClothAsset/Source/ChaosClothAssetEngine/Public/ChaosClothAsset/ClothComponent.h
--- a/Engine/Plugins/Experimental/ChaosClothAsset/Source/ChaosClothAssetEngine/Public/ChaosClothAsset/ClothComponent.h
+++ b/Engine/Plugins/Experimental/ChaosClothAsset/Source/ChaosClothAssetEngine/Public/ChaosClothAsset/ClothComponent.h
@@ -37,23 +37,23 @@ public:
 	~UChaosClothComponent();
 
 	/** Set the cloth asset used by this component. */
-	UFUNCTION(BlueprintCallable, Category = "Components|ClothAsset", Meta = (Keywords = "Chaos Cloth Asset"))
+	UFUNCTION(BlueprintCallable, Category = "ClothComponent", Meta = (Keywords = "Chaos Cloth Asset"))
 	void SetClothAsset(UChaosClothAsset* InClothAsset);
 
 	/** Get the cloth asset used by this component. */
-	UFUNCTION(BlueprintPure, Category = "Components|ClothAsset", Meta = (Keywords = "Chaos Cloth Asset"))
+	UFUNCTION(BlueprintPure, Category = "ClothComponent", Meta = (Keywords = "Chaos Cloth Asset"))
 	UChaosClothAsset* GetClothAsset() const;
 
 	/** Reset the teleport mode. */
-	UFUNCTION(BlueprintCallable, Category = "Components|Teleport", Meta = (Keywords = "Chaos Cloth Teleport"))
+	UFUNCTION(BlueprintCallable, Category = "ClothComponent", Meta = (Keywords = "Chaos Cloth Teleport"))
 	void ResetTeleportMode() { bTeleport = bReset = false; }
 
 	/** Teleport the cloth particles to the new reference bone location keeping pose and velocities prior to advancing the simulation. */
-	UFUNCTION(BlueprintCallable, Category = "Components||Teleport", Meta = (Keywords = "Chaos Cloth Teleport"))
+	UFUNCTION(BlueprintCallable, Category = "ClothComponent", Meta = (Keywords = "Chaos Cloth Teleport"))
 	void ForceNextUpdateTeleport() { bTeleport = true; bReset = false; }
 
 	/** Teleport the cloth particles to the new reference bone location while reseting the pose and velocities prior to advancing the simulation. */
-	UFUNCTION(BlueprintCallable, Category = "Components|Simulation", Meta = (Keywords = "Chaos Cloth Teleport Reset"))
+	UFUNCTION(BlueprintCallable, Category = "ClothComponent", Meta = (Keywords = "Chaos Cloth Teleport Reset"))
 	void ForceNextUpdateTeleportAndReset() { bTeleport = bReset = true; }
 
 	/** Return whether teleport is currently requested. */
@@ -63,31 +63,31 @@ public:
 	bool NeedsReset() const { return bReset; }
 
 	/** Stop the simulation, and keep the cloth in its last pose. */
-	UFUNCTION(BlueprintCallable, Category = "Components|Simulation", Meta = (UnsafeDuringActorConstruction, Keywords = "Chaos Cloth Simulation Suspend"))
+	UFUNCTION(CallInEditor, BlueprintCallable, Category = "ClothComponent", Meta = (UnsafeDuringActorConstruction, Keywords = "Chaos Cloth Simulation Suspend"))
 	void SuspendSimulation() { bSuspendSimulation = true; }
 
 	/** Resume a previously suspended simulation. */
-	UFUNCTION(BlueprintCallable, Category = "Components|Simulation", Meta = (UnsafeDuringActorConstruction, Keywords = "Chaos Cloth Simulation Resume"))
+	UFUNCTION(CallInEditor, BlueprintCallable, Category = "ClothComponent", Meta = (UnsafeDuringActorConstruction, Keywords = "Chaos Cloth Simulation Resume"))
 	void ResumeSimulation() { bSuspendSimulation = false; }
 
 	/** Return whether or not the simulation is currently suspended. */
-	UFUNCTION(BlueprintCallable, Category = "Components|Simulation", Meta = (Keywords = "Chaos Cloth Simulation Suspend"))
+	UFUNCTION(BlueprintCallable, Category = "ClothComponent", Meta = (Keywords = "Chaos Cloth Simulation Suspend"))
 	bool IsSimulationSuspended() const;
 
 	/** Set whether or not to enable simulation. */
-	UFUNCTION(BlueprintCallable, Category = "Components|Simulation", Meta = (UnsafeDuringActorConstruction, Keywords = "Chaos Cloth Simulation Enable"))
+	UFUNCTION(BlueprintCallable, Category = "ClothComponent", Meta = (UnsafeDuringActorConstruction, Keywords = "Chaos Cloth Simulation Enable"))
 	void SetEnableSimulation(bool bEnable) { bEnableSimulation = bEnable; }
 
 	/** Return whether or not the simulation is currently enabled. */
-	UFUNCTION(BlueprintCallable, Category = "Components|Simulation", Meta = (Keywords = "Chaos Cloth Simulation Enable"))
+	UFUNCTION(BlueprintCallable, Category = "ClothComponent", Meta = (Keywords = "Chaos Cloth Simulation Enable"))
 	bool IsSimulationEnabled() const;
 
 	/** Reset all cloth simulation config properties to the values stored in the original cloth asset. */
-	UFUNCTION(BlueprintCallable, Category = "Components|Simulation", Meta = (Keywords = "Chaos Cloth Config Property"))
+	UFUNCTION(BlueprintCallable, Category = "ClothComponent", Meta = (Keywords = "Chaos Cloth Config Property"))
 	void ResetConfigProperties();
 
 	/** Hard reset the cloth simulation by recreating the proxy. */
-	UFUNCTION(BlueprintCallable, Category = "Components|Simulation", Meta = (Keywords = "Chaos Cloth Recreate Simulation Proxy"))
+	UFUNCTION(CallInEditor, BlueprintCallable, Category = "ClothComponent", Meta = (DisplayName = "Hard Reset Simulation", Keywords = "Chaos Cloth Recreate Simulation Proxy"))
 	void RecreateClothSimulationProxy();
 
 	/**
@@ -111,7 +111,7 @@ public:
 	void WaitForExistingParallelClothSimulation_GameThread();
 
 #if WITH_EDITOR
-	/** This will cause the component to tick once in editor. Both flags will be consumed on that tick */
+	/** This will cause the component to tick once in editor. Both flags will be consumed on that tick. Used for the cache adapter. */
 	void SetTickOnceInEditor() { bTickOnceInEditor = true; bTickInEditor = true; }
 #endif
 
@@ -158,8 +158,12 @@ private:
 #if WITH_EDITORONLY_DATA
 	/** Cloth asset used by this component. */
 	UE_DEPRECATED(5.1, "This property isn't deprecated, but getter and setter must be used at all times to preserve correct operations.")
-	UPROPERTY(EditAnywhere, Transient, Setter = SetClothAsset, BlueprintSetter = SetClothAsset, Getter = GetClothAsset, BlueprintGetter = GetClothAsset, Category = ClothAsset)
+	UPROPERTY(EditAnywhere, Transient, Setter = SetClothAsset, BlueprintSetter = SetClothAsset, Getter = GetClothAsset, BlueprintGetter = GetClothAsset, Category = ClothComponent)
 	TObjectPtr<UChaosClothAsset> ClothAsset;
+
+	/** Whether to run the simulation in editor. */
+	UPROPERTY(EditInstanceOnly, Transient, Category = ClothComponent)
+	uint8 bSimulateInEditor : 1;
 #endif
 
 	/** If enabled, and the parent is another Skinned Mesh Component (e.g. another Cloth Component, Poseable Mesh Component, Skeletal Mesh Component, ...etc.), use its pose. */
