Chaos Cloth: calculate RefToLocal matrices for bones that are used for sim mesh skinning that aren't used for render mesh skinning.

diff --git a/Engine/Plugins/Experimental/ChaosClothAsset/Source/ChaosClothAssetEngine/Private/ChaosClothAsset/ClothSimulationContext.cpp b/Engine/Plugins/Experimental/ChaosClothAsset/Source/ChaosClothAssetEngine/Private/ChaosClothAsset/ClothSimulationContext.cpp
--- a/Engine/Plugins/Experimental/ChaosClothAsset/Source/ChaosClothAssetEngine/Private/ChaosClothAsset/ClothSimulationContext.cpp
+++ b/Engine/Plugins/Experimental/ChaosClothAsset/Source/ChaosClothAssetEngine/Private/ChaosClothAsset/ClothSimulationContext.cpp
@@ -3,6 +3,7 @@
 #include "ChaosClothAsset/ClothSimulationContext.h"
 #include "ChaosClothAsset/ClothAsset.h"
 #include "ChaosClothAsset/ClothComponent.h"
+#include "ChaosClothAsset/ClothSimulationModel.h"
 #include "Engine/World.h"
 #include "HAL/IConsoleManager.h"
 #include "SceneInterface.h"
@@ -95,7 +96,15 @@ namespace UE::Chaos::ClothAsset
 		RefToLocalMatrices.Reset(NumBones);
 		if (!bIsInitialization)
 		{
-			ClothComponent.GetCurrentRefToLocalMatrices(RefToLocalMatrices, LodIndex);
+			const TArray<uint16>* RequiredExtraBones = nullptr;
+			if (TSharedPtr<const FChaosClothSimulationModel > ClothModel = ClothAsset ? ClothAsset->GetClothSimulationModel() : TSharedPtr<const FChaosClothSimulationModel >(nullptr))
+			{
+				if (ClothModel->IsValidLodIndex(LodIndex))
+				{
+					RequiredExtraBones = &ClothModel->ClothSimulationLodModels[LodIndex].RequiredExtraBoneIndices;
+				}
+			}
+			ClothComponent.GetCurrentRefToLocalMatrices(RefToLocalMatrices, LodIndex, RequiredExtraBones);
 		}
 		else
 		{
diff --git a/Engine/Plugins/Experimental/ChaosClothAsset/Source/ChaosClothAssetEngine/Private/ChaosClothAsset/ClothSimulationModel.cpp b/Engine/Plugins/Experimental/ChaosClothAsset/Source/ChaosClothAssetEngine/Private/ChaosClothAsset/ClothSimulationModel.cpp
--- a/Engine/Plugins/Experimental/ChaosClothAsset/Source/ChaosClothAssetEngine/Private/ChaosClothAsset/ClothSimulationModel.cpp
+++ b/Engine/Plugins/Experimental/ChaosClothAsset/Source/ChaosClothAssetEngine/Private/ChaosClothAsset/ClothSimulationModel.cpp
@@ -210,7 +210,7 @@ FChaosClothSimulationModel::FChaosClothSimulationModel(const TArray<TSharedRef<c
 			LodModel.FaceIntMaps.Add(FaceIntMapName) = Cloth.GetUserDefinedAttribute<int32>(FaceIntMapName, ClothCollectionGroup::SimFaces);
 		}
 
-		// Copy bone influences and gather tether data
+		// Copy bone influences (and track all used sim bones) and gather tether data.
 		const int32 NumSimVertices3D = Cloth.GetNumSimVertices3D();
 		TConstArrayView<TArray<int32>> SimBoneIndices = Cloth.GetSimBoneIndices();
 		TConstArrayView<TArray<float>> SimBoneWeights = Cloth.GetSimBoneWeights();
@@ -219,6 +219,8 @@ FChaosClothSimulationModel::FChaosClothSimulationModel(const TArray<TSharedRef<c
 		TArray<TArray<TPair<float, int32>>> MergedTetherData;
 		LodModel.BoneData.SetNum(NumSimVertices3D);
 		MergedTetherData.SetNum(NumSimVertices3D);
+		TSet<FBoneIndexType> RequiredExtraSimBones;
+		RequiredExtraSimBones.Reserve(ReferenceSkeleton.GetRawBoneNum());
 
 		for (int32 VertexIndex = 0; VertexIndex < NumSimVertices3D; ++VertexIndex)
 		{
@@ -229,6 +231,7 @@ FChaosClothSimulationModel::FChaosClothSimulationModel(const TArray<TSharedRef<c
 			{
 				LodModel.BoneData[VertexIndex].BoneIndices[BoneIndex] = SimBoneIndices[VertexIndex][BoneIndex];
 				LodModel.BoneData[VertexIndex].BoneWeights[BoneIndex] = SimBoneWeights[VertexIndex][BoneIndex];
+				RequiredExtraSimBones.Add(LodModel.BoneData[VertexIndex].BoneIndices[BoneIndex]);
 			}
 
 			check(TetherKinematicIndex[VertexIndex].Num() == TetherReferenceLength[VertexIndex].Num());
@@ -241,6 +244,19 @@ FChaosClothSimulationModel::FChaosClothSimulationModel(const TArray<TSharedRef<c
 
 		// Batch tethers
 		LodModel.TetherData.GenerateTethers(MoveTemp(MergedTetherData));
+
+		// Find all of the used render bones. We need to store the used sim bones that aren't used for rendering.
+		TSet<FBoneIndexType> UsedRenderBones;
+		TConstArrayView<TArray<int32>> RenderBoneIndices = Cloth.GetRenderBoneIndices();
+		for (const TArray<int32>& Indices : RenderBoneIndices)
+		{
+			for (const int32 Index : Indices)
+			{
+				UsedRenderBones.Add((FBoneIndexType)Index);
+			}
+		}
+
+		LodModel.RequiredExtraBoneIndices = RequiredExtraSimBones.Difference(UsedRenderBones).Array();
 	}
 
 	// Populate used bone names and indices
diff --git a/Engine/Plugins/Experimental/ChaosClothAsset/Source/ChaosClothAssetEngine/Private/ChaosClothAsset/ClothSimulationModel.h b/Engine/Plugins/Experimental/ChaosClothAsset/Source/ChaosClothAssetEngine/Private/ChaosClothAsset/ClothSimulationModel.h
--- a/Engine/Plugins/Experimental/ChaosClothAsset/Source/ChaosClothAssetEngine/Private/ChaosClothAsset/ClothSimulationModel.h
+++ b/Engine/Plugins/Experimental/ChaosClothAsset/Source/ChaosClothAssetEngine/Private/ChaosClothAsset/ClothSimulationModel.h
@@ -39,6 +39,10 @@ struct FChaosClothSimulationLodModel
 	UPROPERTY()
 	TArray<FClothVertBoneData> BoneData;
 
+	/** Already remapped using UsedBoneIndices. These are bones that are needed by this LOD for skinning that aren't needed for render.*/
+	UPROPERTY()
+	TArray<uint16> RequiredExtraBoneIndices;
+
 	/** LOD Transition mesh to mesh skinning weights. */
 	TArray<FMeshToMeshVertData> LODTransitionUpData;
 	TArray<FMeshToMeshVertData> LODTransitionDownData;
diff --git a/Engine/Source/Runtime/Engine/Classes/Components/SkinnedMeshComponent.h b/Engine/Source/Runtime/Engine/Classes/Components/SkinnedMeshComponent.h
--- a/Engine/Source/Runtime/Engine/Classes/Components/SkinnedMeshComponent.h
+++ b/Engine/Source/Runtime/Engine/Classes/Components/SkinnedMeshComponent.h
@@ -1214,7 +1214,7 @@ public:
 	ENGINE_API void CacheRefToLocalMatrices(TArray<FMatrix44f>& OutRefToLocal) const;
 
 	/** Return the skinning matrices used for rendering. */
-	ENGINE_API void GetCurrentRefToLocalMatrices(TArray<FMatrix44f>& OutRefToLocals, int32 InLodIdx) const;
+	ENGINE_API void GetCurrentRefToLocalMatrices(TArray<FMatrix44f>& OutRefToLocals, int32 InLodIdx, const TArray<FBoneIndexType>* ExtraRequiredBoneIndices = nullptr) const;
 
 	FORCEINLINE	const USkinnedMeshComponent* GetBaseComponent()const
 	{
diff --git a/Engine/Source/Runtime/Engine/Private/Components/SkinnedMeshComponent.cpp b/Engine/Source/Runtime/Engine/Private/Components/SkinnedMeshComponent.cpp
--- a/Engine/Source/Runtime/Engine/Private/Components/SkinnedMeshComponent.cpp
+++ b/Engine/Source/Runtime/Engine/Private/Components/SkinnedMeshComponent.cpp
@@ -3464,7 +3464,7 @@ void USkinnedMeshComponent::CacheRefToLocalMatrices(TArray<FMatrix44f>& OutRefTo
 	}
 }
 
-void USkinnedMeshComponent::GetCurrentRefToLocalMatrices(TArray<FMatrix44f>& OutRefToLocals, int32 InLodIdx) const
+void USkinnedMeshComponent::GetCurrentRefToLocalMatrices(TArray<FMatrix44f>& OutRefToLocals, int32 InLodIdx, const TArray<FBoneIndexType>* ExtraRequiredBoneIndices) const
 {
 	if (const USkinnedAsset* const Asset = GetSkinnedAsset())
 	{
@@ -3474,7 +3474,7 @@ void USkinnedMeshComponent::GetCurrentRefToLocalMatrices(TArray<FMatrix44f>& Out
 			TEXT("GetCurrentRefToLocalMatrices (SkelMesh :%s) input LODIndex (%d) doesn't match with render data size (%d)."),
 			*Asset->GetPathName(), InLodIdx, RenderData->LODRenderData.Num()))
 		{
-			UpdateRefToLocalMatrices(OutRefToLocals, this, RenderData, InLodIdx, nullptr);
+			UpdateRefToLocalMatrices(OutRefToLocals, this, RenderData, InLodIdx, ExtraRequiredBoneIndices);
 		}
 		else
 		{
