Chaos Cloth Asset - Added deformer mappings to the cloth collection.

diff --git a/Engine/Plugins/Experimental/ChaosClothAsset/Source/ChaosClothAsset/Private/ChaosClothAsset/ClothCollection.cpp b/Engine/Plugins/Experimental/ChaosClothAsset/Source/ChaosClothAsset/Private/ChaosClothAsset/ClothCollection.cpp
--- a/Engine/Plugins/Experimental/ChaosClothAsset/Source/ChaosClothAsset/Private/ChaosClothAsset/ClothCollection.cpp
+++ b/Engine/Plugins/Experimental/ChaosClothAsset/Source/ChaosClothAsset/Private/ChaosClothAsset/ClothCollection.cpp
@@ -2,156 +2,120 @@
 
 #include "ChaosClothAsset/ClothCollection.h"
 #include "ChaosClothAsset/ClothCollectionGroup.h"
+#include "ChaosClothAsset/ClothCollectionAttribute.h"
 #include "GeometryCollection/ManagedArrayCollection.h"
 
 namespace UE::Chaos::ClothAsset::Private
 {
 	// Lods Group
-	static const FName PhysicsAssetPathNameAttribute(TEXT("PhysicsAssetPathName"));
-	static const FName SkeletalMeshPathNameAttribute(TEXT("SkeletalMeshPathName"));
 	static const TArray<FName> LodsGroupAttributes =
 	{
-		PhysicsAssetPathNameAttribute,
-		SkeletalMeshPathNameAttribute
+		ClothCollectionAttribute::PhysicsAssetPathName,
+		ClothCollectionAttribute::SkeletalMeshPathName
 	};
 
 	// Fabrics Group
-	static const FName FabricBendingStiffnessAttribute(TEXT("FabricBendingStiffness"));
-	static const FName FabricBucklingStiffnessAttribute(TEXT("FabricBucklingStiffness"));
-	static const FName FabricStretchStiffnessAttribute(TEXT("FabricStretchStiffness"));
-	static const FName FabricBucklingRatioAttribute(TEXT("FabricBucklingRatio"));
-	static const FName FabricClothDensityAttribute(TEXT("FabricClothDensity"));
-	static const FName FabricClothFrictionAttribute(TEXT("FabricClothFriction"));
-	static const FName FabricClothThicknessAttribute(TEXT("FabricClothThickness"));
-	static const FName FabricClothDampingAttribute(TEXT("FabricClothDamping"));
 	static const TArray<FName> SimFabricGroupAttributes =
 	{
-		FabricBendingStiffnessAttribute,
-		FabricBucklingStiffnessAttribute,
-		FabricStretchStiffnessAttribute,
-		FabricBucklingRatioAttribute,
-		FabricClothDensityAttribute,
-		FabricClothFrictionAttribute,
-		FabricClothThicknessAttribute,
-		FabricClothDampingAttribute
+		ClothCollectionAttribute::FabricBendingStiffness,
+		ClothCollectionAttribute::FabricBucklingStiffness,
+		ClothCollectionAttribute::FabricStretchStiffness,
+		ClothCollectionAttribute::FabricBucklingRatio,
+		ClothCollectionAttribute::FabricClothDensity,
+		ClothCollectionAttribute::FabricClothFriction,
+		ClothCollectionAttribute::FabricClothThickness,
+		ClothCollectionAttribute::FabricClothDamping
 	};
 
 	// Seam Group
-	static const FName SeamStitchStartAttribute(TEXT("SeamStitchStart"));
-	static const FName SeamStitchEndAttribute(TEXT("SeamStitchEnd"));
 	static const TArray<FName> SeamsGroupAttributes =
 	{
-		SeamStitchStartAttribute,
-		SeamStitchEndAttribute
+		ClothCollectionAttribute::SeamStitchStart,
+		ClothCollectionAttribute::SeamStitchEnd
 	};
 
 	// Seam Stitches Group
-	static const FName SeamStitch2DEndIndicesAttribute(TEXT("SeamStitch2DEndIndices"));
-	static const FName SeamStitch3DIndexAttribute(TEXT("SeamStitch3DIndex"));
 	static const TArray<FName> SeamStitchesGroupAttributes =
 	{
-		SeamStitch2DEndIndicesAttribute,
-		SeamStitch3DIndexAttribute
+		ClothCollectionAttribute::SeamStitch2DEndIndices,
+		ClothCollectionAttribute::SeamStitch3DIndex
 	};
 
 	// Sim Patterns Group
-	static const FName SimVertices2DStartAttribute(TEXT("SimVertices2DStart"));
-	static const FName SimVertices2DEndAttribute(TEXT("SimVertices2DEnd"));
-	static const FName SimFacesStartAttribute(TEXT("SimFacesStart"));
-	static const FName SimFacesEndAttribute(TEXT("SimFacesEnd"));
-	static const FName SimPatternFabricAttribute(TEXT("SimPatternFabric"));
 	static const TArray<FName> SimPatternsGroupAttributes =
 	{ 
-		SimVertices2DStartAttribute,
-		SimVertices2DEndAttribute,
-		SimFacesStartAttribute,
-		SimFacesEndAttribute,
-		SimPatternFabricAttribute
+		ClothCollectionAttribute::SimVertices2DStart,
+		ClothCollectionAttribute::SimVertices2DEnd,
+		ClothCollectionAttribute::SimFacesStart,
+		ClothCollectionAttribute::SimFacesEnd,
+		ClothCollectionAttribute::SimPatternFabric
 	};
 
-	// RenderPatterns Group
-	static const FName RenderVerticesStartAttribute(TEXT("RenderVerticesStart"));
-	static const FName RenderVerticesEndAttribute(TEXT("RenderVerticesEnd"));
-	static const FName RenderFacesStartAttribute(TEXT("RenderFacesStart"));
-	static const FName RenderFacesEndAttribute(TEXT("RenderFacesEnd"));
-	static const FName RenderMaterialPathNameAttribute(TEXT("RenderMaterialPathName"));
+	// Render Patterns Group
 	static const TArray<FName> RenderPatternsGroupAttributes =
 	{
-		RenderVerticesStartAttribute,
-		RenderVerticesEndAttribute,
-		RenderFacesStartAttribute,
-		RenderFacesEndAttribute,
-		RenderMaterialPathNameAttribute
+		ClothCollectionAttribute::RenderVerticesStart,
+		ClothCollectionAttribute::RenderVerticesEnd,
+		ClothCollectionAttribute::RenderFacesStart,
+		ClothCollectionAttribute::RenderFacesEnd,
+		ClothCollectionAttribute::RenderDeformerNumInfluences,
+		ClothCollectionAttribute::RenderMaterialPathName
 	};
 
 	// Sim Faces Group
-	static const FName SimIndices2DAttribute(TEXT("SimIndices2D"));
-	static const FName SimIndices3DAttribute(TEXT("SimIndices3D"));
 	static const TArray<FName> SimFacesGroupAttributes =
 	{
-		SimIndices2DAttribute,
-		SimIndices3DAttribute
+		ClothCollectionAttribute::SimIndices2D,
+		ClothCollectionAttribute::SimIndices3D
 	};
 
 	// Sim Vertices 2D Group
-	static const FName SimPosition2DAttribute(TEXT("SimPosition2D"));
-	static const FName SimVertex3DLookupAttribute(TEXT("SimVertex3DLookup"));
 	static const TArray<FName> SimVertices2DGroupAttributes =
 	{
-		SimPosition2DAttribute,
-		SimVertex3DLookupAttribute,
+		ClothCollectionAttribute::SimPosition2D,
+		ClothCollectionAttribute::SimVertex3DLookup,
 	};
 
 	// Sim Vertices 3D Group
 	// NOTE: if you add anything here, you need to implement how to merge it 
 	// in CollectionClothSeamsFacade. Otherwise, the data in the lowest vertex index 
 	// will survive and the other data will be lost whenever seams are added.
-	static const FName SimPosition3DAttribute(TEXT("SimPosition3D"));
-	static const FName SimNormalAttribute(TEXT("SimNormal"));
-	static const FName SimBoneIndicesAttribute(TEXT("SimBoneIndices"));
-	static const FName SimBoneWeightsAttribute(TEXT("SimBoneWeights"));
-	static const FName TetherKinematicIndexAttribute(TEXT("TetherKinematicIndex"));
-	static const FName TetherReferenceLengthAttribute(TEXT("TetherReferenceLength"));
-	static const FName SimVertex2DLookupAttribute(TEXT("SimVertex2DLookup"));
-	static const FName SeamStitchLookupAttribute(TEXT("SeamStitchLookup"));
 	static const TArray<FName> SimVertices3DGroupAttributes =
 	{
-		SimPosition3DAttribute,
-		SimNormalAttribute,
-		SimBoneIndicesAttribute,
-		SimBoneWeightsAttribute,
-		SimVertex2DLookupAttribute,
-		SeamStitchLookupAttribute
+		ClothCollectionAttribute::SimPosition3D,
+		ClothCollectionAttribute::SimNormal,
+		ClothCollectionAttribute::SimBoneIndices,
+		ClothCollectionAttribute::SimBoneWeights,
+		ClothCollectionAttribute::SimVertex2DLookup,
+		ClothCollectionAttribute::SeamStitchLookup
 	};
 
 	// Render Faces Group
-	static const FName RenderIndicesAttribute(TEXT("RenderIndices"));
 	static const TArray<FName> RenderFacesGroupAttributes =
 	{
-		RenderIndicesAttribute,
+		ClothCollectionAttribute::RenderIndices,
 	};
 
 	// Render Vertices Group
-	static const FName RenderPositionAttribute(TEXT("RenderPosition"));
-	static const FName RenderNormalAttribute(TEXT("RenderNormal"));
-	static const FName RenderTangentUAttribute(TEXT("RenderTangentU"));
-	static const FName RenderTangentVAttribute(TEXT("RenderTangentV"));
-	static const FName RenderUVsAttribute(TEXT("RenderUVs"));
-	static const FName RenderColorAttribute(TEXT("RenderColor"));
-	static const FName RenderBoneIndicesAttribute(TEXT("RenderBoneIndices"));
-	static const FName RenderBoneWeightsAttribute(TEXT("RenderBoneWeights"));
 	static const TArray<FName> RenderVerticesGroupAttributes =
 	{
-		RenderPositionAttribute,
-		RenderNormalAttribute,
-		RenderTangentUAttribute,
-		RenderTangentVAttribute,
-		RenderUVsAttribute,
-		RenderColorAttribute,
-		RenderBoneIndicesAttribute,
-		RenderBoneWeightsAttribute
+		ClothCollectionAttribute::RenderPosition,
+		ClothCollectionAttribute::RenderNormal,
+		ClothCollectionAttribute::RenderTangentU,
+		ClothCollectionAttribute::RenderTangentV,
+		ClothCollectionAttribute::RenderUVs,
+		ClothCollectionAttribute::RenderColor,
+		ClothCollectionAttribute::RenderBoneIndices,
+		ClothCollectionAttribute::RenderBoneWeights,
+		ClothCollectionAttribute::RenderDeformerPositionBaryCoordsAndDist,
+		ClothCollectionAttribute::RenderDeformerNormalBaryCoordsAndDist,
+		ClothCollectionAttribute::RenderDeformerTangentBaryCoordsAndDist,
+		ClothCollectionAttribute::RenderDeformerSimIndices3D,
+		ClothCollectionAttribute::RenderDeformerWeight,
+		ClothCollectionAttribute::RenderDeformerSkinningBlend
 	};
 
+	// All fixed attributes for this collection
 	static const TMap<FName, TArray<FName>> FixedAttributeNamesMap =
 	{
 		{ ClothCollectionGroup::Lods, LodsGroupAttributes },
@@ -176,76 +140,83 @@ namespace UE::Chaos::ClothAsset
 		using namespace UE::Chaos::ClothAsset::Private;
 
 		// LODs Group
-		PhysicsAssetPathName = ManagedArrayCollection->FindAttribute<FString>(PhysicsAssetPathNameAttribute, ClothCollectionGroup::Lods);
-		SkeletalMeshPathName = ManagedArrayCollection->FindAttribute<FString>(SkeletalMeshPathNameAttribute, ClothCollectionGroup::Lods);
+		PhysicsAssetPathName = ManagedArrayCollection->FindAttribute<FString>(ClothCollectionAttribute::PhysicsAssetPathName, ClothCollectionGroup::Lods);
+		SkeletalMeshPathName = ManagedArrayCollection->FindAttribute<FString>(ClothCollectionAttribute::SkeletalMeshPathName, ClothCollectionGroup::Lods);
 
 		// Seam Group
-		SeamStitchStart = ManagedArrayCollection->FindAttribute<int32>(SeamStitchStartAttribute, ClothCollectionGroup::Seams);
-		SeamStitchEnd = ManagedArrayCollection->FindAttribute<int32>(SeamStitchEndAttribute, ClothCollectionGroup::Seams);
+		SeamStitchStart = ManagedArrayCollection->FindAttribute<int32>(ClothCollectionAttribute::SeamStitchStart, ClothCollectionGroup::Seams);
+		SeamStitchEnd = ManagedArrayCollection->FindAttribute<int32>(ClothCollectionAttribute::SeamStitchEnd, ClothCollectionGroup::Seams);
 
 		// Seam Stitches Group
-		SeamStitch2DEndIndices = ManagedArrayCollection->FindAttribute<FIntVector2>(SeamStitch2DEndIndicesAttribute, ClothCollectionGroup::SeamStitches);
-		SeamStitch3DIndex = ManagedArrayCollection->FindAttribute<int32>(SeamStitch3DIndexAttribute, ClothCollectionGroup::SeamStitches);
+		SeamStitch2DEndIndices = ManagedArrayCollection->FindAttribute<FIntVector2>(ClothCollectionAttribute::SeamStitch2DEndIndices, ClothCollectionGroup::SeamStitches);
+		SeamStitch3DIndex = ManagedArrayCollection->FindAttribute<int32>(ClothCollectionAttribute::SeamStitch3DIndex, ClothCollectionGroup::SeamStitches);
 
 		// Sim Patterns Group
-		SimVertices2DStart = ManagedArrayCollection->FindAttribute<int32>(SimVertices2DStartAttribute, ClothCollectionGroup::SimPatterns);
-		SimVertices2DEnd = ManagedArrayCollection->FindAttribute<int32>(SimVertices2DEndAttribute, ClothCollectionGroup::SimPatterns);
-		SimFacesStart = ManagedArrayCollection->FindAttribute<int32>(SimFacesStartAttribute, ClothCollectionGroup::SimPatterns);
-		SimFacesEnd = ManagedArrayCollection->FindAttribute<int32>(SimFacesEndAttribute, ClothCollectionGroup::SimPatterns);
-		SimPatternFabric = ManagedArrayCollection->FindAttribute<int32>(SimPatternFabricAttribute, ClothCollectionGroup::SimPatterns);
+		SimVertices2DStart = ManagedArrayCollection->FindAttribute<int32>(ClothCollectionAttribute::SimVertices2DStart, ClothCollectionGroup::SimPatterns);
+		SimVertices2DEnd = ManagedArrayCollection->FindAttribute<int32>(ClothCollectionAttribute::SimVertices2DEnd, ClothCollectionGroup::SimPatterns);
+		SimFacesStart = ManagedArrayCollection->FindAttribute<int32>(ClothCollectionAttribute::SimFacesStart, ClothCollectionGroup::SimPatterns);
+		SimFacesEnd = ManagedArrayCollection->FindAttribute<int32>(ClothCollectionAttribute::SimFacesEnd, ClothCollectionGroup::SimPatterns);
+		SimPatternFabric = ManagedArrayCollection->FindAttribute<int32>(ClothCollectionAttribute::SimPatternFabric, ClothCollectionGroup::SimPatterns);
 
 		// Render Patterns Group
-		RenderVerticesStart = ManagedArrayCollection->FindAttribute<int32>(RenderVerticesStartAttribute, ClothCollectionGroup::RenderPatterns);
-		RenderVerticesEnd = ManagedArrayCollection->FindAttribute<int32>(RenderVerticesEndAttribute, ClothCollectionGroup::RenderPatterns);
-		RenderFacesStart = ManagedArrayCollection->FindAttribute<int32>(RenderFacesStartAttribute, ClothCollectionGroup::RenderPatterns);
-		RenderFacesEnd = ManagedArrayCollection->FindAttribute<int32>(RenderFacesEndAttribute, ClothCollectionGroup::RenderPatterns);
-		RenderMaterialPathName = ManagedArrayCollection->FindAttribute<FString>(RenderMaterialPathNameAttribute, ClothCollectionGroup::RenderPatterns);
+		RenderVerticesStart = ManagedArrayCollection->FindAttribute<int32>(ClothCollectionAttribute::RenderVerticesStart, ClothCollectionGroup::RenderPatterns);
+		RenderVerticesEnd = ManagedArrayCollection->FindAttribute<int32>(ClothCollectionAttribute::RenderVerticesEnd, ClothCollectionGroup::RenderPatterns);
+		RenderFacesStart = ManagedArrayCollection->FindAttribute<int32>(ClothCollectionAttribute::RenderFacesStart, ClothCollectionGroup::RenderPatterns);
+		RenderFacesEnd = ManagedArrayCollection->FindAttribute<int32>(ClothCollectionAttribute::RenderFacesEnd, ClothCollectionGroup::RenderPatterns);
+		RenderDeformerNumInfluences = ManagedArrayCollection->FindAttribute<int32>(ClothCollectionAttribute::RenderDeformerNumInfluences, ClothCollectionGroup::RenderPatterns);
+		RenderMaterialPathName = ManagedArrayCollection->FindAttribute<FString>(ClothCollectionAttribute::RenderMaterialPathName, ClothCollectionGroup::RenderPatterns);
 
 		//~ Fabric Group
-		FabricBendingStiffness = ManagedArrayCollection->FindAttribute<FVector3f>(FabricBendingStiffnessAttribute, ClothCollectionGroup::Fabrics);
-		FabricBucklingStiffness = ManagedArrayCollection->FindAttribute<FVector3f>(FabricBucklingStiffnessAttribute, ClothCollectionGroup::Fabrics);
-		FabricStretchStiffness = ManagedArrayCollection->FindAttribute<FVector3f>(FabricStretchStiffnessAttribute, ClothCollectionGroup::Fabrics);
-		FabricBucklingRatio = ManagedArrayCollection->FindAttribute<float>(FabricBucklingRatioAttribute, ClothCollectionGroup::Fabrics);
-		FabricClothDensity = ManagedArrayCollection->FindAttribute<float>(FabricClothDensityAttribute, ClothCollectionGroup::Fabrics);
-		FabricClothFriction = ManagedArrayCollection->FindAttribute<float>(FabricClothFrictionAttribute, ClothCollectionGroup::Fabrics);
-		FabricClothThickness = ManagedArrayCollection->FindAttribute<float>(FabricClothThicknessAttribute, ClothCollectionGroup::Fabrics);
-		FabricClothDamping = ManagedArrayCollection->FindAttribute<float>(FabricClothDampingAttribute, ClothCollectionGroup::Fabrics);
+		FabricBendingStiffness = ManagedArrayCollection->FindAttribute<FVector3f>(ClothCollectionAttribute::FabricBendingStiffness, ClothCollectionGroup::Fabrics);
+		FabricBucklingStiffness = ManagedArrayCollection->FindAttribute<FVector3f>(ClothCollectionAttribute::FabricBucklingStiffness, ClothCollectionGroup::Fabrics);
+		FabricStretchStiffness = ManagedArrayCollection->FindAttribute<FVector3f>(ClothCollectionAttribute::FabricStretchStiffness, ClothCollectionGroup::Fabrics);
+		FabricBucklingRatio = ManagedArrayCollection->FindAttribute<float>(ClothCollectionAttribute::FabricBucklingRatio, ClothCollectionGroup::Fabrics);
+		FabricClothDensity = ManagedArrayCollection->FindAttribute<float>(ClothCollectionAttribute::FabricClothDensity, ClothCollectionGroup::Fabrics);
+		FabricClothFriction = ManagedArrayCollection->FindAttribute<float>(ClothCollectionAttribute::FabricClothFriction, ClothCollectionGroup::Fabrics);
+		FabricClothThickness = ManagedArrayCollection->FindAttribute<float>(ClothCollectionAttribute::FabricClothThickness, ClothCollectionGroup::Fabrics);
+		FabricClothDamping = ManagedArrayCollection->FindAttribute<float>(ClothCollectionAttribute::FabricClothDamping, ClothCollectionGroup::Fabrics);
 		
 		// Sim Faces Group
-		SimIndices2D = ManagedArrayCollection->FindAttribute<FIntVector3>(SimIndices2DAttribute, ClothCollectionGroup::SimFaces);
-		SimIndices3D = ManagedArrayCollection->FindAttribute<FIntVector3>(SimIndices3DAttribute, ClothCollectionGroup::SimFaces);
+		SimIndices2D = ManagedArrayCollection->FindAttribute<FIntVector3>(ClothCollectionAttribute::SimIndices2D, ClothCollectionGroup::SimFaces);
+		SimIndices3D = ManagedArrayCollection->FindAttribute<FIntVector3>(ClothCollectionAttribute::SimIndices3D, ClothCollectionGroup::SimFaces);
 
 		// Sim Vertices 2D Group
-		SimPosition2D = ManagedArrayCollection->FindAttribute<FVector2f>(SimPosition2DAttribute, ClothCollectionGroup::SimVertices2D);
-		SimVertex3DLookup = ManagedArrayCollection->FindAttribute<int32>(SimVertex3DLookupAttribute, ClothCollectionGroup::SimVertices2D);
+		SimPosition2D = ManagedArrayCollection->FindAttribute<FVector2f>(ClothCollectionAttribute::SimPosition2D, ClothCollectionGroup::SimVertices2D);
+		SimVertex3DLookup = ManagedArrayCollection->FindAttribute<int32>(ClothCollectionAttribute::SimVertex3DLookup, ClothCollectionGroup::SimVertices2D);
 
 		// Sim Vertices 3D Group
-		SimPosition3D = ManagedArrayCollection->FindAttribute<FVector3f>(SimPosition3DAttribute, ClothCollectionGroup::SimVertices3D);
-		SimNormal = ManagedArrayCollection->FindAttribute<FVector3f>(SimNormalAttribute, ClothCollectionGroup::SimVertices3D);
-		SimBoneIndices = ManagedArrayCollection->FindAttribute<TArray<int32>>(SimBoneIndicesAttribute, ClothCollectionGroup::SimVertices3D);
-		SimBoneWeights = ManagedArrayCollection->FindAttribute<TArray<float>>(SimBoneWeightsAttribute, ClothCollectionGroup::SimVertices3D);
-		TetherKinematicIndex = ManagedArrayCollection->FindAttribute<TArray<int32>>(TetherKinematicIndexAttribute, ClothCollectionGroup::SimVertices3D);
-		TetherReferenceLength = ManagedArrayCollection->FindAttribute<TArray<float>>(TetherReferenceLengthAttribute, ClothCollectionGroup::SimVertices3D);
-		SimVertex2DLookup = ManagedArrayCollection->FindAttribute<TArray<int32>>(SimVertex2DLookupAttribute, ClothCollectionGroup::SimVertices3D);
-		SeamStitchLookup = ManagedArrayCollection->FindAttribute<TArray<int32>>(SeamStitchLookupAttribute, ClothCollectionGroup::SimVertices3D);
+		SimPosition3D = ManagedArrayCollection->FindAttribute<FVector3f>(ClothCollectionAttribute::SimPosition3D, ClothCollectionGroup::SimVertices3D);
+		SimNormal = ManagedArrayCollection->FindAttribute<FVector3f>(ClothCollectionAttribute::SimNormal, ClothCollectionGroup::SimVertices3D);
+		SimBoneIndices = ManagedArrayCollection->FindAttribute<TArray<int32>>(ClothCollectionAttribute::SimBoneIndices, ClothCollectionGroup::SimVertices3D);
+		SimBoneWeights = ManagedArrayCollection->FindAttribute<TArray<float>>(ClothCollectionAttribute::SimBoneWeights, ClothCollectionGroup::SimVertices3D);
+		TetherKinematicIndex = ManagedArrayCollection->FindAttribute<TArray<int32>>(ClothCollectionAttribute::TetherKinematicIndex, ClothCollectionGroup::SimVertices3D);
+		TetherReferenceLength = ManagedArrayCollection->FindAttribute<TArray<float>>(ClothCollectionAttribute::TetherReferenceLength, ClothCollectionGroup::SimVertices3D);
+		SimVertex2DLookup = ManagedArrayCollection->FindAttribute<TArray<int32>>(ClothCollectionAttribute::SimVertex2DLookup, ClothCollectionGroup::SimVertices3D);
+		SeamStitchLookup = ManagedArrayCollection->FindAttribute<TArray<int32>>(ClothCollectionAttribute::SeamStitchLookup, ClothCollectionGroup::SimVertices3D);
 
 		// Render Faces Group
-		RenderIndices = ManagedArrayCollection->FindAttribute<FIntVector3>(RenderIndicesAttribute, ClothCollectionGroup::RenderFaces);
+		RenderIndices = ManagedArrayCollection->FindAttribute<FIntVector3>(ClothCollectionAttribute::RenderIndices, ClothCollectionGroup::RenderFaces);
 
 		// Render Vertices Group
-		RenderPosition = ManagedArrayCollection->FindAttribute<FVector3f>(RenderPositionAttribute, ClothCollectionGroup::RenderVertices);
-		RenderNormal = ManagedArrayCollection->FindAttribute<FVector3f>(RenderNormalAttribute, ClothCollectionGroup::RenderVertices);
-		RenderTangentU = ManagedArrayCollection->FindAttribute<FVector3f>(RenderTangentUAttribute, ClothCollectionGroup::RenderVertices);
-		RenderTangentV = ManagedArrayCollection->FindAttribute<FVector3f>(RenderTangentVAttribute, ClothCollectionGroup::RenderVertices);
-		RenderUVs = ManagedArrayCollection->FindAttribute<TArray<FVector2f>>(RenderUVsAttribute, ClothCollectionGroup::RenderVertices);
-		RenderColor = ManagedArrayCollection->FindAttribute<FLinearColor>(RenderColorAttribute, ClothCollectionGroup::RenderVertices);
-		RenderBoneIndices = ManagedArrayCollection->FindAttribute<TArray<int32>>(RenderBoneIndicesAttribute, ClothCollectionGroup::RenderVertices);
-		RenderBoneWeights = ManagedArrayCollection->FindAttribute<TArray<float>>(RenderBoneWeightsAttribute, ClothCollectionGroup::RenderVertices);
+		RenderPosition = ManagedArrayCollection->FindAttribute<FVector3f>(ClothCollectionAttribute::RenderPosition, ClothCollectionGroup::RenderVertices);
+		RenderNormal = ManagedArrayCollection->FindAttribute<FVector3f>(ClothCollectionAttribute::RenderNormal, ClothCollectionGroup::RenderVertices);
+		RenderTangentU = ManagedArrayCollection->FindAttribute<FVector3f>(ClothCollectionAttribute::RenderTangentU, ClothCollectionGroup::RenderVertices);
+		RenderTangentV = ManagedArrayCollection->FindAttribute<FVector3f>(ClothCollectionAttribute::RenderTangentV, ClothCollectionGroup::RenderVertices);
+		RenderUVs = ManagedArrayCollection->FindAttribute<TArray<FVector2f>>(ClothCollectionAttribute::RenderUVs, ClothCollectionGroup::RenderVertices);
+		RenderColor = ManagedArrayCollection->FindAttribute<FLinearColor>(ClothCollectionAttribute::RenderColor, ClothCollectionGroup::RenderVertices);
+		RenderBoneIndices = ManagedArrayCollection->FindAttribute<TArray<int32>>(ClothCollectionAttribute::RenderBoneIndices, ClothCollectionGroup::RenderVertices);
+		RenderBoneWeights = ManagedArrayCollection->FindAttribute<TArray<float>>(ClothCollectionAttribute::RenderBoneWeights, ClothCollectionGroup::RenderVertices);
+		RenderDeformerPositionBaryCoordsAndDist = ManagedArrayCollection->FindAttribute<TArray<FVector4f>>(ClothCollectionAttribute::RenderDeformerPositionBaryCoordsAndDist, ClothCollectionGroup::RenderVertices);
+		RenderDeformerNormalBaryCoordsAndDist = ManagedArrayCollection->FindAttribute<TArray<FVector4f>>(ClothCollectionAttribute::RenderDeformerNormalBaryCoordsAndDist, ClothCollectionGroup::RenderVertices);
+		RenderDeformerTangentBaryCoordsAndDist = ManagedArrayCollection->FindAttribute<TArray<FVector4f>>(ClothCollectionAttribute::RenderDeformerTangentBaryCoordsAndDist, ClothCollectionGroup::RenderVertices);
+		RenderDeformerSimIndices3D = ManagedArrayCollection->FindAttribute<TArray<FIntVector3>>(ClothCollectionAttribute::RenderDeformerSimIndices3D, ClothCollectionGroup::RenderVertices);
+		RenderDeformerWeight = ManagedArrayCollection->FindAttribute<TArray<float>>(ClothCollectionAttribute::RenderDeformerWeight, ClothCollectionGroup::RenderVertices);
+		RenderDeformerSkinningBlend = ManagedArrayCollection->FindAttribute<float>(ClothCollectionAttribute::RenderDeformerSkinningBlend, ClothCollectionGroup::RenderVertices);
 	}
 
-	bool FClothCollection::IsValid() const
+	bool FClothCollection::IsValid(EClothCollectionOptionalSchemas OptionalSchemas) const
 	{
-		return 
+		return
 			// LODs Group
 			PhysicsAssetPathName &&
 			SkeletalMeshPathName &&
@@ -271,6 +242,8 @@ namespace UE::Chaos::ClothAsset
 			RenderFacesStart &&
 			RenderFacesEnd &&
 			RenderMaterialPathName &&
+			(!EnumHasAnyFlags(OptionalSchemas, EClothCollectionOptionalSchemas::RenderDeformer) ||
+				RenderDeformerNumInfluences) &&
 
 			// Sim Faces Group
 			SimIndices2D &&
@@ -287,7 +260,7 @@ namespace UE::Chaos::ClothAsset
 			FabricClothDamping &&
 
 			// Sim Vertices 2D Group
-			SimPosition2D  &&
+			SimPosition2D &&
 			SimVertex3DLookup &&
 
 			// Sim Vertices 3D Group
@@ -311,10 +284,17 @@ namespace UE::Chaos::ClothAsset
 			RenderUVs &&
 			RenderColor &&
 			RenderBoneIndices &&
-			RenderBoneWeights;
+			RenderBoneWeights &&
+			(!EnumHasAnyFlags(OptionalSchemas, EClothCollectionOptionalSchemas::RenderDeformer) ||
+				(RenderDeformerPositionBaryCoordsAndDist &&
+				RenderDeformerNormalBaryCoordsAndDist &&
+				RenderDeformerTangentBaryCoordsAndDist &&
+				RenderDeformerSimIndices3D &&
+				RenderDeformerWeight &&
+				RenderDeformerSkinningBlend));
 	}
 
-	void FClothCollection::DefineSchema()
+	void FClothCollection::DefineSchema(EClothCollectionOptionalSchemas OptionalSchemas)
 	{
 		using namespace UE::Chaos::ClothAsset::Private;
 
@@ -330,71 +310,85 @@ namespace UE::Chaos::ClothAsset
 		FManagedArrayCollection::FConstructionParameters SimVertices3DDependency(ClothCollectionGroup::SimVertices3D, bSaved, bAllowCircularDependency);  // Any attribute with this dependency must handle welding and splitting in FCollectionClothSeamFacade
 
 		// LODs Group
-		PhysicsAssetPathName = &ManagedArrayCollection->AddAttribute<FString>(PhysicsAssetPathNameAttribute, ClothCollectionGroup::Lods);
-		SkeletalMeshPathName = &ManagedArrayCollection->AddAttribute<FString>(SkeletalMeshPathNameAttribute, ClothCollectionGroup::Lods);
+		PhysicsAssetPathName = &ManagedArrayCollection->AddAttribute<FString>(ClothCollectionAttribute::PhysicsAssetPathName, ClothCollectionGroup::Lods);
+		SkeletalMeshPathName = &ManagedArrayCollection->AddAttribute<FString>(ClothCollectionAttribute::SkeletalMeshPathName, ClothCollectionGroup::Lods);
 
 		// Seams Group
-		SeamStitchStart = &ManagedArrayCollection->AddAttribute<int32>(SeamStitchStartAttribute, ClothCollectionGroup::Seams, SeamStitchesDependency);
-		SeamStitchEnd = &ManagedArrayCollection->AddAttribute<int32>(SeamStitchEndAttribute, ClothCollectionGroup::Seams, SeamStitchesDependency);
+		SeamStitchStart = &ManagedArrayCollection->AddAttribute<int32>(ClothCollectionAttribute::SeamStitchStart, ClothCollectionGroup::Seams, SeamStitchesDependency);
+		SeamStitchEnd = &ManagedArrayCollection->AddAttribute<int32>(ClothCollectionAttribute::SeamStitchEnd, ClothCollectionGroup::Seams, SeamStitchesDependency);
 
 		// Seam Stitches Group
-		SeamStitch2DEndIndices = &ManagedArrayCollection->AddAttribute<FIntVector2>(SeamStitch2DEndIndicesAttribute, ClothCollectionGroup::SeamStitches, SimVertices2DDependency);
-		SeamStitch3DIndex = &ManagedArrayCollection->AddAttribute<int32>(SeamStitch3DIndexAttribute, ClothCollectionGroup::SeamStitches, SimVertices3DDependency);
+		SeamStitch2DEndIndices = &ManagedArrayCollection->AddAttribute<FIntVector2>(ClothCollectionAttribute::SeamStitch2DEndIndices, ClothCollectionGroup::SeamStitches, SimVertices2DDependency);
+		SeamStitch3DIndex = &ManagedArrayCollection->AddAttribute<int32>(ClothCollectionAttribute::SeamStitch3DIndex, ClothCollectionGroup::SeamStitches, SimVertices3DDependency);
 
 		// Sim Patterns Group
-		SimVertices2DStart = &ManagedArrayCollection->AddAttribute<int32>(SimVertices2DStartAttribute, ClothCollectionGroup::SimPatterns, SimVertices2DDependency);
-		SimVertices2DEnd = &ManagedArrayCollection->AddAttribute<int32>(SimVertices2DEndAttribute, ClothCollectionGroup::SimPatterns, SimVertices2DDependency);
-		SimFacesStart = &ManagedArrayCollection->AddAttribute<int32>(SimFacesStartAttribute, ClothCollectionGroup::SimPatterns, SimFacesDependency);
-		SimFacesEnd = &ManagedArrayCollection->AddAttribute<int32>(SimFacesEndAttribute, ClothCollectionGroup::SimPatterns, SimFacesDependency);
-		SimPatternFabric = &ManagedArrayCollection->AddAttribute<int32>(SimPatternFabricAttribute, ClothCollectionGroup::SimPatterns, SimFabricsDependency);
+		SimVertices2DStart = &ManagedArrayCollection->AddAttribute<int32>(ClothCollectionAttribute::SimVertices2DStart, ClothCollectionGroup::SimPatterns, SimVertices2DDependency);
+		SimVertices2DEnd = &ManagedArrayCollection->AddAttribute<int32>(ClothCollectionAttribute::SimVertices2DEnd, ClothCollectionGroup::SimPatterns, SimVertices2DDependency);
+		SimFacesStart = &ManagedArrayCollection->AddAttribute<int32>(ClothCollectionAttribute::SimFacesStart, ClothCollectionGroup::SimPatterns, SimFacesDependency);
+		SimFacesEnd = &ManagedArrayCollection->AddAttribute<int32>(ClothCollectionAttribute::SimFacesEnd, ClothCollectionGroup::SimPatterns, SimFacesDependency);
+		SimPatternFabric = &ManagedArrayCollection->AddAttribute<int32>(ClothCollectionAttribute::SimPatternFabric, ClothCollectionGroup::SimPatterns, SimFabricsDependency);
 
 		// Render Patterns Group
-		RenderVerticesStart = &ManagedArrayCollection->AddAttribute<int32>(RenderVerticesStartAttribute, ClothCollectionGroup::RenderPatterns, RenderVerticesDependency);
-		RenderVerticesEnd = &ManagedArrayCollection->AddAttribute<int32>(RenderVerticesEndAttribute, ClothCollectionGroup::RenderPatterns, RenderVerticesDependency);
-		RenderFacesStart = &ManagedArrayCollection->AddAttribute<int32>(RenderFacesStartAttribute, ClothCollectionGroup::RenderPatterns, RenderFacesDependency);
-		RenderFacesEnd = &ManagedArrayCollection->AddAttribute<int32>(RenderFacesEndAttribute, ClothCollectionGroup::RenderPatterns, RenderFacesDependency);
-		RenderMaterialPathName = &ManagedArrayCollection->AddAttribute<FString>(RenderMaterialPathNameAttribute, ClothCollectionGroup::RenderPatterns);
+		RenderVerticesStart = &ManagedArrayCollection->AddAttribute<int32>(ClothCollectionAttribute::RenderVerticesStart, ClothCollectionGroup::RenderPatterns, RenderVerticesDependency);
+		RenderVerticesEnd = &ManagedArrayCollection->AddAttribute<int32>(ClothCollectionAttribute::RenderVerticesEnd, ClothCollectionGroup::RenderPatterns, RenderVerticesDependency);
+		RenderFacesStart = &ManagedArrayCollection->AddAttribute<int32>(ClothCollectionAttribute::RenderFacesStart, ClothCollectionGroup::RenderPatterns, RenderFacesDependency);
+		RenderFacesEnd = &ManagedArrayCollection->AddAttribute<int32>(ClothCollectionAttribute::RenderFacesEnd, ClothCollectionGroup::RenderPatterns, RenderFacesDependency);
+		RenderMaterialPathName = &ManagedArrayCollection->AddAttribute<FString>(ClothCollectionAttribute::RenderMaterialPathName, ClothCollectionGroup::RenderPatterns);
 
 		//~ Fabric Group
-		FabricBendingStiffness = &ManagedArrayCollection->AddAttribute<FVector3f>(FabricBendingStiffnessAttribute, ClothCollectionGroup::Fabrics);
-		FabricBucklingStiffness = &ManagedArrayCollection->AddAttribute<FVector3f>(FabricBucklingStiffnessAttribute, ClothCollectionGroup::Fabrics);
-		FabricStretchStiffness = &ManagedArrayCollection->AddAttribute<FVector3f>(FabricStretchStiffnessAttribute, ClothCollectionGroup::Fabrics);
-		FabricBucklingRatio = &ManagedArrayCollection->AddAttribute<float>(FabricBucklingRatioAttribute, ClothCollectionGroup::Fabrics);
-		FabricClothDensity = &ManagedArrayCollection->AddAttribute<float>(FabricClothDensityAttribute, ClothCollectionGroup::Fabrics);
-		FabricClothFriction = &ManagedArrayCollection->AddAttribute<float>(FabricClothFrictionAttribute, ClothCollectionGroup::Fabrics);
-		FabricClothThickness = &ManagedArrayCollection->AddAttribute<float>(FabricClothThicknessAttribute, ClothCollectionGroup::Fabrics);
-		FabricClothDamping = &ManagedArrayCollection->AddAttribute<float>(FabricClothDampingAttribute, ClothCollectionGroup::Fabrics);
+		FabricBendingStiffness = &ManagedArrayCollection->AddAttribute<FVector3f>(ClothCollectionAttribute::FabricBendingStiffness, ClothCollectionGroup::Fabrics);
+		FabricBucklingStiffness = &ManagedArrayCollection->AddAttribute<FVector3f>(ClothCollectionAttribute::FabricBucklingStiffness, ClothCollectionGroup::Fabrics);
+		FabricStretchStiffness = &ManagedArrayCollection->AddAttribute<FVector3f>(ClothCollectionAttribute::FabricStretchStiffness, ClothCollectionGroup::Fabrics);
+		FabricBucklingRatio = &ManagedArrayCollection->AddAttribute<float>(ClothCollectionAttribute::FabricBucklingRatio, ClothCollectionGroup::Fabrics);
+		FabricClothDensity = &ManagedArrayCollection->AddAttribute<float>(ClothCollectionAttribute::FabricClothDensity, ClothCollectionGroup::Fabrics);
+		FabricClothFriction = &ManagedArrayCollection->AddAttribute<float>(ClothCollectionAttribute::FabricClothFriction, ClothCollectionGroup::Fabrics);
+		FabricClothThickness = &ManagedArrayCollection->AddAttribute<float>(ClothCollectionAttribute::FabricClothThickness, ClothCollectionGroup::Fabrics);
+		FabricClothDamping = &ManagedArrayCollection->AddAttribute<float>(ClothCollectionAttribute::FabricClothDamping, ClothCollectionGroup::Fabrics);
 		
 		// Sim Faces Group
-		SimIndices2D = &ManagedArrayCollection->AddAttribute<FIntVector3>(SimIndices2DAttribute, ClothCollectionGroup::SimFaces, SimVertices2DDependency);
-		SimIndices3D = &ManagedArrayCollection->AddAttribute<FIntVector3>(SimIndices3DAttribute, ClothCollectionGroup::SimFaces, SimVertices3DDependency);
+		SimIndices2D = &ManagedArrayCollection->AddAttribute<FIntVector3>(ClothCollectionAttribute::SimIndices2D, ClothCollectionGroup::SimFaces, SimVertices2DDependency);
+		SimIndices3D = &ManagedArrayCollection->AddAttribute<FIntVector3>(ClothCollectionAttribute::SimIndices3D, ClothCollectionGroup::SimFaces, SimVertices3DDependency);
 
 		// Sim Vertices 2D Group
-		SimPosition2D = &ManagedArrayCollection->AddAttribute<FVector2f>(SimPosition2DAttribute, ClothCollectionGroup::SimVertices2D);
-		SimVertex3DLookup = &ManagedArrayCollection->AddAttribute<int32>(SimVertex3DLookupAttribute, ClothCollectionGroup::SimVertices2D, SimVertices3DDependency);
+		SimPosition2D = &ManagedArrayCollection->AddAttribute<FVector2f>(ClothCollectionAttribute::SimPosition2D, ClothCollectionGroup::SimVertices2D);
+		SimVertex3DLookup = &ManagedArrayCollection->AddAttribute<int32>(ClothCollectionAttribute::SimVertex3DLookup, ClothCollectionGroup::SimVertices2D, SimVertices3DDependency);
 
 		// Sim Vertices 3D Group
-		SimPosition3D = &ManagedArrayCollection->AddAttribute<FVector3f>(SimPosition3DAttribute, ClothCollectionGroup::SimVertices3D);
-		SimNormal = &ManagedArrayCollection->AddAttribute<FVector3f>(SimNormalAttribute, ClothCollectionGroup::SimVertices3D);
-		SimBoneIndices = &ManagedArrayCollection->AddAttribute<TArray<int32>>(SimBoneIndicesAttribute, ClothCollectionGroup::SimVertices3D);
-		SimBoneWeights = &ManagedArrayCollection->AddAttribute<TArray<float>>(SimBoneWeightsAttribute, ClothCollectionGroup::SimVertices3D);
-		TetherKinematicIndex = &ManagedArrayCollection->AddAttribute<TArray<int32>>(TetherKinematicIndexAttribute, ClothCollectionGroup::SimVertices3D, SimVertices3DDependency);
-		TetherReferenceLength = &ManagedArrayCollection->AddAttribute<TArray<float>>(TetherReferenceLengthAttribute, ClothCollectionGroup::SimVertices3D);
-		SimVertex2DLookup = &ManagedArrayCollection->AddAttribute<TArray<int32>>(SimVertex2DLookupAttribute, ClothCollectionGroup::SimVertices3D, SimVertices2DDependency);
-		SeamStitchLookup = &ManagedArrayCollection->AddAttribute<TArray<int32>>(SeamStitchLookupAttribute, ClothCollectionGroup::SimVertices3D, SeamStitchesDependency);
+		SimPosition3D = &ManagedArrayCollection->AddAttribute<FVector3f>(ClothCollectionAttribute::SimPosition3D, ClothCollectionGroup::SimVertices3D);
+		SimNormal = &ManagedArrayCollection->AddAttribute<FVector3f>(ClothCollectionAttribute::SimNormal, ClothCollectionGroup::SimVertices3D);
+		SimBoneIndices = &ManagedArrayCollection->AddAttribute<TArray<int32>>(ClothCollectionAttribute::SimBoneIndices, ClothCollectionGroup::SimVertices3D);
+		SimBoneWeights = &ManagedArrayCollection->AddAttribute<TArray<float>>(ClothCollectionAttribute::SimBoneWeights, ClothCollectionGroup::SimVertices3D);
+		TetherKinematicIndex = &ManagedArrayCollection->AddAttribute<TArray<int32>>(ClothCollectionAttribute::TetherKinematicIndex, ClothCollectionGroup::SimVertices3D, SimVertices3DDependency);
+		TetherReferenceLength = &ManagedArrayCollection->AddAttribute<TArray<float>>(ClothCollectionAttribute::TetherReferenceLength, ClothCollectionGroup::SimVertices3D);
+		SimVertex2DLookup = &ManagedArrayCollection->AddAttribute<TArray<int32>>(ClothCollectionAttribute::SimVertex2DLookup, ClothCollectionGroup::SimVertices3D, SimVertices2DDependency);
+		SeamStitchLookup = &ManagedArrayCollection->AddAttribute<TArray<int32>>(ClothCollectionAttribute::SeamStitchLookup, ClothCollectionGroup::SimVertices3D, SeamStitchesDependency);
 
 		// Render Faces Group
-		RenderIndices = &ManagedArrayCollection->AddAttribute<FIntVector3>(RenderIndicesAttribute, ClothCollectionGroup::RenderFaces, RenderVerticesDependency);
+		RenderIndices = &ManagedArrayCollection->AddAttribute<FIntVector3>(ClothCollectionAttribute::RenderIndices, ClothCollectionGroup::RenderFaces, RenderVerticesDependency);
 
 		// Render Vertices Group
-		RenderPosition = &ManagedArrayCollection->AddAttribute<FVector3f>(RenderPositionAttribute, ClothCollectionGroup::RenderVertices);
-		RenderNormal = &ManagedArrayCollection->AddAttribute<FVector3f>(RenderNormalAttribute, ClothCollectionGroup::RenderVertices);
-		RenderTangentU = &ManagedArrayCollection->AddAttribute<FVector3f>(RenderTangentUAttribute, ClothCollectionGroup::RenderVertices);
-		RenderTangentV = &ManagedArrayCollection->AddAttribute<FVector3f>(RenderTangentVAttribute, ClothCollectionGroup::RenderVertices);
-		RenderUVs = &ManagedArrayCollection->AddAttribute<TArray<FVector2f>>(RenderUVsAttribute, ClothCollectionGroup::RenderVertices);
-		RenderColor = &ManagedArrayCollection->AddAttribute<FLinearColor>(RenderColorAttribute, ClothCollectionGroup::RenderVertices);
-		RenderBoneIndices = &ManagedArrayCollection->AddAttribute<TArray<int32>>(RenderBoneIndicesAttribute, ClothCollectionGroup::RenderVertices);
-		RenderBoneWeights = &ManagedArrayCollection->AddAttribute<TArray<float>>(RenderBoneWeightsAttribute, ClothCollectionGroup::RenderVertices);
+		RenderPosition = &ManagedArrayCollection->AddAttribute<FVector3f>(ClothCollectionAttribute::RenderPosition, ClothCollectionGroup::RenderVertices);
+		RenderNormal = &ManagedArrayCollection->AddAttribute<FVector3f>(ClothCollectionAttribute::RenderNormal, ClothCollectionGroup::RenderVertices);
+		RenderTangentU = &ManagedArrayCollection->AddAttribute<FVector3f>(ClothCollectionAttribute::RenderTangentU, ClothCollectionGroup::RenderVertices);
+		RenderTangentV = &ManagedArrayCollection->AddAttribute<FVector3f>(ClothCollectionAttribute::RenderTangentV, ClothCollectionGroup::RenderVertices);
+		RenderUVs = &ManagedArrayCollection->AddAttribute<TArray<FVector2f>>(ClothCollectionAttribute::RenderUVs, ClothCollectionGroup::RenderVertices);
+		RenderColor = &ManagedArrayCollection->AddAttribute<FLinearColor>(ClothCollectionAttribute::RenderColor, ClothCollectionGroup::RenderVertices);
+		RenderBoneIndices = &ManagedArrayCollection->AddAttribute<TArray<int32>>(ClothCollectionAttribute::RenderBoneIndices, ClothCollectionGroup::RenderVertices, SimVertices3DDependency);
+		RenderBoneWeights = &ManagedArrayCollection->AddAttribute<TArray<float>>(ClothCollectionAttribute::RenderBoneWeights, ClothCollectionGroup::RenderVertices);
+
+		if (EnumHasAnyFlags(OptionalSchemas, EClothCollectionOptionalSchemas::RenderDeformer))
+		{
+			// Render Patterns Group
+			RenderDeformerNumInfluences = &ManagedArrayCollection->AddAttribute<int32>(ClothCollectionAttribute::RenderDeformerNumInfluences, ClothCollectionGroup::RenderPatterns);
+
+			// Render Vertices Group
+			RenderDeformerPositionBaryCoordsAndDist = &ManagedArrayCollection->AddAttribute<TArray<FVector4f>>(ClothCollectionAttribute::RenderDeformerPositionBaryCoordsAndDist, ClothCollectionGroup::RenderVertices);
+			RenderDeformerNormalBaryCoordsAndDist = &ManagedArrayCollection->AddAttribute<TArray<FVector4f>>(ClothCollectionAttribute::RenderDeformerNormalBaryCoordsAndDist, ClothCollectionGroup::RenderVertices);
+			RenderDeformerTangentBaryCoordsAndDist = &ManagedArrayCollection->AddAttribute<TArray<FVector4f>>(ClothCollectionAttribute::RenderDeformerTangentBaryCoordsAndDist, ClothCollectionGroup::RenderVertices);
+			RenderDeformerSimIndices3D = &ManagedArrayCollection->AddAttribute<TArray<FIntVector3>>(ClothCollectionAttribute::RenderDeformerSimIndices3D, ClothCollectionGroup::RenderVertices);
+			RenderDeformerWeight = &ManagedArrayCollection->AddAttribute<TArray<float>>(ClothCollectionAttribute::RenderDeformerWeight, ClothCollectionGroup::RenderVertices);
+			RenderDeformerSkinningBlend = &ManagedArrayCollection->AddAttribute<float>(ClothCollectionAttribute::RenderDeformerSkinningBlend, ClothCollectionGroup::RenderVertices);
+		}
 	}
 
 	int32 FClothCollection::GetNumElements(const FName& GroupName) const
@@ -767,17 +761,4 @@ namespace UE::Chaos::ClothAsset
 		}
 		return false;
 	}
-
-	void FClothCollection::CopyArrayViewDataAndApplyOffset(const TArrayView<TArray<int32>>& To, const TConstArrayView<TArray<int32>>& From, const int32 Offset)
-	{
-		check(To.Num() == From.Num());
-		for (int32 Index = 0; Index < To.Num(); ++Index)
-		{
-			To[Index] = From[Index];
-			for (int32& Value : To[Index])
-			{
-				Value += Offset;
-			}
-		}
-	}
 } // End namespace UE::Chaos::ClothAsset
diff --git a/Engine/Plugins/Experimental/ChaosClothAsset/Source/ChaosClothAsset/Private/ChaosClothAsset/ClothCollection.h b/Engine/Plugins/Experimental/ChaosClothAsset/Source/ChaosClothAsset/Private/ChaosClothAsset/ClothCollection.h
--- a/Engine/Plugins/Experimental/ChaosClothAsset/Source/ChaosClothAsset/Private/ChaosClothAsset/ClothCollection.h
+++ b/Engine/Plugins/Experimental/ChaosClothAsset/Source/ChaosClothAsset/Private/ChaosClothAsset/ClothCollection.h
@@ -3,6 +3,7 @@
 #pragma once
 
 #include "GeometryCollection/ManagedArray.h"
+#include "ChaosClothAsset/ClothCollectionOptionalSchemas.h"
 #include "ChaosClothAsset/IsUserAttributeType.h"
 
 struct FManagedArrayCollection;
@@ -27,10 +28,10 @@ namespace UE::Chaos::ClothAsset
 		FClothCollection& operator=(FClothCollection&&) = delete;
 
 		/** Return whether the underlying collection is a valid cloth collection. */
-		bool IsValid() const;
+		bool IsValid(EClothCollectionOptionalSchemas OptionalSchemas = EClothCollectionOptionalSchemas::None) const;
 
 		/** Make the underlying collection a cloth collection. */
-		void DefineSchema();
+		void DefineSchema(EClothCollectionOptionalSchemas OptionalSchemas = EClothCollectionOptionalSchemas::None);
 
 		/** Get the number of elements of a group. */
 		int32 GetNumElements(const FName& GroupName) const;
@@ -117,10 +118,11 @@ namespace UE::Chaos::ClothAsset
 		template<typename T>
 		static inline void CopyArrayViewData(const TArrayView<T>& To, const TConstArrayView<T>& From);
 
-		template<typename T>
+		template<typename T, TEMPLATE_REQUIRES(!TIsArray<T>::Value)>
 		static inline void CopyArrayViewDataAndApplyOffset(const TArrayView<T>& To, const TConstArrayView<T>& From, const T Offset);
 
-		static void CopyArrayViewDataAndApplyOffset(const TArrayView<TArray<int32>>& To, const TConstArrayView<TArray<int32>>& From, const int32 Offset);
+		template<typename T>
+		static inline void CopyArrayViewDataAndApplyOffset(const TArrayView<TArray<T>>& To, const TConstArrayView<TArray<T>>& From, const T Offset);
 
 		template<typename T>
 		static inline uint32 GetElementsTypeHash(const TManagedArray<T>* ElementArray);
@@ -186,6 +188,7 @@ namespace UE::Chaos::ClothAsset
 		const TManagedArray<int32>* GetRenderVerticesEnd() const { return RenderVerticesEnd; }
 		const TManagedArray<int32>* GetRenderFacesStart() const { return RenderFacesStart; }
 		const TManagedArray<int32>* GetRenderFacesEnd() const { return RenderFacesEnd; }
+		const TManagedArray<int32>* GetRenderDeformerNumInfluences() const { return RenderDeformerNumInfluences; }
 		const TManagedArray<FString>* GetRenderMaterialPathName() const { return RenderMaterialPathName; }
 
 		//~ Sim Faces Group
@@ -218,6 +221,12 @@ namespace UE::Chaos::ClothAsset
 		const TManagedArray<FLinearColor>* GetRenderColor() const { return RenderColor; }
 		const TManagedArray<TArray<int32>>* GetRenderBoneIndices() const { return RenderBoneIndices; }
 		const TManagedArray<TArray<float>>* GetRenderBoneWeights() const { return RenderBoneWeights; }
+		const TManagedArray<TArray<FVector4f>>* GetRenderDeformerPositionBaryCoordsAndDist() const { return RenderDeformerPositionBaryCoordsAndDist; }
+		const TManagedArray<TArray<FVector4f>>* GetRenderDeformerNormalBaryCoordsAndDist() const { return RenderDeformerNormalBaryCoordsAndDist; }
+		const TManagedArray<TArray<FVector4f>>* GetRenderDeformerTangentBaryCoordsAndDist() const { return RenderDeformerTangentBaryCoordsAndDist; }
+		const TManagedArray<TArray<FIntVector3>>* GetRenderDeformerSimIndices3D() const { return RenderDeformerSimIndices3D; }
+		const TManagedArray<TArray<float>>* GetRenderDeformerWeight() const { return RenderDeformerWeight; }
+		const TManagedArray<float>* GetRenderDeformerSkinningBlend() const { return RenderDeformerSkinningBlend; }
 
 		//~ LODs Group (There should be only one LOD per ClothCollection)
 		TManagedArray<FString>* GetPhysicsAssetPathName(){ return PhysicsAssetPathName; }
@@ -253,6 +262,7 @@ namespace UE::Chaos::ClothAsset
 		TManagedArray<int32>* GetRenderVerticesEnd() { return RenderVerticesEnd; }
 		TManagedArray<int32>* GetRenderFacesStart() { return RenderFacesStart; }
 		TManagedArray<int32>* GetRenderFacesEnd() { return RenderFacesEnd; }
+		TManagedArray<int32>* GetRenderDeformerNumInfluences() { return RenderDeformerNumInfluences; }
 		TManagedArray<FString>* GetRenderMaterialPathName() { return RenderMaterialPathName; }
 
 		//~ Sim Faces Group
@@ -285,6 +295,12 @@ namespace UE::Chaos::ClothAsset
 		TManagedArray<FLinearColor>* GetRenderColor() { return RenderColor; }
 		TManagedArray<TArray<int32>>* GetRenderBoneIndices() { return RenderBoneIndices; }
 		TManagedArray<TArray<float>>* GetRenderBoneWeights() { return RenderBoneWeights; }
+		TManagedArray<TArray<FVector4f>>* GetRenderDeformerPositionBaryCoordsAndDist() { return RenderDeformerPositionBaryCoordsAndDist; }
+		TManagedArray<TArray<FVector4f>>* GetRenderDeformerNormalBaryCoordsAndDist() { return RenderDeformerNormalBaryCoordsAndDist; }
+		TManagedArray<TArray<FVector4f>>* GetRenderDeformerTangentBaryCoordsAndDist() { return RenderDeformerTangentBaryCoordsAndDist; }
+		TManagedArray<TArray<FIntVector3>>* GetRenderDeformerSimIndices3D() { return RenderDeformerSimIndices3D; }
+		TManagedArray<TArray<float>>* GetRenderDeformerWeight() { return RenderDeformerWeight; }
+		TManagedArray<float>* GetRenderDeformerSkinningBlend() { return RenderDeformerSkinningBlend; }
 
 	private:
 		//~ Cloth collection
@@ -324,6 +340,7 @@ namespace UE::Chaos::ClothAsset
 		TManagedArray<int32>* RenderVerticesEnd;
 		TManagedArray<int32>* RenderFacesStart;
 		TManagedArray<int32>* RenderFacesEnd;
+		TManagedArray<int32>* RenderDeformerNumInfluences;  // Number of deformer mapping influences per render vertex, either 0 (no mappings), 1 or 5 (= NUM_INFLUENCES_PER_VERTEX as defined in the ush files)
 		TManagedArray<FString>* RenderMaterialPathName;
 
 		//~ Sim Faces Group
@@ -356,6 +373,12 @@ namespace UE::Chaos::ClothAsset
 		TManagedArray<FLinearColor>* RenderColor;
 		TManagedArray<TArray<int32>>* RenderBoneIndices;
 		TManagedArray<TArray<float>>* RenderBoneWeights;
+		TManagedArray<TArray<FVector4f>>* RenderDeformerPositionBaryCoordsAndDist;  // Barycentric coords and distance along normal for the position of the final vert
+		TManagedArray<TArray<FVector4f>>* RenderDeformerNormalBaryCoordsAndDist;  // Barycentric coords and distance along normal for the location of the unit normal endpoint
+		TManagedArray<TArray<FVector4f>>* RenderDeformerTangentBaryCoordsAndDist;  // Barycentric coords and distance along normal for the location of the unit Tangent endpoint
+		TManagedArray<TArray<FIntVector3>>* RenderDeformerSimIndices3D;  // The source mesh triangle
+		TManagedArray<TArray<float>>* RenderDeformerWeight;  // For weighted averaging of multiple triangle influences
+		TManagedArray<float>* RenderDeformerSkinningBlend;  // Render weight map. How much the vertex actually contributes, value between 0 (fully deformed) and 1 (fully skinned)
 	};
 
 	template<typename T>
@@ -466,7 +489,7 @@ namespace UE::Chaos::ClothAsset
 		}
 	}
 
-	template<typename T>
+	template<typename T, typename TEnableIf<!TIsArray<T>::Value, int>::type>
 	inline void FClothCollection::CopyArrayViewDataAndApplyOffset(const TArrayView<T>& To, const TConstArrayView<T>& From, const T Offset)
 	{
 		check(To.Num() == From.Num());
@@ -476,6 +499,20 @@ namespace UE::Chaos::ClothAsset
 		}
 	}
 
+	template<typename T>
+	inline void FClothCollection::CopyArrayViewDataAndApplyOffset(const TArrayView<TArray<T>>& To, const TConstArrayView<TArray<T>>& From, const T Offset)
+	{
+		check(To.Num() == From.Num());
+		for (int32 Index = 0; Index < To.Num(); ++Index)
+		{
+			To[Index] = From[Index];
+			for (T& Value : To[Index])
+			{
+				Value += Offset;
+			}
+		}
+	}
+
 	template<typename T>
 	inline uint32 FClothCollection::GetElementsTypeHash(const TManagedArray<T>* ElementArray)
 	{
diff --git a/Engine/Plugins/Experimental/ChaosClothAsset/Source/ChaosClothAsset/Private/ChaosClothAsset/CollectionClothFacade.cpp b/Engine/Plugins/Experimental/ChaosClothAsset/Source/ChaosClothAsset/Private/ChaosClothAsset/CollectionClothFacade.cpp
--- a/Engine/Plugins/Experimental/ChaosClothAsset/Source/ChaosClothAsset/Private/ChaosClothAsset/CollectionClothFacade.cpp
+++ b/Engine/Plugins/Experimental/ChaosClothAsset/Source/ChaosClothAsset/Private/ChaosClothAsset/CollectionClothFacade.cpp
@@ -19,9 +19,9 @@ namespace UE::Chaos::ClothAsset
 	{
 	}
 
-	bool FCollectionClothConstFacade::IsValid() const
+	bool FCollectionClothConstFacade::IsValid(EClothCollectionOptionalSchemas OptionalSchemas) const
 	{
-		return ClothCollection->IsValid() && ClothCollection->GetNumElements(ClothCollectionGroup::Lods) == 1;
+		return ClothCollection->IsValid(OptionalSchemas) && ClothCollection->GetNumElements(ClothCollectionGroup::Lods) == 1;
 	}
 
 	bool FCollectionClothConstFacade::HasValidData() const
@@ -66,6 +66,7 @@ namespace UE::Chaos::ClothAsset
 		ResultHash = HashCombineFast(ResultHash, ClothCollection->GetElementsTypeHash(ClothCollection->GetRenderFacesStart()));
 		ResultHash = HashCombineFast(ResultHash, ClothCollection->GetElementsTypeHash(ClothCollection->GetRenderFacesEnd()));
 		ResultHash = HashCombineFast(ResultHash, ClothCollection->GetElementsTypeHash(ClothCollection->GetRenderMaterialPathName()));
+		ResultHash = HashCombineFast(ResultHash, ClothCollection->GetElementsTypeHash(ClothCollection->GetRenderDeformerNumInfluences()));
 
 		//~ Sim Faces Group
 		ResultHash = HashCombineFast(ResultHash, ClothCollection->GetElementsTypeHash(ClothCollection->GetSimIndices2D()));
@@ -107,6 +108,12 @@ namespace UE::Chaos::ClothAsset
 		ResultHash = HashCombineFast(ResultHash, ClothCollection->GetElementsTypeHash(ClothCollection->GetRenderColor()));
 		ResultHash = HashCombineFast(ResultHash, ClothCollection->GetElementsTypeHash(ClothCollection->GetRenderBoneIndices()));
 		ResultHash = HashCombineFast(ResultHash, ClothCollection->GetElementsTypeHash(ClothCollection->GetRenderBoneWeights()));
+		ResultHash = HashCombineFast(ResultHash, ClothCollection->GetElementsTypeHash(ClothCollection->GetRenderDeformerPositionBaryCoordsAndDist()));
+		ResultHash = HashCombineFast(ResultHash, ClothCollection->GetElementsTypeHash(ClothCollection->GetRenderDeformerNormalBaryCoordsAndDist()));
+		ResultHash = HashCombineFast(ResultHash, ClothCollection->GetElementsTypeHash(ClothCollection->GetRenderDeformerTangentBaryCoordsAndDist()));
+		ResultHash = HashCombineFast(ResultHash, ClothCollection->GetElementsTypeHash(ClothCollection->GetRenderDeformerSimIndices3D()));
+		ResultHash = HashCombineFast(ResultHash, ClothCollection->GetElementsTypeHash(ClothCollection->GetRenderDeformerWeight()));
+		ResultHash = HashCombineFast(ResultHash, ClothCollection->GetElementsTypeHash(ClothCollection->GetRenderDeformerSkinningBlend()));
 
 		if (bIncludeWeightMaps)
 		{
@@ -260,6 +267,11 @@ namespace UE::Chaos::ClothAsset
 		return FCollectionClothRenderPatternConstFacade(ClothCollection, PatternIndex);
 	}
 
+	TConstArrayView<int32> FCollectionClothConstFacade::GetRenderDeformerNumInfluences() const
+	{
+		return ClothCollection->GetElements(ClothCollection->GetRenderDeformerNumInfluences());
+	}
+
 	TConstArrayView<FString> FCollectionClothConstFacade::GetRenderMaterialPathName() const
 	{
 		return ClothCollection->GetElements(ClothCollection->GetRenderMaterialPathName());
@@ -340,6 +352,36 @@ namespace UE::Chaos::ClothAsset
 		return ClothCollection->GetElements(ClothCollection->GetRenderBoneWeights());
 	}
 
+	TConstArrayView<TArray<FVector4f>> FCollectionClothConstFacade::GetRenderDeformerPositionBaryCoordsAndDist() const
+	{
+		return ClothCollection->GetElements(ClothCollection->GetRenderDeformerPositionBaryCoordsAndDist());
+	}
+
+	TConstArrayView<TArray<FVector4f>> FCollectionClothConstFacade::GetRenderDeformerNormalBaryCoordsAndDist() const
+	{
+		return ClothCollection->GetElements(ClothCollection->GetRenderDeformerNormalBaryCoordsAndDist());
+	}
+
+	TConstArrayView<TArray<FVector4f>> FCollectionClothConstFacade::GetRenderDeformerTangentBaryCoordsAndDist() const
+	{
+		return ClothCollection->GetElements(ClothCollection->GetRenderDeformerTangentBaryCoordsAndDist());
+	}
+
+	TConstArrayView<TArray<FIntVector3>> FCollectionClothConstFacade::GetRenderDeformerSimIndices3D() const
+	{
+		return ClothCollection->GetElements(ClothCollection->GetRenderDeformerSimIndices3D());
+	}
+
+	TConstArrayView<TArray<float>> FCollectionClothConstFacade::GetRenderDeformerWeight() const
+	{
+		return ClothCollection->GetElements(ClothCollection->GetRenderDeformerWeight());
+	}
+
+	TConstArrayView<float> FCollectionClothConstFacade::GetRenderDeformerSkinningBlend() const
+	{
+		return ClothCollection->GetElements(ClothCollection->GetRenderDeformerSkinningBlend());
+	}
+
 	int32 FCollectionClothConstFacade::GetNumRenderFaces() const
 	{
 		return ClothCollection->GetNumElements(ClothCollectionGroup::RenderFaces);
@@ -443,9 +485,9 @@ namespace UE::Chaos::ClothAsset
 	{
 	}
 
-	void FCollectionClothFacade::DefineSchema()
+	void FCollectionClothFacade::DefineSchema(EClothCollectionOptionalSchemas OptionalSchemas)
 	{
-		GetClothCollection()->DefineSchema();
+		GetClothCollection()->DefineSchema(OptionalSchemas);
 		SetDefaults();
 	}
 
@@ -536,7 +578,7 @@ namespace UE::Chaos::ClothAsset
 		SetNumRenderPatterns(StartNumRenderPatterns + OtherNumRenderPatterns);
 		for (int32 PatternIndex = 0; PatternIndex < OtherNumRenderPatterns; ++PatternIndex)
 		{
-			GetRenderPattern(PatternIndex + StartNumRenderPatterns).Initialize(Other.GetRenderPattern(PatternIndex));
+			GetRenderPattern(PatternIndex + StartNumRenderPatterns).Initialize(Other.GetRenderPattern(PatternIndex), StartNumSimVertices3D);
 		}
 
 		// Weight maps
@@ -737,6 +779,11 @@ namespace UE::Chaos::ClothAsset
 		GetClothCollection()->RemoveElements(ClothCollectionGroup::RenderPatterns, SortedDeletionList);
 	}
 
+	TArrayView<int32> FCollectionClothFacade::GetRenderDeformerNumInfluences()
+	{
+		return GetClothCollection()->GetElements(GetClothCollection()->GetRenderDeformerNumInfluences());
+	}
+
 	TArrayView<FString> FCollectionClothFacade::GetRenderMaterialPathName()
 	{
 		return GetClothCollection()->GetElements(GetClothCollection()->GetRenderMaterialPathName());
@@ -859,6 +906,36 @@ namespace UE::Chaos::ClothAsset
 		return GetClothCollection()->GetElements(GetClothCollection()->GetRenderBoneWeights());
 	}
 
+	TArrayView<TArray<FVector4f>> FCollectionClothFacade::GetRenderDeformerPositionBaryCoordsAndDist()
+	{
+		return GetClothCollection()->GetElements(GetClothCollection()->GetRenderDeformerPositionBaryCoordsAndDist());
+	}
+
+	TArrayView<TArray<FVector4f>> FCollectionClothFacade::GetRenderDeformerNormalBaryCoordsAndDist()
+	{
+		return GetClothCollection()->GetElements(GetClothCollection()->GetRenderDeformerNormalBaryCoordsAndDist());
+	}
+
+	TArrayView<TArray<FVector4f>> FCollectionClothFacade::GetRenderDeformerTangentBaryCoordsAndDist()
+	{
+		return GetClothCollection()->GetElements(GetClothCollection()->GetRenderDeformerTangentBaryCoordsAndDist());
+	}
+
+	TArrayView<TArray<FIntVector3>> FCollectionClothFacade::GetRenderDeformerSimIndices3D()
+	{
+		return GetClothCollection()->GetElements(GetClothCollection()->GetRenderDeformerSimIndices3D());
+	}
+
+	TArrayView<TArray<float>> FCollectionClothFacade::GetRenderDeformerWeight()
+	{
+		return GetClothCollection()->GetElements(GetClothCollection()->GetRenderDeformerWeight());
+	}
+
+	TArrayView<float> FCollectionClothFacade::GetRenderDeformerSkinningBlend()
+	{
+		return GetClothCollection()->GetElements(GetClothCollection()->GetRenderDeformerSkinningBlend());
+	}
+
 	TArrayView<FIntVector3> FCollectionClothFacade::GetRenderIndices()
 	{
 		return GetClothCollection()->GetElements(GetClothCollection()->GetRenderIndices());
diff --git a/Engine/Plugins/Experimental/ChaosClothAsset/Source/ChaosClothAsset/Private/ChaosClothAsset/CollectionClothRenderPatternFacade.cpp b/Engine/Plugins/Experimental/ChaosClothAsset/Source/ChaosClothAsset/Private/ChaosClothAsset/CollectionClothRenderPatternFacade.cpp
--- a/Engine/Plugins/Experimental/ChaosClothAsset/Source/ChaosClothAsset/Private/ChaosClothAsset/CollectionClothRenderPatternFacade.cpp
+++ b/Engine/Plugins/Experimental/ChaosClothAsset/Source/ChaosClothAsset/Private/ChaosClothAsset/CollectionClothRenderPatternFacade.cpp
@@ -7,6 +7,12 @@
 
 namespace UE::Chaos::ClothAsset
 {
+	int32 FCollectionClothRenderPatternConstFacade::GetRenderDeformerNumInfluences() const
+	{
+		return ClothCollection->GetRenderDeformerNumInfluences() && ClothCollection->GetNumElements(ClothCollectionGroup::RenderPatterns) > GetElementIndex() ?
+			(*ClothCollection->GetRenderDeformerNumInfluences())[GetElementIndex()] : 0;
+	}
+
 	const FString& FCollectionClothRenderPatternConstFacade::GetRenderMaterialPathName() const
 	{
 		static const FString EmptyString;
@@ -102,6 +108,60 @@ namespace UE::Chaos::ClothAsset
 			GetElementIndex());
 	}
 
+	TConstArrayView<TArray<FVector4f>> FCollectionClothRenderPatternConstFacade::GetRenderDeformerPositionBaryCoordsAndDist() const
+	{
+		return ClothCollection->GetElements(
+			ClothCollection->GetRenderDeformerPositionBaryCoordsAndDist(),
+			ClothCollection->GetRenderVerticesStart(),
+			ClothCollection->GetRenderVerticesEnd(),
+			GetElementIndex());
+	}
+
+	TConstArrayView<TArray<FVector4f>> FCollectionClothRenderPatternConstFacade::GetRenderDeformerNormalBaryCoordsAndDist() const
+	{
+		return ClothCollection->GetElements(
+			ClothCollection->GetRenderDeformerNormalBaryCoordsAndDist(),
+			ClothCollection->GetRenderVerticesStart(),
+			ClothCollection->GetRenderVerticesEnd(),
+			GetElementIndex());
+	}
+
+	TConstArrayView<TArray<FVector4f>> FCollectionClothRenderPatternConstFacade::GetRenderDeformerTangentBaryCoordsAndDist() const
+	{
+		return ClothCollection->GetElements(
+			ClothCollection->GetRenderDeformerTangentBaryCoordsAndDist(),
+			ClothCollection->GetRenderVerticesStart(),
+			ClothCollection->GetRenderVerticesEnd(),
+			GetElementIndex());
+	}
+
+	TConstArrayView<TArray<FIntVector3>> FCollectionClothRenderPatternConstFacade::GetRenderDeformerSimIndices3D() const
+	{
+		return ClothCollection->GetElements(
+			ClothCollection->GetRenderDeformerSimIndices3D(),
+			ClothCollection->GetRenderVerticesStart(),
+			ClothCollection->GetRenderVerticesEnd(),
+			GetElementIndex());
+	}
+
+	TConstArrayView<TArray<float>> FCollectionClothRenderPatternConstFacade::GetRenderDeformerWeight() const
+	{
+		return ClothCollection->GetElements(
+			ClothCollection->GetRenderDeformerWeight(),
+			ClothCollection->GetRenderVerticesStart(),
+			ClothCollection->GetRenderVerticesEnd(),
+			GetElementIndex());
+	}
+
+	TConstArrayView<float> FCollectionClothRenderPatternConstFacade::GetRenderDeformerSkinningBlend() const
+	{
+		return ClothCollection->GetElements(
+			ClothCollection->GetRenderDeformerSkinningBlend(),
+			ClothCollection->GetRenderVerticesStart(),
+			ClothCollection->GetRenderVerticesEnd(),
+			GetElementIndex());
+	}
+
 	int32 FCollectionClothRenderPatternConstFacade::GetNumRenderFaces() const
 	{
 		return ClothCollection->GetNumElements(
@@ -147,10 +207,14 @@ namespace UE::Chaos::ClothAsset
 		SetDefaults();
 	}
 
-	void FCollectionClothRenderPatternFacade::Initialize(const FCollectionClothRenderPatternConstFacade& Other)
+	void FCollectionClothRenderPatternFacade::Initialize(const FCollectionClothRenderPatternConstFacade& Other, int32 SimVertex3DOffset)
 	{
 		Reset();
 
+		//~ Render Patterns Group
+		SetRenderDeformerNumInfluences(Other.GetRenderDeformerNumInfluences());  // Must be called prior to copying the Render Deformer as it also sets the optional RenderDeformer schema
+		SetRenderMaterialPathName(Other.GetRenderMaterialPathName());
+
 		//~ Render Vertices Group
 		SetNumRenderVertices(Other.GetNumRenderVertices());
 		FClothCollection::CopyArrayViewData(GetRenderPosition(), Other.GetRenderPosition());
@@ -161,13 +225,31 @@ namespace UE::Chaos::ClothAsset
 		FClothCollection::CopyArrayViewData(GetRenderColor(), Other.GetRenderColor());
 		FClothCollection::CopyArrayViewData(GetRenderBoneIndices(), Other.GetRenderBoneIndices());
 		FClothCollection::CopyArrayViewData(GetRenderBoneWeights(), Other.GetRenderBoneWeights());
+		FClothCollection::CopyArrayViewData(GetRenderDeformerPositionBaryCoordsAndDist(), Other.GetRenderDeformerPositionBaryCoordsAndDist());
+		FClothCollection::CopyArrayViewData(GetRenderDeformerNormalBaryCoordsAndDist(), Other.GetRenderDeformerNormalBaryCoordsAndDist());
+		FClothCollection::CopyArrayViewData(GetRenderDeformerTangentBaryCoordsAndDist(), Other.GetRenderDeformerTangentBaryCoordsAndDist());
+		FClothCollection::CopyArrayViewDataAndApplyOffset(GetRenderDeformerSimIndices3D(), Other.GetRenderDeformerSimIndices3D(), FIntVector3(SimVertex3DOffset));
+		FClothCollection::CopyArrayViewData(GetRenderDeformerWeight(), Other.GetRenderDeformerWeight());
+		FClothCollection::CopyArrayViewData(GetRenderDeformerSkinningBlend(), Other.GetRenderDeformerSkinningBlend());
 
 		//~ Render Faces Group
 		const int32 RenderVertexOffset = GetRenderVerticesOffset() - Other.GetRenderVerticesOffset();
 		SetNumRenderFaces(Other.GetNumRenderFaces());
-		FClothCollection::CopyArrayViewDataAndApplyOffset(GetRenderIndices(), Other.GetRenderIndices(), FIntVector(RenderVertexOffset));
+		FClothCollection::CopyArrayViewDataAndApplyOffset(GetRenderIndices(), Other.GetRenderIndices(), FIntVector3(RenderVertexOffset));
+	}
 
-		SetRenderMaterialPathName(Other.GetRenderMaterialPathName());
+	void FCollectionClothRenderPatternFacade::SetRenderDeformerNumInfluences(int32 NumInfluences)
+	{
+		check(NumInfluences >= 0);
+		if (!GetClothCollection()->IsValid(EClothCollectionOptionalSchemas::RenderDeformer))
+		{
+			if (NumInfluences == 0)
+			{
+				return;  // Not having the schema is the same as having zero influences
+			}
+			GetClothCollection()->DefineSchema(EClothCollectionOptionalSchemas::RenderDeformer);
+		}
+		(*GetClothCollection()->GetRenderDeformerNumInfluences())[GetElementIndex()] = NumInfluences;
 	}
 
 	void FCollectionClothRenderPatternFacade::SetRenderMaterialPathName(const FString& PathName)
@@ -275,6 +357,60 @@ namespace UE::Chaos::ClothAsset
 			GetElementIndex());
 	}
 
+	TArrayView<TArray<FVector4f>> FCollectionClothRenderPatternFacade::GetRenderDeformerPositionBaryCoordsAndDist()
+	{
+		return GetClothCollection()->GetElements(
+			GetClothCollection()->GetRenderDeformerPositionBaryCoordsAndDist(),
+			GetClothCollection()->GetRenderVerticesStart(),
+			GetClothCollection()->GetRenderVerticesEnd(),
+			GetElementIndex());
+	}
+
+	TArrayView<TArray<FVector4f>> FCollectionClothRenderPatternFacade::GetRenderDeformerNormalBaryCoordsAndDist()
+	{
+		return GetClothCollection()->GetElements(
+			GetClothCollection()->GetRenderDeformerNormalBaryCoordsAndDist(),
+			GetClothCollection()->GetRenderVerticesStart(),
+			GetClothCollection()->GetRenderVerticesEnd(),
+			GetElementIndex());
+	}
+
+	TArrayView<TArray<FVector4f>> FCollectionClothRenderPatternFacade::GetRenderDeformerTangentBaryCoordsAndDist()
+	{
+		return GetClothCollection()->GetElements(
+			GetClothCollection()->GetRenderDeformerTangentBaryCoordsAndDist(),
+			GetClothCollection()->GetRenderVerticesStart(),
+			GetClothCollection()->GetRenderVerticesEnd(),
+			GetElementIndex());
+	}
+
+	TArrayView<TArray<FIntVector3>> FCollectionClothRenderPatternFacade::GetRenderDeformerSimIndices3D()
+	{
+		return GetClothCollection()->GetElements(
+			GetClothCollection()->GetRenderDeformerSimIndices3D(),
+			GetClothCollection()->GetRenderVerticesStart(),
+			GetClothCollection()->GetRenderVerticesEnd(),
+			GetElementIndex());
+	}
+
+	TArrayView<TArray<float>> FCollectionClothRenderPatternFacade::GetRenderDeformerWeight()
+	{
+		return GetClothCollection()->GetElements(
+			GetClothCollection()->GetRenderDeformerWeight(),
+			GetClothCollection()->GetRenderVerticesStart(),
+			GetClothCollection()->GetRenderVerticesEnd(),
+			GetElementIndex());
+	}
+
+	TArrayView<float> FCollectionClothRenderPatternFacade::GetRenderDeformerSkinningBlend()
+	{
+		return GetClothCollection()->GetElements(
+			GetClothCollection()->GetRenderDeformerSkinningBlend(),
+			GetClothCollection()->GetRenderVerticesStart(),
+			GetClothCollection()->GetRenderVerticesEnd(),
+			GetElementIndex());
+	}
+
 	void FCollectionClothRenderPatternFacade::SetNumRenderFaces(int32 NumRenderFaces)
 	{
 		GetClothCollection()->SetNumElements(
diff --git a/Engine/Plugins/Experimental/ChaosClothAsset/Source/ChaosClothAsset/Public/ChaosClothAsset/ClothCollectionAttribute.h b/Engine/Plugins/Experimental/ChaosClothAsset/Source/ChaosClothAsset/Public/ChaosClothAsset/ClothCollectionAttribute.h
new file mode 100644
--- /dev/null
+++ b/Engine/Plugins/Experimental/ChaosClothAsset/Source/ChaosClothAsset/Public/ChaosClothAsset/ClothCollectionAttribute.h
@@ -0,0 +1,88 @@
+// Copyright Epic Games, Inc. All Rights Reserved.
+
+#pragma once
+
+#include "UObject/NameTypes.h"
+
+namespace UE::Chaos::ClothAsset
+{
+	/**
+	 * Cloth collection reserved attribute names.
+	 */
+	namespace ClothCollectionAttribute
+	{
+		// Lods Group
+		inline const FName PhysicsAssetPathName(TEXT("PhysicsAssetPathName"));
+		inline const FName SkeletalMeshPathName(TEXT("SkeletalMeshPathName"));
+
+		// Fabrics Group
+		inline const FName FabricBendingStiffness(TEXT("FabricBendingStiffness"));
+		inline const FName FabricBucklingStiffness(TEXT("FabricBucklingStiffness"));
+		inline const FName FabricStretchStiffness(TEXT("FabricStretchStiffness"));
+		inline const FName FabricBucklingRatio(TEXT("FabricBucklingRatio"));
+		inline const FName FabricClothDensity(TEXT("FabricClothDensity"));
+		inline const FName FabricClothFriction(TEXT("FabricClothFriction"));
+		inline const FName FabricClothThickness(TEXT("FabricClothThickness"));
+		inline const FName FabricClothDamping(TEXT("FabricClothDamping"));
+
+		// Seam Group
+		inline const FName SeamStitchStart(TEXT("SeamStitchStart"));
+		inline const FName SeamStitchEnd(TEXT("SeamStitchEnd"));
+
+		// Seam Stitches Group
+		inline const FName SeamStitch2DEndIndices(TEXT("SeamStitch2DEndIndices"));
+		inline const FName SeamStitch3DIndex(TEXT("SeamStitch3DIndex"));
+
+		// Sim Patterns Group
+		inline const FName SimVertices2DStart(TEXT("SimVertices2DStart"));
+		inline const FName SimVertices2DEnd(TEXT("SimVertices2DEnd"));
+		inline const FName SimFacesStart(TEXT("SimFacesStart"));
+		inline const FName SimFacesEnd(TEXT("SimFacesEnd"));
+		inline const FName SimPatternFabric(TEXT("SimPatternFabric"));
+
+		// Render Patterns Group
+		inline const FName RenderVerticesStart(TEXT("RenderVerticesStart"));
+		inline const FName RenderVerticesEnd(TEXT("RenderVerticesEnd"));
+		inline const FName RenderFacesStart(TEXT("RenderFacesStart"));
+		inline const FName RenderFacesEnd(TEXT("RenderFacesEnd"));
+		inline const FName RenderMaterialPathName(TEXT("RenderMaterialPathName"));
+		inline const FName RenderDeformerNumInfluences(TEXT("RenderDeformerNumInfluences"));
+
+		// Sim Faces Group
+		inline const FName SimIndices2D(TEXT("SimIndices2D"));
+		inline const FName SimIndices3D(TEXT("SimIndices3D"));
+
+		// Sim Vertices 2D Group
+		inline const FName SimPosition2D(TEXT("SimPosition2D"));
+		inline const FName SimVertex3DLookup(TEXT("SimVertex3DLookup"));
+
+		// Sim Vertices 3D Group
+		inline const FName SimPosition3D(TEXT("SimPosition3D"));
+		inline const FName SimNormal(TEXT("SimNormal"));
+		inline const FName SimBoneIndices(TEXT("SimBoneIndices"));
+		inline const FName SimBoneWeights(TEXT("SimBoneWeights"));
+		inline const FName TetherKinematicIndex(TEXT("TetherKinematicIndex"));
+		inline const FName TetherReferenceLength(TEXT("TetherReferenceLength"));
+		inline const FName SimVertex2DLookup(TEXT("SimVertex2DLookup"));
+		inline const FName SeamStitchLookup(TEXT("SeamStitchLookup"));
+
+		// Render Faces Group
+		inline const FName RenderIndices(TEXT("RenderIndices"));
+
+		// Render Vertices Group
+		inline const FName RenderPosition(TEXT("RenderPosition"));
+		inline const FName RenderNormal(TEXT("RenderNormal"));
+		inline const FName RenderTangentU(TEXT("RenderTangentU"));
+		inline const FName RenderTangentV(TEXT("RenderTangentV"));
+		inline const FName RenderUVs(TEXT("RenderUVs"));
+		inline const FName RenderColor(TEXT("RenderColor"));
+		inline const FName RenderBoneIndices(TEXT("RenderBoneIndices"));
+		inline const FName RenderBoneWeights(TEXT("RenderBoneWeights"));
+		inline const FName RenderDeformerPositionBaryCoordsAndDist(TEXT("RenderDeformerPositionBaryCoordsAndDist"));
+		inline const FName RenderDeformerNormalBaryCoordsAndDist(TEXT("RenderDeformerNormalBaryCoordsAndDist"));
+		inline const FName RenderDeformerTangentBaryCoordsAndDist(TEXT("RenderDeformerTangentBaryCoordsAndDist"));
+		inline const FName RenderDeformerSimIndices3D(TEXT("RenderDeformerSimIndices3D"));
+		inline const FName RenderDeformerWeight(TEXT("RenderDeformerWeight"));
+		inline const FName RenderDeformerSkinningBlend(TEXT("RenderDeformerSkinningBlend"));
+	}
+}  // End namespace UE::Chaos::ClothAsset
diff --git a/Engine/Plugins/Experimental/ChaosClothAsset/Source/ChaosClothAsset/Public/ChaosClothAsset/ClothCollectionOptionalSchemas.h b/Engine/Plugins/Experimental/ChaosClothAsset/Source/ChaosClothAsset/Public/ChaosClothAsset/ClothCollectionOptionalSchemas.h
new file mode 100644
--- /dev/null
+++ b/Engine/Plugins/Experimental/ChaosClothAsset/Source/ChaosClothAsset/Public/ChaosClothAsset/ClothCollectionOptionalSchemas.h
@@ -0,0 +1,17 @@
+// Copyright Epic Games, Inc. All Rights Reserved.
+
+#pragma once
+
+#include "HAL/Platform.h"
+#include "Misc/EnumClassFlags.h"
+
+namespace UE::Chaos::ClothAsset
+{
+	/** Cloth collection optional schemas. */
+	enum class EClothCollectionOptionalSchemas : uint8
+	{
+		None = 0,
+		RenderDeformer = 1 << 1
+	};
+	ENUM_CLASS_FLAGS(EClothCollectionOptionalSchemas)
+}  // End namespace UE::Chaos::ClothAsset
diff --git a/Engine/Plugins/Experimental/ChaosClothAsset/Source/ChaosClothAsset/Public/ChaosClothAsset/CollectionClothFacade.h b/Engine/Plugins/Experimental/ChaosClothAsset/Source/ChaosClothAsset/Public/ChaosClothAsset/CollectionClothFacade.h
--- a/Engine/Plugins/Experimental/ChaosClothAsset/Source/ChaosClothAsset/Public/ChaosClothAsset/CollectionClothFacade.h
+++ b/Engine/Plugins/Experimental/ChaosClothAsset/Source/ChaosClothAsset/Public/ChaosClothAsset/CollectionClothFacade.h
@@ -2,11 +2,12 @@
 
 #pragma once
 
-#include "ChaosClothAsset/CollectionClothSimPatternFacade.h"
+#include "ChaosClothAsset/ClothCollectionOptionalSchemas.h"
+#include "ChaosClothAsset/CollectionClothFabricFacade.h"
 #include "ChaosClothAsset/CollectionClothRenderPatternFacade.h"
 #include "ChaosClothAsset/CollectionClothSeamFacade.h"
+#include "ChaosClothAsset/CollectionClothSimPatternFacade.h"
 #include "ChaosClothAsset/IsUserAttributeType.h"
-#include "ChaosClothAsset/CollectionClothFabricFacade.h"
 
 namespace Chaos
 {
@@ -35,7 +36,7 @@ namespace UE::Chaos::ClothAsset
 		virtual ~FCollectionClothConstFacade() = default;
 
 		/** Return whether the facade is defined on the collection. */
-		bool IsValid() const;
+		bool IsValid(EClothCollectionOptionalSchemas OptionalSchemas = EClothCollectionOptionalSchemas::None) const;
 
 		/**
 		 * Return whether the facade has a non-empty sim and render mesh data.
@@ -46,7 +47,6 @@ namespace UE::Chaos::ClothAsset
 		uint32 CalculateWeightMapTypeHash(uint32 PreviousHash = 0) const;
 		template<typename T, TEMPLATE_REQUIRES(TIsUserAttributeType<T>::Value)>
 		uint32 CalculateUserDefinedAttributesTypeHash(const FName& GroupName, uint32 PreviousHash = 0) const;
-		
 
 		//~ LOD (single per collection) Group
 		/** Return the physics asset path names used for this collection. */
@@ -88,12 +88,13 @@ namespace UE::Chaos::ClothAsset
 		/** Convenience to find which sim pattern a sim face belongs to */
 		int32 FindSimPatternByFaceIndex(int32 FaceIndex) const;
 
-
 		//~ Render Patterns Group
 		/** Return the number of patterns in this collection. */
 		int32 GetNumRenderPatterns() const;
 		/** Return a pattern facade for the specified pattern index. */
 		FCollectionClothRenderPatternConstFacade GetRenderPattern(int32 PatternIndex) const;
+		/** Return a view of all the render deformer number of influences used on this collection across all patterns. */
+		TConstArrayView<int32> GetRenderDeformerNumInfluences() const;
 		/** Return a view of all the render materials used on this collection across all patterns. */
 		TConstArrayView<FString> GetRenderMaterialPathName() const;
 		/** Convenience to find which render pattern a render vertex belongs to */
@@ -124,6 +125,12 @@ namespace UE::Chaos::ClothAsset
 		TConstArrayView<FLinearColor> GetRenderColor() const;
 		TConstArrayView<TArray<int32>> GetRenderBoneIndices() const;
 		TConstArrayView<TArray<float>> GetRenderBoneWeights() const;
+		TConstArrayView<TArray<FVector4f>> GetRenderDeformerPositionBaryCoordsAndDist() const;
+		TConstArrayView<TArray<FVector4f>> GetRenderDeformerNormalBaryCoordsAndDist() const;
+		TConstArrayView<TArray<FVector4f>> GetRenderDeformerTangentBaryCoordsAndDist() const;
+		TConstArrayView<TArray<FIntVector3>> GetRenderDeformerSimIndices3D() const;
+		TConstArrayView<TArray<float>> GetRenderDeformerWeight() const;
+		TConstArrayView<float> GetRenderDeformerSkinningBlend() const;
 
 		//~ Render Faces Group
 		int32 GetNumRenderFaces() const;
@@ -177,7 +184,7 @@ namespace UE::Chaos::ClothAsset
 		virtual ~FCollectionClothFacade() override = default;
 
 		/** Create this facade's groups and attributes. */
-		void DefineSchema();
+		void DefineSchema(EClothCollectionOptionalSchemas OptionalSchemas = EClothCollectionOptionalSchemas::None);
 
 		/** Remove all LODs from this cloth. */
 		void Reset();
@@ -244,6 +251,8 @@ namespace UE::Chaos::ClothAsset
 		FCollectionClothRenderPatternFacade AddGetRenderPattern() { return GetRenderPattern(AddRenderPattern()); }
 		/** Remove a sorted list of render patterns. */
 		void RemoveRenderPatterns(const TArray<int32>& SortedDeletionList);
+		/** Return a view of all the render deformer number of influences used on this collection across all patterns. */
+		TArrayView<int32> GetRenderDeformerNumInfluences();
 		/** Return a view of all the render materials used on this collection across all patterns. */
 		TArrayView<FString> GetRenderMaterialPathName();
 
@@ -281,6 +290,12 @@ namespace UE::Chaos::ClothAsset
 		TArrayView<FLinearColor> GetRenderColor();
 		TArrayView<TArray<int32>> GetRenderBoneIndices();
 		TArrayView<TArray<float>> GetRenderBoneWeights();
+		TArrayView<TArray<FVector4f>> GetRenderDeformerPositionBaryCoordsAndDist();
+		TArrayView<TArray<FVector4f>> GetRenderDeformerNormalBaryCoordsAndDist();
+		TArrayView<TArray<FVector4f>> GetRenderDeformerTangentBaryCoordsAndDist();
+		TArrayView<TArray<FIntVector3>> GetRenderDeformerSimIndices3D();
+		TArrayView<TArray<float>> GetRenderDeformerWeight();
+		TArrayView<float> GetRenderDeformerSkinningBlend();
 
 		//~ Render Faces Group
 		/** SetNumRenderFaces per pattern within pattern facade. */
diff --git a/Engine/Plugins/Experimental/ChaosClothAsset/Source/ChaosClothAsset/Public/ChaosClothAsset/CollectionClothRenderPatternFacade.h b/Engine/Plugins/Experimental/ChaosClothAsset/Source/ChaosClothAsset/Public/ChaosClothAsset/CollectionClothRenderPatternFacade.h
--- a/Engine/Plugins/Experimental/ChaosClothAsset/Source/ChaosClothAsset/Public/ChaosClothAsset/CollectionClothRenderPatternFacade.h
+++ b/Engine/Plugins/Experimental/ChaosClothAsset/Source/ChaosClothAsset/Public/ChaosClothAsset/CollectionClothRenderPatternFacade.h
@@ -26,6 +26,8 @@ namespace UE::Chaos::ClothAsset
 
 		virtual ~FCollectionClothRenderPatternConstFacade() = default;
 
+		/** Return the render deformer number of influences for this pattern. */
+		int32 GetRenderDeformerNumInfluences() const;
 		/** Return the render material for this pattern. */
 		const FString& GetRenderMaterialPathName() const;
 
@@ -43,7 +45,13 @@ namespace UE::Chaos::ClothAsset
 		TConstArrayView<FLinearColor> GetRenderColor() const;
 		TConstArrayView<TArray<int32>> GetRenderBoneIndices() const;
 		TConstArrayView<TArray<float>> GetRenderBoneWeights() const;
-		
+		TConstArrayView<TArray<FVector4f>> GetRenderDeformerPositionBaryCoordsAndDist() const;
+		TConstArrayView<TArray<FVector4f>> GetRenderDeformerNormalBaryCoordsAndDist() const;
+		TConstArrayView<TArray<FVector4f>> GetRenderDeformerTangentBaryCoordsAndDist() const;
+		TConstArrayView<TArray<FIntVector3>> GetRenderDeformerSimIndices3D() const;
+		TConstArrayView<TArray<float>> GetRenderDeformerWeight() const;
+		TConstArrayView<float> GetRenderDeformerSkinningBlend() const;
+
 		//~ Render Faces Group
 		// Note: RenderIndices points to the collection arrays, not the pattern arrays
 		int32 GetNumRenderFaces() const;
@@ -90,8 +98,14 @@ namespace UE::Chaos::ClothAsset
 		void Reset();
 
 		/** Initialize from another render pattern. Assumes all indices match between source and target. */
-		void Initialize(const FCollectionClothRenderPatternConstFacade& Other);
+		void Initialize(const FCollectionClothRenderPatternConstFacade& Other, int32 SimVertex3DOffset);
+
+		/** Initialize from another render pattern. Assumes all indices match between source and target. */
+		UE_DEPRECATED(5.4, "Use Initialize with the SimVertex3DOffset instead")
+		void Initialize(const FCollectionClothRenderPatternConstFacade& Other) { Initialize(Other, 0); }
 
+		/** Set the render deformer number of influences for this pattern. */
+		void SetRenderDeformerNumInfluences(int32 NumInfluences);
 		/** Set the render material for this pattern. */
 		void SetRenderMaterialPathName(const FString& PathName);
 
@@ -108,6 +122,12 @@ namespace UE::Chaos::ClothAsset
 		TArrayView<FLinearColor> GetRenderColor();
 		TArrayView<TArray<int32>> GetRenderBoneIndices();
 		TArrayView<TArray<float>> GetRenderBoneWeights();
+		TArrayView<TArray<FVector4f>> GetRenderDeformerPositionBaryCoordsAndDist();
+		TArrayView<TArray<FVector4f>> GetRenderDeformerNormalBaryCoordsAndDist();
+		TArrayView<TArray<FVector4f>> GetRenderDeformerTangentBaryCoordsAndDist();
+		TArrayView<TArray<FIntVector3>> GetRenderDeformerSimIndices3D();
+		TArrayView<TArray<float>> GetRenderDeformerWeight();
+		TArrayView<float> GetRenderDeformerSkinningBlend();
 
 		//~ Render Faces Group
 		/** Grow or shrink the space reserved for render faces for this pattern within the cloth collection and return its start index. */
diff --git a/Engine/Plugins/Experimental/ChaosClothAsset/Source/ChaosClothAssetEngine/Private/ChaosClothAsset/ClothAssetBuilder.cpp b/Engine/Plugins/Experimental/ChaosClothAsset/Source/ChaosClothAssetEngine/Private/ChaosClothAsset/ClothAssetBuilder.cpp
--- a/Engine/Plugins/Experimental/ChaosClothAsset/Source/ChaosClothAssetEngine/Private/ChaosClothAsset/ClothAssetBuilder.cpp
+++ b/Engine/Plugins/Experimental/ChaosClothAsset/Source/ChaosClothAssetEngine/Private/ChaosClothAsset/ClothAssetBuilder.cpp
@@ -338,40 +338,83 @@ void UChaosClothAsset::FBuilder::BuildLod(FSkeletalMeshLODModel& LODModel, const
 			Section.ClothingData.AssetGuid = ClothAsset.AssetGuid;  // There is only one cloth asset,
 			Section.CorrespondClothAssetIndex = 0;       // this one
 
-			TArray<FVector3f> SectionRenderPositions;
-			TArray<FVector3f> SectionRenderNormals;
-			TArray<FVector3f> SectionRenderTangents;
-			SectionRenderPositions.Reserve(NumVertices);
-			SectionRenderNormals.Reserve(NumVertices);
-			SectionRenderTangents.Reserve(NumVertices);
-			for (const FSoftSkinVertex& SoftVert : Section.SoftVertices)
+			const int32 RenderDeformerNumInfluences = RenderPatternFacade.GetRenderDeformerNumInfluences();
+			if (RenderDeformerNumInfluences > 0)
 			{
-				SectionRenderPositions.Add(SoftVert.Position);
-				SectionRenderNormals.Add(SoftVert.TangentZ);
-				SectionRenderTangents.Add(SoftVert.TangentX);
+				TArray<FMeshToMeshVertData>& MeshToMeshVertData = Section.ClothMappingDataLODs[0];
+				MeshToMeshVertData.SetNumUninitialized(RenderDeformerNumInfluences * NumVertices);
+
+				const TConstArrayView<TArray<FVector4f>> RenderDeformerPositionBaryCoordsAndDist = RenderPatternFacade.GetRenderDeformerPositionBaryCoordsAndDist();
+				const TConstArrayView<TArray<FVector4f>> RenderDeformerNormalBaryCoordsAndDist = RenderPatternFacade.GetRenderDeformerNormalBaryCoordsAndDist();
+				const TConstArrayView<TArray<FVector4f>> RenderDeformerTangentBaryCoordsAndDist = RenderPatternFacade.GetRenderDeformerTangentBaryCoordsAndDist();
+				const TConstArrayView<TArray<FIntVector3>> RenderDeformerSimIndices3D = RenderPatternFacade.GetRenderDeformerSimIndices3D();
+				const TConstArrayView<TArray<float>> RenderDeformerWeight = RenderPatternFacade.GetRenderDeformerWeight();
+				const TConstArrayView<float> RenderDeformerSkinningBlend = RenderPatternFacade.GetRenderDeformerSkinningBlend();
+
+				for (int32 Index = 0; Index < NumVertices; ++Index)
+				{
+					const uint16 SkinningBlend = (uint16)(FMath::Clamp(RenderDeformerSkinningBlend[Index], 0.f, 1.f) * (float)TNumericLimits<uint16>::Max());
+
+					check(RenderDeformerPositionBaryCoordsAndDist[Index].Num() == RenderDeformerNumInfluences);
+					check(RenderDeformerNormalBaryCoordsAndDist[Index].Num() == RenderDeformerNumInfluences);
+					check(RenderDeformerTangentBaryCoordsAndDist[Index].Num() == RenderDeformerNumInfluences);
+					check(RenderDeformerSimIndices3D[Index].Num() == RenderDeformerNumInfluences);
+					check(RenderDeformerWeight[Index].Num() == RenderDeformerNumInfluences);
+
+					for (int32 Influence = 0; Influence < RenderDeformerNumInfluences; ++Influence)
+					{
+						const int32 InfluenceIndex = Index * RenderDeformerNumInfluences + Influence;
+
+						MeshToMeshVertData[InfluenceIndex].PositionBaryCoordsAndDist = RenderDeformerPositionBaryCoordsAndDist[Index][Influence];
+						MeshToMeshVertData[InfluenceIndex].NormalBaryCoordsAndDist = RenderDeformerNormalBaryCoordsAndDist[Index][Influence];
+						MeshToMeshVertData[InfluenceIndex].TangentBaryCoordsAndDist =  RenderDeformerTangentBaryCoordsAndDist[Index][Influence];
+						MeshToMeshVertData[InfluenceIndex].Weight = RenderDeformerWeight[Index][Influence];
+						const FIntVector3& SimIndices3D = RenderDeformerSimIndices3D[Index][Influence];
+						MeshToMeshVertData[InfluenceIndex].SourceMeshVertIndices[0] = (uint16)SimIndices3D[0];
+						MeshToMeshVertData[InfluenceIndex].SourceMeshVertIndices[1] = (uint16)SimIndices3D[1];
+						MeshToMeshVertData[InfluenceIndex].SourceMeshVertIndices[2] = (uint16)SimIndices3D[2];
+						MeshToMeshVertData[InfluenceIndex].SourceMeshVertIndices[3] = SkinningBlend;
+					}
+				}
 			}
-			TArray<uint32> SectionRenderIndices;
-			SectionRenderIndices.Reserve(NumIndices);
-			const TArrayView<uint32> SectionIndexBuffer(LODModel.IndexBuffer.GetData() + Section.BaseIndex, NumIndices);
-			for (const uint32 LodModelVertIndex : SectionIndexBuffer)
+			else
 			{
-				SectionRenderIndices.Add(LodModelVertIndex - Section.BaseVertexIndex);
-			}
+				// No mapping found in the collection, create them from scratch
+				TArray<FVector3f> SectionRenderPositions;
+				TArray<FVector3f> SectionRenderNormals;
+				TArray<FVector3f> SectionRenderTangents;
+				SectionRenderPositions.Reserve(NumVertices);
+				SectionRenderNormals.Reserve(NumVertices);
+				SectionRenderTangents.Reserve(NumVertices);
+				for (const FSoftSkinVertex& SoftVert : Section.SoftVertices)
+				{
+					SectionRenderPositions.Add(SoftVert.Position);
+					SectionRenderNormals.Add(SoftVert.TangentZ);
+					SectionRenderTangents.Add(SoftVert.TangentX);
+				}
+				TArray<uint32> SectionRenderIndices;
+				SectionRenderIndices.Reserve(NumIndices);
+				const TArrayView<uint32> SectionIndexBuffer(LODModel.IndexBuffer.GetData() + Section.BaseIndex, NumIndices);
+				for (const uint32 LodModelVertIndex : SectionIndexBuffer)
+				{
+					SectionRenderIndices.Add(LodModelVertIndex - Section.BaseVertexIndex);
+				}
 
-			const ClothingMeshUtils::ClothMeshDesc TargetMesh(
-				SectionRenderPositions,
-				SectionRenderNormals,
-				SectionRenderTangents,
-				SectionRenderIndices);
-
-			ClothingMeshUtils::GenerateMeshToMeshVertData(
-				Section.ClothMappingDataLODs[0],
-				TargetMesh,
-				SourceMesh,
-				&MaxDistances,
-				ClothAsset.bSmoothTransition,
-				ClothAsset.bUseMultipleInfluences,
-				ClothAsset.SkinningKernelRadius);
+				const ClothingMeshUtils::ClothMeshDesc TargetMesh(
+					SectionRenderPositions,
+					SectionRenderNormals,
+					SectionRenderTangents,
+					SectionRenderIndices);
+
+				ClothingMeshUtils::GenerateMeshToMeshVertData(
+					Section.ClothMappingDataLODs[0],
+					TargetMesh,
+					SourceMesh,
+					&MaxDistances,
+					ClothAsset.bSmoothTransition,
+					ClothAsset.bUseMultipleInfluences,
+					ClothAsset.SkinningKernelRadius);
+			}
 		}
 
 		// Compute the overlapping vertices map (inspired from MeshUtilities::BuildSkeletalMesh)
diff --git a/Engine/Plugins/Experimental/ChaosClothAssetEditor/Source/ChaosClothAssetDataflowNodes/Private/ChaosClothAsset/ChaosClothAssetDataflowNodesModule.cpp b/Engine/Plugins/Experimental/ChaosClothAssetEditor/Source/ChaosClothAssetDataflowNodes/Private/ChaosClothAsset/ChaosClothAssetDataflowNodesModule.cpp
--- a/Engine/Plugins/Experimental/ChaosClothAssetEditor/Source/ChaosClothAssetDataflowNodes/Private/ChaosClothAsset/ChaosClothAssetDataflowNodesModule.cpp
+++ b/Engine/Plugins/Experimental/ChaosClothAssetEditor/Source/ChaosClothAssetDataflowNodes/Private/ChaosClothAsset/ChaosClothAssetDataflowNodesModule.cpp
@@ -12,6 +12,7 @@
 #include "ChaosClothAsset/RemeshNode.h"
 #include "ChaosClothAsset/ImportNode.h"
 #include "ChaosClothAsset/MergeClothCollectionsNode.h"
+#include "ChaosClothAsset/ProxyDeformerNode.h"
 #include "ChaosClothAsset/ReverseNormalsNode.h"
 #include "ChaosClothAsset/SelectionNode.h"
 #include "ChaosClothAsset/SelectionToIntMapNode.h"
@@ -20,12 +21,11 @@
 #include "ChaosClothAsset/SimulationAerodynamicsConfigNode.h"
 #include "ChaosClothAsset/SimulationAnimDriveConfigNode.h"
 #include "ChaosClothAsset/SimulationBackstopConfigNode.h"
+#include "ChaosClothAsset/SimulationBendingConfigNode.h"
 #include "ChaosClothAsset/SimulationCollisionConfigNode.h"
 #include "ChaosClothAsset/SimulationDampingConfigNode.h"
 #include "ChaosClothAsset/SimulationDefaultConfigNode.h"
 #include "ChaosClothAsset/SimulationGravityConfigNode.h"
-#include "ChaosClothAsset/SimulationBendingConfigNode.h"
-#include "ChaosClothAsset/SimulationStretchConfigNode.h"
 #include "ChaosClothAsset/SimulationLongRangeAttachmentConfigNode.h"
 #include "ChaosClothAsset/SimulationMassConfigNode.h"
 #include "ChaosClothAsset/SimulationMaxDistanceConfigNode.h"
@@ -37,6 +37,7 @@
 #include "ChaosClothAsset/SimulationSelfCollisionConfigNode.h"
 #include "ChaosClothAsset/SimulationSelfCollisionSpheresConfigNode.h"
 #include "ChaosClothAsset/SimulationSolverConfigNode.h"
+#include "ChaosClothAsset/SimulationStretchConfigNode.h"
 #include "ChaosClothAsset/SimulationVelocityScaleConfigNode.h"
 #include "ChaosClothAsset/SimulationXPBDAreaSpringConfigNode.h"
 #include "ChaosClothAsset/SimulationXPBDAnisoBendingConfigNode.h"
@@ -73,6 +74,7 @@ namespace UE::Chaos::ClothAsset
 			DATAFLOW_NODE_REGISTER_CREATION_FACTORY(FChaosClothAssetDeleteElementNode);
 			DATAFLOW_NODE_REGISTER_CREATION_FACTORY(FChaosClothAssetImportNode);
 			DATAFLOW_NODE_REGISTER_CREATION_FACTORY(FChaosClothAssetMergeClothCollectionsNode);
+			DATAFLOW_NODE_REGISTER_CREATION_FACTORY(FChaosClothAssetProxyDeformerNode);
 			DATAFLOW_NODE_REGISTER_CREATION_FACTORY(FChaosClothAssetRemeshNode);
 			DATAFLOW_NODE_REGISTER_CREATION_FACTORY(FChaosClothAssetReverseNormalsNode);
 			DATAFLOW_NODE_REGISTER_CREATION_FACTORY(FChaosClothAssetSelectionNode);
diff --git a/Engine/Plugins/Experimental/ChaosClothAssetEditor/Source/ChaosClothAssetDataflowNodes/Private/ChaosClothAsset/ProxyDeformerNode.cpp b/Engine/Plugins/Experimental/ChaosClothAssetEditor/Source/ChaosClothAssetDataflowNodes/Private/ChaosClothAsset/ProxyDeformerNode.cpp
new file mode 100644
--- /dev/null
+++ b/Engine/Plugins/Experimental/ChaosClothAssetEditor/Source/ChaosClothAssetDataflowNodes/Private/ChaosClothAsset/ProxyDeformerNode.cpp
@@ -0,0 +1,190 @@
+// Copyright Epic Games, Inc. All Rights Reserved.
+
+#include "ChaosClothAsset/ProxyDeformerNode.h"
+#include "ChaosClothAsset/ClothDataflowTools.h"
+#include "ChaosClothAsset/CollectionClothFacade.h"
+#include "ChaosClothAsset/ClothCollectionAttribute.h"
+#include "Dataflow/DataflowInputOutput.h"
+#include "Utils/ClothingMeshUtils.h"
+#include "PointWeightMap.h"
+
+#include UE_INLINE_GENERATED_CPP_BY_NAME(ProxyDeformerNode)
+
+#define LOCTEXT_NAMESPACE "ChaosClothAssetProxyDeformerNode"
+
+namespace UE::Chaos::ClothAsset::Private
+{
+	struct FDeformerMappingDataGenerator
+	{
+		TConstArrayView<FVector3f> SimPositions;
+		TConstArrayView<FIntVector3> SimIndices;
+		TConstArrayView<float> MaxDistances;
+		TConstArrayView<FVector3f> RenderPositions;
+		TConstArrayView<FVector3f> RenderNormals;
+		TConstArrayView<FIntVector3> RenderIndices;
+		TArrayView<TArray<FVector4f>> RenderDeformerPositionBaryCoordsAndDist;
+		TArrayView<TArray<FVector4f>> RenderDeformerNormalBaryCoordsAndDist;
+		TArrayView<TArray<FVector4f>> RenderDeformerTangentBaryCoordsAndDist;
+		TArrayView<TArray<FIntVector3>> RenderDeformerSimIndices3D;
+		TArrayView<TArray<float>> RenderDeformerWeight;
+		TArrayView<float> RenderDeformerSkinningBlend;
+
+		int32 Generate(bool bUseSmoothTransition, bool bUseMultipleInfluences, float InfluenceRadius)
+		{
+			check(RenderPositions.Num() == RenderNormals.Num());
+			check(RenderPositions.Num() == RenderDeformerPositionBaryCoordsAndDist.Num());
+			check(RenderPositions.Num() == RenderDeformerNormalBaryCoordsAndDist.Num());
+			check(RenderPositions.Num() == RenderDeformerTangentBaryCoordsAndDist.Num());
+			check(RenderPositions.Num() == RenderDeformerSimIndices3D.Num());
+			check(RenderPositions.Num() == RenderDeformerWeight.Num());
+			check(RenderPositions.Num() == RenderDeformerSkinningBlend.Num());
+
+			if (!ensureMsgf(SimPositions.Num() <= (int32)TNumericLimits<uint16>::Max() + 1, TEXT("FMeshToMeshVertData data is limited to 16bit unsigned int indexes (65536 indices max).")))
+			{
+				return 0;
+			}
+
+			TArray<uint32> ScalarSimIndices;
+			ScalarSimIndices.Reserve(SimIndices.Num() * 3);
+			for (const FIntVector3& SimIndex : SimIndices)
+			{
+				ScalarSimIndices.Add(SimIndex[0]);
+				ScalarSimIndices.Add(SimIndex[1]);
+				ScalarSimIndices.Add(SimIndex[2]);
+			}
+			TArray<uint32> ScalarRenderIndices;
+			ScalarRenderIndices.Reserve(RenderIndices.Num() * 3);
+			for (const FIntVector3& RenderIndex : RenderIndices)
+			{
+				ScalarRenderIndices.Add(RenderIndex[0]);
+				ScalarRenderIndices.Add(RenderIndex[1]);
+				ScalarRenderIndices.Add(RenderIndex[2]);
+			}
+
+			const ClothingMeshUtils::ClothMeshDesc SimMeshDesc(SimPositions, ScalarSimIndices);
+			const ClothingMeshUtils::ClothMeshDesc RenderMeshDesc(RenderPositions, RenderNormals, ScalarRenderIndices);
+
+			FPointWeightMap PointWeightMap = (MaxDistances.Num() == SimPositions.Num()) ? FPointWeightMap(MaxDistances) : FPointWeightMap(SimPositions.Num(), 1.f);
+
+			TArray<FMeshToMeshVertData> MeshToMeshVertData;
+
+			ClothingMeshUtils::GenerateMeshToMeshVertData(
+				MeshToMeshVertData,
+				RenderMeshDesc,
+				SimMeshDesc,
+				&PointWeightMap,
+				bUseSmoothTransition,
+				bUseMultipleInfluences,
+				InfluenceRadius);
+
+			const int32 NumInfluences = MeshToMeshVertData.Num() / RenderPositions.Num();
+			check(MeshToMeshVertData.Num() == RenderPositions.Num() * NumInfluences);  // Check modulo
+			check((!bUseMultipleInfluences && NumInfluences == 1) || (bUseMultipleInfluences && NumInfluences > 1));
+
+			for (int32 Index = 0; Index < RenderPositions.Num(); ++Index)
+			{
+				RenderDeformerPositionBaryCoordsAndDist[Index].SetNum(NumInfluences);
+				RenderDeformerNormalBaryCoordsAndDist[Index].SetNum(NumInfluences);
+				RenderDeformerTangentBaryCoordsAndDist[Index].SetNum(NumInfluences);
+				RenderDeformerSimIndices3D[Index].SetNum(NumInfluences);
+				RenderDeformerWeight[Index].SetNum(NumInfluences);
+
+				RenderDeformerSkinningBlend[Index] = 0.f;
+
+				for (int32 Influence = 0; Influence < NumInfluences; ++Influence)
+				{
+					const FMeshToMeshVertData& MeshToMeshVertDatum = MeshToMeshVertData[Index * NumInfluences + Influence];
+
+					RenderDeformerPositionBaryCoordsAndDist[Index][Influence] = MeshToMeshVertDatum.PositionBaryCoordsAndDist;
+					RenderDeformerNormalBaryCoordsAndDist[Index][Influence] = MeshToMeshVertDatum.NormalBaryCoordsAndDist;
+					RenderDeformerTangentBaryCoordsAndDist[Index][Influence] = MeshToMeshVertDatum.TangentBaryCoordsAndDist;
+					RenderDeformerSimIndices3D[Index][Influence] = FIntVector3(
+						MeshToMeshVertDatum.SourceMeshVertIndices[0],
+						MeshToMeshVertDatum.SourceMeshVertIndices[1],
+						MeshToMeshVertDatum.SourceMeshVertIndices[2]);
+					RenderDeformerWeight[Index][Influence] = MeshToMeshVertDatum.Weight;
+
+					RenderDeformerSkinningBlend[Index] += MeshToMeshVertDatum.Weight * (float)MeshToMeshVertDatum.SourceMeshVertIndices[3] / (float)TNumericLimits<uint16>::Max();
+				}
+			}
+			return NumInfluences;
+		}
+	};
+}
+
+FChaosClothAssetProxyDeformerNode::FChaosClothAssetProxyDeformerNode(const Dataflow::FNodeParameters& InParam, FGuid InGuid)
+	: FDataflowNode(InParam, InGuid)
+{
+	using namespace UE::Chaos::ClothAsset;
+	MaxDistance.WeightMap = FString();  // An empty weight map is an accepted input, but a non existing one isn't
+	SkinningBlendName = ClothCollectionAttribute::RenderDeformerSkinningBlend.ToString();
+
+	RegisterInputConnection(&Collection);
+	RegisterInputConnection(&MaxDistance.WeightMap, GET_MEMBER_NAME_CHECKED(FChaosClothAssetWeightedValueNonAnimatableNoLowHighRange, WeightMap));
+	RegisterOutputConnection(&Collection, &Collection);
+	RegisterOutputConnection(&SkinningBlendName);
+}
+
+void FChaosClothAssetProxyDeformerNode::Evaluate(Dataflow::FContext& Context, const FDataflowOutput* Out) const
+{
+	if (Out->IsA<FManagedArrayCollection>(&Collection))
+	{
+		using namespace UE::Chaos::ClothAsset;
+
+		// Update the weight map override
+		MaxDistance.WeightMap_Override = GetValue<FString>(Context, &MaxDistance.WeightMap, UE::Chaos::ClothAsset::FWeightMapTools::NotOverridden);
+
+		// Evaluate in collection
+		FManagedArrayCollection InCollection = GetValue<FManagedArrayCollection>(Context, &Collection);
+		const TSharedRef<FManagedArrayCollection> ClothCollection = MakeShared<FManagedArrayCollection>(MoveTemp(InCollection));
+
+		// Always check for a valid cloth collection/facade to avoid processing non cloth collections
+		FCollectionClothFacade ClothFacade(ClothCollection);
+		if (ClothFacade.IsValid())
+		{
+			// Retrieve the MaxDistance weight map name
+			FName MaxDistanceWeightMapName = FName(*GetValue<FString>(Context, &MaxDistance.WeightMap));
+			if (MaxDistanceWeightMapName != NAME_None && !ClothFacade.HasWeightMap(MaxDistanceWeightMapName))
+			{
+				FClothDataflowTools::LogAndToastWarning(*this,
+					LOCTEXT("ClothFacadeHasWeightMapHeadline", "Unknown MaxDistance weight map."),
+					LOCTEXT("ClothFacadeHasWeightMapDetails", "The specified MaxDistance weight map does't exist within the input Cloth Collection."));
+
+				MaxDistanceWeightMapName = NAME_None;
+			}
+
+			// Add the optional render deformer schema
+			if (!ClothFacade.IsValid(EClothCollectionOptionalSchemas::RenderDeformer))
+			{
+				ClothFacade.DefineSchema(EClothCollectionOptionalSchemas::RenderDeformer);
+			}
+
+			// Create the render weight map for storing the skinning blend weights
+			Private::FDeformerMappingDataGenerator DeformerMappingDataGenerator;
+			DeformerMappingDataGenerator.SimPositions = ClothFacade.GetSimPosition3D();
+			DeformerMappingDataGenerator.SimIndices = ClothFacade.GetSimIndices3D();
+			DeformerMappingDataGenerator.MaxDistances = ClothFacade.GetWeightMap(MaxDistanceWeightMapName);
+			DeformerMappingDataGenerator.RenderPositions = ClothFacade.GetRenderPosition();
+			DeformerMappingDataGenerator.RenderNormals = ClothFacade.GetRenderNormal();
+			DeformerMappingDataGenerator.RenderIndices = ClothFacade.GetRenderIndices();
+			DeformerMappingDataGenerator.RenderDeformerPositionBaryCoordsAndDist = ClothFacade.GetRenderDeformerPositionBaryCoordsAndDist();
+			DeformerMappingDataGenerator.RenderDeformerNormalBaryCoordsAndDist = ClothFacade.GetRenderDeformerNormalBaryCoordsAndDist();
+			DeformerMappingDataGenerator.RenderDeformerTangentBaryCoordsAndDist = ClothFacade.GetRenderDeformerTangentBaryCoordsAndDist();
+			DeformerMappingDataGenerator.RenderDeformerSimIndices3D = ClothFacade.GetRenderDeformerSimIndices3D();
+			DeformerMappingDataGenerator.RenderDeformerWeight = ClothFacade.GetRenderDeformerWeight();
+			DeformerMappingDataGenerator.RenderDeformerSkinningBlend = ClothFacade.GetRenderDeformerSkinningBlend();
+
+			const int32 NumInfluences = DeformerMappingDataGenerator.Generate(bUseSmoothTransition, bUseMultipleInfluences, InfluenceRadius);
+
+			for (int32 RenderPatternIndex = 0; RenderPatternIndex < ClothFacade.GetNumRenderPatterns(); ++RenderPatternIndex)
+			{
+				FCollectionClothRenderPatternFacade RenderPatternFacade = ClothFacade.GetRenderPattern(RenderPatternIndex);
+				RenderPatternFacade.SetRenderDeformerNumInfluences(NumInfluences);
+			}
+		}
+
+		SetValue(Context, MoveTemp(*ClothCollection), &Collection);
+	}
+}
+
+#undef LOCTEXT_NAMESPACE
diff --git a/Engine/Plugins/Experimental/ChaosClothAssetEditor/Source/ChaosClothAssetDataflowNodes/Public/ChaosClothAsset/ProxyDeformerNode.h b/Engine/Plugins/Experimental/ChaosClothAssetEditor/Source/ChaosClothAssetDataflowNodes/Public/ChaosClothAsset/ProxyDeformerNode.h
new file mode 100644
--- /dev/null
+++ b/Engine/Plugins/Experimental/ChaosClothAssetEditor/Source/ChaosClothAssetDataflowNodes/Public/ChaosClothAsset/ProxyDeformerNode.h
@@ -0,0 +1,54 @@
+// Copyright Epic Games, Inc. All Rights Reserved.
+
+#pragma once
+
+#include "Dataflow/DataflowNode.h"
+#include "GeometryCollection/ManagedArrayCollection.h"
+#include "ChaosClothAsset/WeightedValue.h"
+#include "ProxyDeformerNode.generated.h"
+
+/** Add the proxy deformer information to this cloth collection's render data. */
+USTRUCT(Meta = (DataflowCloth))
+struct FChaosClothAssetProxyDeformerNode : public FDataflowNode
+{
+	GENERATED_USTRUCT_BODY()
+	DATAFLOW_NODE_DEFINE_INTERNAL(FChaosClothAssetProxyDeformerNode, "ProxyDeformer", "Cloth", "Cloth Simulation Proxy Deformer")
+
+public:
+
+	UPROPERTY(Meta = (Dataflowinput, DataflowOutput, DataflowPassthrough = "Collection"))
+	FManagedArrayCollection Collection;
+
+	/** Whether using multiple simulation mesh triangles to influence the position of the deformed render vertex. */
+	UPROPERTY(EditAnywhere, Category = "ProxyDeformer")
+	bool bUseMultipleInfluences = false;
+
+	/** The radius around the render vertices to look for all simulation mesh triangles influencing it (AKA SkinningKernelRadius). */
+	UPROPERTY(EditAnywhere, Category = "ProxyDeformer", Meta = (EditCondition = "bUseMultipleInfluences"))
+	float InfluenceRadius = 5.f;
+
+	/**
+	 * Whether to create a smoothed _SkinningBlendWeight render weight map to ease the transition between the deformed part and the skinned part of the render mesh.
+	 * When no transition is created there will be a visible step in the rendered triangles around the edge of the kinematic/dynamic transition of the proxy simulation mesh.
+	 * The _SkinningBlendWeight render weight map is created regardless of the transition being created smooth or not, and can be later adjusted using the weight map tool.
+	 */
+	UPROPERTY(EditAnywhere, Category = "ProxyDeformer")
+	bool bUseSmoothTransition = true;
+
+	/** The name of the simulation mesh MaxDistance weight map used to determine the kinematic/dynamic points. */
+	UPROPERTY(EditAnywhere, Category = "ProxyDeformer")
+	FChaosClothAssetWeightedValueNonAnimatableNoLowHighRange MaxDistance;
+
+	/**
+	 * The name of the render mesh weight map generated by this node detailing the contribution of the proxy deformer.
+	 * Value ranges between 0 (fully deformed) and 1 (fully skinned).
+	 * The name of this render mesh weight map cannot be changed and is only provided for further tweaking.
+	 */
+	UPROPERTY(VisibleAnywhere, Category = "ProxyDeformer", Meta = (DataflowOutput))
+	FString SkinningBlendName;
+
+	FChaosClothAssetProxyDeformerNode(const Dataflow::FNodeParameters& InParam, FGuid InGuid = FGuid::NewGuid());
+
+private:
+	virtual void Evaluate(Dataflow::FContext& Context, const FDataflowOutput* Out) const override;
+};
diff --git a/Engine/Plugins/Experimental/ChaosClothAssetEditor/Source/ChaosClothAssetEditor/Private/ChaosClothAsset/SClothCollectionOutliner.cpp b/Engine/Plugins/Experimental/ChaosClothAssetEditor/Source/ChaosClothAssetEditor/Private/ChaosClothAsset/SClothCollectionOutliner.cpp
--- a/Engine/Plugins/Experimental/ChaosClothAssetEditor/Source/ChaosClothAssetEditor/Private/ChaosClothAsset/SClothCollectionOutliner.cpp
+++ b/Engine/Plugins/Experimental/ChaosClothAssetEditor/Source/ChaosClothAssetEditor/Private/ChaosClothAsset/SClothCollectionOutliner.cpp
@@ -312,6 +312,12 @@ namespace ClothCollectionOutlinerHelpers
 		case FManagedArrayCollection::EArrayType::FUInt8Type:
 			ValueAsString = AttributeValueToString<uint8>(ClothCollection, AttributeName, GroupName, AttributeArrayIndex);
 			break;
+		case FManagedArrayCollection::EArrayType::FIntVector3ArrayType:
+			ValueAsString = AttributeValueToString<TArray<FIntVector3>>(ClothCollection, AttributeName, GroupName, AttributeArrayIndex);
+			break;
+		case FManagedArrayCollection::EArrayType::FVector4fArrayType:
+			ValueAsString = AttributeValueToString<TArray<FVector4f>>(ClothCollection, AttributeName, GroupName, AttributeArrayIndex);
+			break;
 		default:
 			ensure(false);
 			ValueAsString = "(Unknown Data Type)";
diff --git a/Engine/Source/Runtime/Experimental/Chaos/Public/GeometryCollection/ManagedArray.h b/Engine/Source/Runtime/Experimental/Chaos/Public/GeometryCollection/ManagedArray.h
--- a/Engine/Source/Runtime/Experimental/Chaos/Public/GeometryCollection/ManagedArray.h
+++ b/Engine/Source/Runtime/Experimental/Chaos/Public/GeometryCollection/ManagedArray.h
@@ -1681,3 +1681,77 @@ public:
 		}
 	}
 };
+
+template<>
+class TManagedArray<TArray<FIntVector3>> : public TManagedArrayBase<TArray<FIntVector3>>
+{
+public:
+	using TManagedArrayBase<TArray<FIntVector3>>::Num;
+
+	TManagedArray() = default;
+
+	TManagedArray(const TArray<TArray<FIntVector3>>& Other)
+		: TManagedArrayBase<TArray<FIntVector3>>(Other)
+	{}
+
+	TManagedArray(const TManagedArray<TArray<FIntVector3>>& Other) = delete;
+	TManagedArray(TManagedArray<TArray<FIntVector3>>&& Other) = default;
+	TManagedArray(TArray<TArray<FIntVector3>>&& Other)
+		: TManagedArrayBase<TArray<FIntVector3>>(MoveTemp(Other))
+	{}
+	
+	TManagedArray& operator=(TManagedArray<TArray<FIntVector3>>&& Other) = default;
+
+	virtual ~TManagedArray() override = default;
+
+	virtual void Reindex(const TArray<int32>& Offsets, const int32& FinalSize, const TArray<int32>& SortedDeletionList, const TSet<int32>& DeletionSet) override
+	{
+		UE_LOG(UManagedArrayLogging, Log, TEXT("TManagedArray<FIntVector>[%p]::Reindex()"), this);
+		const int32 ArraySize = Num();
+		const int32 MaskSize = Offsets.Num();
+		for (int32 Index = 0; Index < ArraySize; ++Index)
+		{
+			const TArray<FIntVector3>& RemapValArray = this->operator[](Index);
+			for (int32 ArrayIndex = 0; ArrayIndex < RemapValArray.Num(); ++ArrayIndex)
+			{
+				const FIntVector3& RemapVal = RemapValArray[ArrayIndex];
+				for (int32 i = 0; i < FIntVector3::Num(); ++i)
+				{
+					if (0 <= RemapVal[i])
+					{
+						ensure(RemapVal[i] < MaskSize);
+						if (DeletionSet.Contains(this->operator[](Index)[ArrayIndex][i]))
+						{
+							this->operator[](Index)[ArrayIndex][i] = INDEX_NONE;
+						}
+						else
+						{
+							this->operator[](Index)[ArrayIndex][i] -= Offsets[RemapVal[i]];
+						}
+						ensure(-1 <= this->operator[](Index)[ArrayIndex][i] && this->operator[](Index)[ArrayIndex][i] <= FinalSize);
+					}
+				}
+			}
+		}
+	}
+
+	virtual void ReindexFromLookup(const TArray<int32>& InverseNewOrder) override
+	{
+		const int32 ArraySize = Num();
+		for (int32 Index = 0; Index < ArraySize; ++Index)
+		{
+			TArray<FIntVector3>& RemapValArray = this->operator[](Index);
+			for (int32 ArrayIndex = 0; ArrayIndex < RemapValArray.Num(); ++ArrayIndex)
+			{
+				FIntVector3& RemapVal = RemapValArray[ArrayIndex];
+				for (int32 i = 0; i < FIntVector3::Num(); ++i)
+				{
+					if (RemapVal[i] >= 0)
+					{
+						RemapVal[i] = InverseNewOrder[RemapVal[i]];
+					}
+				}
+			}
+		}
+	}
+};
diff --git a/Engine/Source/Runtime/Experimental/Chaos/Public/GeometryCollection/ManagedArrayTypeValues.inl b/Engine/Source/Runtime/Experimental/Chaos/Public/GeometryCollection/ManagedArrayTypeValues.inl
--- a/Engine/Source/Runtime/Experimental/Chaos/Public/GeometryCollection/ManagedArrayTypeValues.inl
+++ b/Engine/Source/Runtime/Experimental/Chaos/Public/GeometryCollection/ManagedArrayTypeValues.inl
@@ -54,6 +54,8 @@ MANAGED_ARRAY_TYPE(TUniquePtr<Chaos::FPBDRigidParticle>, TPBDRigidParticle3fUniq
 MANAGED_ARRAY_TYPE(Chaos::FImplicitObjectPtr, FImplicitObjectRefCountedPtr)
 MANAGED_ARRAY_TYPE(Chaos::FConvexPtr, FConvexRefCountedPtr)
 MANAGED_ARRAY_TYPE(FTransform3f, Transform3f)
+MANAGED_ARRAY_TYPE(TArray<FIntVector3>, IntVector3Array)
+MANAGED_ARRAY_TYPE(TArray<FVector4f>, Vector4fArray)
 
 // NOTE: new types must be added at the bottom to keep serialization from breaking
 
