Chaos Cloth Asset - Added missing dependencies on the cloth RenderDeformerSimIndices3D.

diff --git a/Engine/Plugins/Experimental/ChaosClothAsset/Source/ChaosClothAsset/Private/ChaosClothAsset/ClothCollection.cpp b/Engine/Plugins/Experimental/ChaosClothAsset/Source/ChaosClothAsset/Private/ChaosClothAsset/ClothCollection.cpp
--- a/Engine/Plugins/Experimental/ChaosClothAsset/Source/ChaosClothAsset/Private/ChaosClothAsset/ClothCollection.cpp
+++ b/Engine/Plugins/Experimental/ChaosClothAsset/Source/ChaosClothAsset/Private/ChaosClothAsset/ClothCollection.cpp
@@ -373,7 +373,7 @@ namespace UE::Chaos::ClothAsset
 		RenderTangentV = &ManagedArrayCollection->AddAttribute<FVector3f>(ClothCollectionAttribute::RenderTangentV, ClothCollectionGroup::RenderVertices);
 		RenderUVs = &ManagedArrayCollection->AddAttribute<TArray<FVector2f>>(ClothCollectionAttribute::RenderUVs, ClothCollectionGroup::RenderVertices);
 		RenderColor = &ManagedArrayCollection->AddAttribute<FLinearColor>(ClothCollectionAttribute::RenderColor, ClothCollectionGroup::RenderVertices);
-		RenderBoneIndices = &ManagedArrayCollection->AddAttribute<TArray<int32>>(ClothCollectionAttribute::RenderBoneIndices, ClothCollectionGroup::RenderVertices, SimVertices3DDependency);
+		RenderBoneIndices = &ManagedArrayCollection->AddAttribute<TArray<int32>>(ClothCollectionAttribute::RenderBoneIndices, ClothCollectionGroup::RenderVertices);
 		RenderBoneWeights = &ManagedArrayCollection->AddAttribute<TArray<float>>(ClothCollectionAttribute::RenderBoneWeights, ClothCollectionGroup::RenderVertices);
 
 		if (EnumHasAnyFlags(OptionalSchemas, EClothCollectionOptionalSchemas::RenderDeformer))
@@ -385,7 +385,7 @@ namespace UE::Chaos::ClothAsset
 			RenderDeformerPositionBaryCoordsAndDist = &ManagedArrayCollection->AddAttribute<TArray<FVector4f>>(ClothCollectionAttribute::RenderDeformerPositionBaryCoordsAndDist, ClothCollectionGroup::RenderVertices);
 			RenderDeformerNormalBaryCoordsAndDist = &ManagedArrayCollection->AddAttribute<TArray<FVector4f>>(ClothCollectionAttribute::RenderDeformerNormalBaryCoordsAndDist, ClothCollectionGroup::RenderVertices);
 			RenderDeformerTangentBaryCoordsAndDist = &ManagedArrayCollection->AddAttribute<TArray<FVector4f>>(ClothCollectionAttribute::RenderDeformerTangentBaryCoordsAndDist, ClothCollectionGroup::RenderVertices);
-			RenderDeformerSimIndices3D = &ManagedArrayCollection->AddAttribute<TArray<FIntVector3>>(ClothCollectionAttribute::RenderDeformerSimIndices3D, ClothCollectionGroup::RenderVertices);
+			RenderDeformerSimIndices3D = &ManagedArrayCollection->AddAttribute<TArray<FIntVector3>>(ClothCollectionAttribute::RenderDeformerSimIndices3D, ClothCollectionGroup::RenderVertices, SimVertices3DDependency);
 			RenderDeformerWeight = &ManagedArrayCollection->AddAttribute<TArray<float>>(ClothCollectionAttribute::RenderDeformerWeight, ClothCollectionGroup::RenderVertices);
 			RenderDeformerSkinningBlend = &ManagedArrayCollection->AddAttribute<float>(ClothCollectionAttribute::RenderDeformerSkinningBlend, ClothCollectionGroup::RenderVertices);
 		}
