Cloth Asset: use the gravity override from the cloth config.

diff --git a/Engine/Plugins/Experimental/ChaosClothAsset/Source/ChaosClothAssetEngine/Private/ChaosClothAsset/ClothSimulationProxy.cpp b/Engine/Plugins/Experimental/ChaosClothAsset/Source/ChaosClothAssetEngine/Private/ChaosClothAsset/ClothSimulationProxy.cpp
--- a/Engine/Plugins/Experimental/ChaosClothAsset/Source/ChaosClothAssetEngine/Private/ChaosClothAsset/ClothSimulationProxy.cpp
+++ b/Engine/Plugins/Experimental/ChaosClothAsset/Source/ChaosClothAssetEngine/Private/ChaosClothAsset/ClothSimulationProxy.cpp
@@ -228,6 +228,7 @@ PRAGMA_ENABLE_DEPRECATION_WARNINGS
 		Solver->SetLocalSpaceRotation((FQuat)ClothSimulationContext->ComponentTransform.GetRotation());
 		Solver->SetWindVelocity(ClothSimulationContext->WindVelocity);
 		Solver->SetGravity(ClothSimulationContext->WorldGravity);
+		Solver->EnableClothGravityOverride(true);
 
 		// Check teleport modes
 		for (const TUniquePtr<FClothingSimulationCloth>& Cloth : Cloths)
