Cloth Asset Builder: don't assert if there is not one material per render pattern in the asset, just log a warning.
Even though each pattern has a material in the ClothCollection, the user can still remove materials from the asset after the graph evaluates, which can cause this assert to fire when the asset is subsequently loaded.

diff --git a/Engine/Plugins/Experimental/ChaosClothAsset/Source/ChaosClothAssetEngine/Private/ChaosClothAsset/ClothAssetBuilder.cpp b/Engine/Plugins/Experimental/ChaosClothAsset/Source/ChaosClothAssetEngine/Private/ChaosClothAsset/ClothAssetBuilder.cpp
--- a/Engine/Plugins/Experimental/ChaosClothAsset/Source/ChaosClothAssetEngine/Private/ChaosClothAsset/ClothAssetBuilder.cpp
+++ b/Engine/Plugins/Experimental/ChaosClothAsset/Source/ChaosClothAssetEngine/Private/ChaosClothAsset/ClothAssetBuilder.cpp
@@ -6,6 +6,7 @@
 #include "ChaosClothAsset/ClothAsset.h"
 #include "ChaosClothAsset/ClothSimulationModel.h"
 #include "ChaosClothAsset/CollectionClothFacade.h"
+#include "ChaosClothAsset/ClothAssetPrivate.h"
 #include "Chaos/CollectionPropertyFacade.h"
 #include "Engine/RendererSettings.h"
 #include "Engine/SkeletalMesh.h"
@@ -182,7 +183,12 @@ void UChaosClothAsset::FBuilder::BuildLod(FSkeletalMeshLODModel& LODModel, const
 		Section.OriginalDataSectionIndex = SectionIndex;
 
 		const int32 MaterialIndex = MaterialOffset + SectionIndex;
-		check(MaterialIndex < ClothAsset.GetMaterials().Num());
+		if (!ClothAsset.GetMaterials().IsValidIndex(MaterialIndex))
+		{
+			const FString& MaterialPath = RenderPatternFacade.GetRenderMaterialPathName();
+			UE_LOG(LogChaosClothAsset, Warning, TEXT("Cloth Asset BuildLod: Pattern %d in LOD %d has invalid material index. Material index is %d, asset has %d materials. Expected material path for this pattern: %s"), 
+				SectionIndex, LodIndex, MaterialIndex, ClothAsset.GetMaterials().Num(), *MaterialPath);
+		}
 
 		const int32 NumFaces = RenderPatternFacade.GetNumRenderFaces();
 		const int32 NumIndices = NumFaces * 3;
