Cloth Editor Remesh Node: fix crash when collection has invalid stitches

diff --git a/Engine/Plugins/Experimental/ChaosClothAsset/Source/ChaosClothAsset/Private/ChaosClothAsset/ClothGeometryTools.cpp b/Engine/Plugins/Experimental/ChaosClothAsset/Source/ChaosClothAsset/Private/ChaosClothAsset/ClothGeometryTools.cpp
--- a/Engine/Plugins/Experimental/ChaosClothAsset/Source/ChaosClothAsset/Private/ChaosClothAsset/ClothGeometryTools.cpp
+++ b/Engine/Plugins/Experimental/ChaosClothAsset/Source/ChaosClothAsset/Private/ChaosClothAsset/ClothGeometryTools.cpp
@@ -9,6 +9,7 @@
 #include "DynamicMesh/NonManifoldMappingSupport.h"
 #include "Math/Vector.h"
 #include "Util/IndexUtil.h"
+#include "Algo/RemoveIf.h"
 
 namespace UE::Chaos::ClothAsset
 {
@@ -1129,6 +1130,13 @@ namespace UE::Chaos::ClothAsset
 		
 		TArray<FIntVector2> InputStitches(SeamFacade.GetSeamStitch2DEndIndices());
 
+		// filter out any stitches referencing deleted vertices
+		InputStitches.SetNum(Algo::RemoveIf(InputStitches, [](const FIntVector2& Stitch)
+		{
+			return Stitch[0] == INDEX_NONE || Stitch[1] == INDEX_NONE;
+		}));
+
+
 		while (InputStitches.Num() > 0)
 		{
 			TArray<FIntVector2> Seam;
