Chaos Cloth Asset - Add Interp to the cloth component suspend and disable properties, exposing them to UE.

diff --git a/Engine/Plugins/Experimental/ChaosClothAsset/Source/ChaosClothAssetEngine/Public/ChaosClothAsset/ClothComponent.h b/Engine/Plugins/Experimental/ChaosClothAsset/Source/ChaosClothAssetEngine/Public/ChaosClothAsset/ClothComponent.h
--- a/Engine/Plugins/Experimental/ChaosClothAsset/Source/ChaosClothAssetEngine/Public/ChaosClothAsset/ClothComponent.h
+++ b/Engine/Plugins/Experimental/ChaosClothAsset/Source/ChaosClothAssetEngine/Public/ChaosClothAsset/ClothComponent.h
@@ -175,11 +175,11 @@ private:
 	uint8 bWaitForParallelTask : 1;
 
 	/** Whether to enable the simulation or use the skinned pose instead. */
-	UPROPERTY()
+	UPROPERTY(Interp, Category = ClothComponent)
 	uint8 bEnableSimulation : 1;
 
 	/** Whether to suspend the simulation and use the last simulated pose. */
-	UPROPERTY()
+	UPROPERTY(Interp, Category = ClothComponent)
 	uint8 bSuspendSimulation : 1;
 
 	/** Whether to use the leader component pose. */
@@ -187,15 +187,15 @@ private:
 	uint8 bBindToLeaderComponent : 1;
 
 	/** Whether to teleport the cloth prior to advancing the simulation. */
-	UPROPERTY()
+	UPROPERTY(Interp, Category = ClothComponent)
 	uint8 bTeleport : 1;
 
 	/** Whether to reset the pose, bTeleport must be true. */
-	UPROPERTY()
+	UPROPERTY(Interp, Category = ClothComponent)
 	uint8 bReset : 1;
 
 	/** Blend amount between the skinned (=0) and the simulated pose (=1). */
-	UPROPERTY()
+	UPROPERTY(Interp, Category = ClothComponent)
 	float BlendWeight = 1.f;
 
 #if WITH_EDITOR
