Chaos Cloth Asset - Fixed crash when copying render pattern with no valid deformer schema.

#tests Ran the reproducible with the fix and without it. Got a crash without it, and no crash with it.

diff --git a/Engine/Plugins/Experimental/ChaosClothAsset/Source/ChaosClothAsset/Private/ChaosClothAsset/CollectionClothRenderPatternFacade.cpp b/Engine/Plugins/Experimental/ChaosClothAsset/Source/ChaosClothAsset/Private/ChaosClothAsset/CollectionClothRenderPatternFacade.cpp
--- a/Engine/Plugins/Experimental/ChaosClothAsset/Source/ChaosClothAsset/Private/ChaosClothAsset/CollectionClothRenderPatternFacade.cpp
+++ b/Engine/Plugins/Experimental/ChaosClothAsset/Source/ChaosClothAsset/Private/ChaosClothAsset/CollectionClothRenderPatternFacade.cpp
@@ -225,12 +225,16 @@ namespace UE::Chaos::ClothAsset
 		FClothCollection::CopyArrayViewData(GetRenderColor(), Other.GetRenderColor());
 		FClothCollection::CopyArrayViewData(GetRenderBoneIndices(), Other.GetRenderBoneIndices());
 		FClothCollection::CopyArrayViewData(GetRenderBoneWeights(), Other.GetRenderBoneWeights());
-		FClothCollection::CopyArrayViewData(GetRenderDeformerPositionBaryCoordsAndDist(), Other.GetRenderDeformerPositionBaryCoordsAndDist());
-		FClothCollection::CopyArrayViewData(GetRenderDeformerNormalBaryCoordsAndDist(), Other.GetRenderDeformerNormalBaryCoordsAndDist());
-		FClothCollection::CopyArrayViewData(GetRenderDeformerTangentBaryCoordsAndDist(), Other.GetRenderDeformerTangentBaryCoordsAndDist());
-		FClothCollection::CopyArrayViewDataAndApplyOffset(GetRenderDeformerSimIndices3D(), Other.GetRenderDeformerSimIndices3D(), FIntVector3(SimVertex3DOffset));
-		FClothCollection::CopyArrayViewData(GetRenderDeformerWeight(), Other.GetRenderDeformerWeight());
-		FClothCollection::CopyArrayViewData(GetRenderDeformerSkinningBlend(), Other.GetRenderDeformerSkinningBlend());
+
+		if (GetClothCollection()->IsValid(EClothCollectionOptionalSchemas::RenderDeformer))
+		{
+			FClothCollection::CopyArrayViewData(GetRenderDeformerPositionBaryCoordsAndDist(), Other.GetRenderDeformerPositionBaryCoordsAndDist());
+			FClothCollection::CopyArrayViewData(GetRenderDeformerNormalBaryCoordsAndDist(), Other.GetRenderDeformerNormalBaryCoordsAndDist());
+			FClothCollection::CopyArrayViewData(GetRenderDeformerTangentBaryCoordsAndDist(), Other.GetRenderDeformerTangentBaryCoordsAndDist());
+			FClothCollection::CopyArrayViewDataAndApplyOffset(GetRenderDeformerSimIndices3D(), Other.GetRenderDeformerSimIndices3D(), FIntVector3(SimVertex3DOffset));
+			FClothCollection::CopyArrayViewData(GetRenderDeformerWeight(), Other.GetRenderDeformerWeight());
+			FClothCollection::CopyArrayViewData(GetRenderDeformerSkinningBlend(), Other.GetRenderDeformerSkinningBlend());
+		}
 
 		//~ Render Faces Group
 		const int32 RenderVertexOffset = GetRenderVerticesOffset() - Other.GetRenderVerticesOffset();
