Chaos Cloth Asset:
- Fixed BVH causing issue on the Proxy Deformer filtered sets.
- Fixed a crash when the Proxy Deformer node is using multiple inputs.
- Fixed reindexing of the Proxy Deformer when adding new stitches.

#rnx
#tests Created a more complex deformer filter setup with more faces selected. These changes seem to have fixed all remaining issues.

diff --git a/Engine/Plugins/Experimental/ChaosClothAsset/Source/ChaosClothAsset/Private/ChaosClothAsset/CollectionClothSeamFacade.cpp b/Engine/Plugins/Experimental/ChaosClothAsset/Source/ChaosClothAsset/Private/ChaosClothAsset/CollectionClothSeamFacade.cpp
--- a/Engine/Plugins/Experimental/ChaosClothAsset/Source/ChaosClothAsset/Private/ChaosClothAsset/CollectionClothSeamFacade.cpp
+++ b/Engine/Plugins/Experimental/ChaosClothAsset/Source/ChaosClothAsset/Private/ChaosClothAsset/CollectionClothSeamFacade.cpp
@@ -709,6 +709,17 @@ namespace UE::Chaos::ClothAsset
 			Index3D[2] = WeldingMappedValue(WeldingMap, Index3D[2]);
 		}
 
+		// Weld the deformer's SimIndices3D
+		for (TArray<FIntVector3>& SimIndices3D : Cloth.GetRenderDeformerSimIndices3D())
+		{
+			for (FIntVector3& Index3D : SimIndices3D)
+			{
+				Index3D[0] = WeldingMappedValue(WeldingMap, Index3D[0]);
+				Index3D[1] = WeldingMappedValue(WeldingMap, Index3D[1]);
+				Index3D[2] = WeldingMappedValue(WeldingMap, Index3D[2]);
+			}
+		}
+
 		// Weld maps
 		const TArray<FName> WeightMapNames = Cloth.GetWeightMapNames();
 		for (const FName& WeightMapName : WeightMapNames)
diff --git a/Engine/Plugins/Experimental/ChaosClothAssetEditor/Source/ChaosClothAssetDataflowNodes/Private/ChaosClothAsset/ProxyDeformerNode.cpp b/Engine/Plugins/Experimental/ChaosClothAssetEditor/Source/ChaosClothAssetDataflowNodes/Private/ChaosClothAsset/ProxyDeformerNode.cpp
--- a/Engine/Plugins/Experimental/ChaosClothAssetEditor/Source/ChaosClothAssetDataflowNodes/Private/ChaosClothAsset/ProxyDeformerNode.cpp
+++ b/Engine/Plugins/Experimental/ChaosClothAssetEditor/Source/ChaosClothAssetDataflowNodes/Private/ChaosClothAsset/ProxyDeformerNode.cpp
@@ -429,7 +429,7 @@ TArray<FName> FChaosClothAssetProxyDeformerNode::GetSelectionFilterNames(Dataflo
 
 	for (int32 FilterSetIndex = 1; FilterSetIndex < NumFilterSets; ++FilterSetIndex)
 	{
-		SelectionFilterSets[FilterSetIndex] = FName(*GetValue(Context, &Non0SelectionFilterSets[FilterSetIndex]->StringValue));
+		SelectionFilterSets[FilterSetIndex] = FName(*GetValue(Context, &Non0SelectionFilterSets[FilterSetIndex - 1]->StringValue));
 	}
 	return SelectionFilterSets;
 }
diff --git a/Engine/Source/Runtime/ClothingSystemRuntimeCommon/Private/Utils/ClothingMeshUtils.cpp b/Engine/Source/Runtime/ClothingSystemRuntimeCommon/Private/Utils/ClothingMeshUtils.cpp
--- a/Engine/Source/Runtime/ClothingSystemRuntimeCommon/Private/Utils/ClothingMeshUtils.cpp
+++ b/Engine/Source/Runtime/ClothingSystemRuntimeCommon/Private/Utils/ClothingMeshUtils.cpp
@@ -99,7 +99,7 @@ namespace ClothingMeshUtils
 
 		const int32 NumTris = NumFilteredTriangles ? NumFilteredTriangles : Indices.Num() / 3;
 
-		if (NumTris > MinNumTrianglesForBVHCreation)
+		if (!NumFilteredTriangles && NumTris > MinNumTrianglesForBVHCreation)  // Don't use the BVH when using filtered triangle sets  TODO: Fix this
 		{
 			// This is not thread safe
 			if (!bHasValidBVH)
