Chaos Cloth Asset: seam bug fixes.

Seams were getting corrupted when merging assets. There were two problems:
* SeamStitchLookup was getting the wrong offset applied.
* Since we weren't cleaning up seam stitches, they had  -1 references. Remesh code could handle -1s, but when we merge assets, those -1 would get an offset added to them, so then they would just point to the wrong thing.

Fixes
* Update SeamStitchLookup to apply the correct offset.
* Update CleanupAndCompactMesh to cleanup seam data.
* Added seam validation methods that will ensure that the seam data makes sense if you set a cvar.

diff --git a/Engine/Plugins/Experimental/ChaosClothAsset/Source/ChaosClothAsset/Private/ChaosClothAsset/ClothGeometryTools.cpp b/Engine/Plugins/Experimental/ChaosClothAsset/Source/ChaosClothAsset/Private/ChaosClothAsset/ClothGeometryTools.cpp
--- a/Engine/Plugins/Experimental/ChaosClothAsset/Source/ChaosClothAsset/Private/ChaosClothAsset/ClothGeometryTools.cpp
+++ b/Engine/Plugins/Experimental/ChaosClothAsset/Source/ChaosClothAsset/Private/ChaosClothAsset/ClothGeometryTools.cpp
@@ -1030,7 +1030,6 @@ namespace UE::Chaos::ClothAsset
 			// Clean up any references to vertices that no longer exist.
 			// NOTE: should not need to clean up 2D vertices pointing to INDEX_NONE 3D vertices since this should have
 			// meant the 2D vertex either was unused in the faces, or was associated with an invalid face (it should already be cleaned up).
-			// NOTE: not cleaning up seam references for now as they aren't exported as part of the final mesh/sim data.
 			Cloth.CompactSimVertex2DLookup();
 
 			TArrayView<TArray<int32>> TetherKinematicIndex = Cloth.GetTetherKinematicIndex();
@@ -1038,16 +1037,34 @@ namespace UE::Chaos::ClothAsset
 			const int32 NumVertices = TetherKinematicIndex.Num();
 			for (int32 VertexIdx = 0; VertexIdx < NumVertices; ++VertexIdx)
 			{
-				const int32 NumTethers = TetherKinematicIndex[VertexIdx].Num();
-				for (int32 TetherIdx = NumTethers - 1; TetherIdx >= 0; --TetherIdx)
+				for (int32 TetherIdx = 0; TetherIdx < TetherKinematicIndex[VertexIdx].Num(); )
 				{
 					if (TetherKinematicIndex[VertexIdx][TetherIdx] == INDEX_NONE)
 					{
 						TetherKinematicIndex[VertexIdx].RemoveAtSwap(TetherIdx);
 						TetherReferenceLength[VertexIdx].RemoveAtSwap(TetherIdx);
+						continue;
 					}
+					++TetherIdx;
 				}
 			}
+
+			// Clean up seams. Update stitches that refer to invalid indices.
+			TArray<int32> SeamsToRemove;
+			for (int32 SeamIndex = 0; SeamIndex < Cloth.GetNumSeams(); ++SeamIndex)
+			{
+				FCollectionClothSeamFacade Seam = Cloth.GetSeam(SeamIndex);
+				Seam.CleanupAndCompact();
+				if (Seam.GetNumSeamStitches() == 0)
+				{
+					SeamsToRemove.Add(SeamIndex);
+				}
+			}
+			if (!SeamsToRemove.IsEmpty())
+			{
+				Cloth.RemoveSeams(SeamsToRemove);
+			}
+			Cloth.CompactSeamStitchLookup();
 		}
 
 		TArray<int32> RenderPatternsToRemove;
@@ -1112,6 +1129,12 @@ namespace UE::Chaos::ClothAsset
 		{
 			Cloth.RemoveRenderPatterns(RenderPatternsToRemove);
 		}
+#if DO_ENSURE
+		for (int32 SeamIndex = 0; SeamIndex < Cloth.GetNumSeams(); ++SeamIndex)
+		{
+			Cloth.GetSeam(SeamIndex).ValidateSeam();
+		}
+#endif
 	}
 
 
diff --git a/Engine/Plugins/Experimental/ChaosClothAsset/Source/ChaosClothAsset/Private/ChaosClothAsset/CollectionClothFacade.cpp b/Engine/Plugins/Experimental/ChaosClothAsset/Source/ChaosClothAsset/Private/ChaosClothAsset/CollectionClothFacade.cpp
--- a/Engine/Plugins/Experimental/ChaosClothAsset/Source/ChaosClothAsset/Private/ChaosClothAsset/CollectionClothFacade.cpp
+++ b/Engine/Plugins/Experimental/ChaosClothAsset/Source/ChaosClothAsset/Private/ChaosClothAsset/CollectionClothFacade.cpp
@@ -515,6 +515,13 @@ namespace UE::Chaos::ClothAsset
 
 	void FCollectionClothFacade::Append(const FCollectionClothConstFacade& Other)
 	{
+#if DO_ENSURE
+		for (int32 SeamIndex = 0; SeamIndex < GetNumSeams(); ++SeamIndex)
+		{
+			GetSeam(SeamIndex).ValidateSeam();
+		}
+#endif
+
 		// LODs Group 
 		// Just keep original data.
 
@@ -556,11 +563,13 @@ namespace UE::Chaos::ClothAsset
 		}
 
 		// Seams Group
+		const int32 StartNumSeamStitches = ClothCollection->GetNumElements(ClothCollectionGroup::SeamStitches);
 		const int32 StartNumSeams = GetNumSeams();
 		const int32 OtherNumSeams = Other.GetNumSeams();
 		SetNumSeams(StartNumSeams + OtherNumSeams);
 		for (int32 SeamIndex = 0; SeamIndex < OtherNumSeams; ++SeamIndex)
 		{
+			Other.GetSeam(SeamIndex).ValidateSeam();
 			GetSeam(SeamIndex + StartNumSeams).Initialize(Other.GetSeam(SeamIndex), StartNumSimVertices2D, StartNumSimVertices3D);
 		}
 		
@@ -570,7 +579,14 @@ namespace UE::Chaos::ClothAsset
 			Other.GetSimVertex2DLookup(), StartNumSimVertices2D);
 		FClothCollection::CopyArrayViewDataAndApplyOffset(
 			GetClothCollection()->GetElements(GetClothCollection()->GetSeamStitchLookup()).Right(OtherNumSimVertices3D), 
-			Other.GetSeamStitchLookup(), StartNumSeams);
+			Other.GetSeamStitchLookup(), StartNumSeamStitches);
+
+#if DO_ENSURE
+		for (int32 SeamIndex = 0; SeamIndex < OtherNumSeams; ++SeamIndex)
+		{
+			GetSeam(SeamIndex + StartNumSeams).ValidateSeam();
+		}
+#endif
 
 		// Render Patterns Group
 		const int32 StartNumRenderPatterns = GetNumRenderPatterns();
@@ -681,12 +697,31 @@ namespace UE::Chaos::ClothAsset
 		TArrayView<TArray<int32>> SimVertex2DLookup = GetSimVertex2DLookupPrivate();
 		for (TArray<int32>& VertexLookup : SimVertex2DLookup)
 		{
-			for (int32 Idx = VertexLookup.Num() - 1; Idx >= 0; --Idx)
+			for (int32 Idx = 0; Idx < VertexLookup.Num();)
 			{
 				if (VertexLookup[Idx] == INDEX_NONE)
 				{
 					VertexLookup.RemoveAtSwap(Idx);
+					continue;
+				}
+				++Idx;
+			}
+		}
+	}
+
+	void FCollectionClothFacade::CompactSeamStitchLookup()
+	{
+		TArrayView<TArray<int32>> SeamStitchLookup = GetSeamStitchLookupPrivate();
+		for (TArray<int32>& StitchLookup : SeamStitchLookup)
+		{
+			for (int32 Idx = 0; Idx < StitchLookup.Num();)
+			{
+				if (StitchLookup[Idx] == INDEX_NONE)
+				{
+					StitchLookup.RemoveAtSwap(Idx);
+					continue;
 				}
+				++Idx;
 			}
 		}
 	}
diff --git a/Engine/Plugins/Experimental/ChaosClothAsset/Source/ChaosClothAsset/Private/ChaosClothAsset/CollectionClothSeamFacade.cpp b/Engine/Plugins/Experimental/ChaosClothAsset/Source/ChaosClothAsset/Private/ChaosClothAsset/CollectionClothSeamFacade.cpp
--- a/Engine/Plugins/Experimental/ChaosClothAsset/Source/ChaosClothAsset/Private/ChaosClothAsset/CollectionClothSeamFacade.cpp
+++ b/Engine/Plugins/Experimental/ChaosClothAsset/Source/ChaosClothAsset/Private/ChaosClothAsset/CollectionClothSeamFacade.cpp
@@ -9,6 +9,10 @@
 
 namespace UE::Chaos::ClothAsset
 {
+#if DO_ENSURE
+static bool bEnableSeamChecks = false;
+static FAutoConsoleVariableRef CVarEnableSeamChecks(TEXT("p.ChaosCloth.EnableSeamChecks"), bEnableSeamChecks, TEXT("Enable seam validation checks"));
+#endif
 	namespace Private
 	{
 		//
@@ -462,6 +466,116 @@ namespace UE::Chaos::ClothAsset
 			}
 		}
 
+		// Cleanup methods
+		TArray<TSet<int32>> BuildStitchGroupsForVertex3D(const TArray<int32>& StitchLookup, const TArray<int32>& Vertex2DLookup, const TArrayView<FIntVector2>& AllEnds, const int32 OurIndex, bool& bFoundOtherInvalidStitches)
+		{
+			bFoundOtherInvalidStitches = false;
+			TArray<TSet<int32>> Vertex2DGroups;
+			for (int32 OtherStitchIndex : StitchLookup)
+			{
+				if (OtherStitchIndex == OurIndex || OtherStitchIndex == INDEX_NONE)
+				{
+					// This is us
+					continue;
+				}
+				const FIntVector2& OtherStitchEnds = AllEnds[OtherStitchIndex];
+				if (OtherStitchEnds[0] == INDEX_NONE || OtherStitchEnds[1] == INDEX_NONE)
+				{
+					bFoundOtherInvalidStitches = true;
+					continue;
+				}
+				int32 GroupIndex0 = INDEX_NONE;
+				int32 GroupIndex1 = INDEX_NONE;
+				for (int32 GroupIndex = 0; GroupIndex < Vertex2DGroups.Num(); ++GroupIndex)
+				{
+					if (Vertex2DGroups[GroupIndex].Contains(OtherStitchEnds[0]))
+					{
+						GroupIndex0 = GroupIndex;
+					}
+					if (Vertex2DGroups[GroupIndex].Contains(OtherStitchEnds[1]))
+					{
+						GroupIndex1 = GroupIndex;
+					}
+					else
+					{
+						continue;
+					}
+					if (GroupIndex0 != INDEX_NONE && GroupIndex1 != INDEX_NONE)
+					{
+						break;
+					}
+				}
+
+				if (GroupIndex0 == INDEX_NONE)
+				{
+					if (GroupIndex1 == INDEX_NONE)
+					{
+						// Start a new group
+						Vertex2DGroups.Emplace(TSet<int32>({ OtherStitchEnds[0], OtherStitchEnds[1]}));
+					}
+					else
+					{
+						// Add to Group1
+						Vertex2DGroups[GroupIndex1].Add(OtherStitchEnds[0]);
+					}
+				}
+				else if (GroupIndex1 == INDEX_NONE)
+				{
+					Vertex2DGroups[GroupIndex0].Add(OtherStitchEnds[1]);
+				}
+				else
+				{
+					// Both in different groups. Merge them.
+					Vertex2DGroups[GroupIndex0].Append(Vertex2DGroups[GroupIndex1]);
+					Vertex2DGroups.RemoveAtSwap(GroupIndex1);
+				}
+			}
+
+			// Add all Vertex2D that aren't in one any groups
+			for (const int32 Vertex2D : Vertex2DLookup)
+			{
+				if (Vertex2D == INDEX_NONE)
+				{
+					continue;
+				}
+				bool bFoundInGroup = false;
+				for (const TSet<int32>& Group : Vertex2DGroups)
+				{
+					if (Group.Contains(Vertex2D))
+					{
+						bFoundInGroup = true;
+						break;
+					}
+				}
+				if (!bFoundInGroup)
+				{
+					Vertex2DGroups.Emplace(TSet<int32>({ Vertex2D }));
+				}
+			}
+			return Vertex2DGroups;
+		}
+
+		TArray<FIntVector2> GenerateRemainderStitchesAfterCleanup(const TArray<TSet<int32>>& Vertex2DGroups)
+		{
+			TArray<FIntVector2> Stitches;
+			if (Vertex2DGroups.Num() > 1)
+			{
+				Stitches.Reserve(Vertex2DGroups.Num() - 1);
+				FIntVector2* Stitch = &Stitches.AddDefaulted_GetRef();
+				(*Stitch)[0] = *Vertex2DGroups[0].CreateConstIterator();
+				for (int32 GroupIndex = 1; GroupIndex < Vertex2DGroups.Num(); ++GroupIndex)
+				{
+					(*Stitch)[1] = *Vertex2DGroups[GroupIndex].CreateConstIterator();
+					if (GroupIndex < Vertex2DGroups.Num() - 1)
+					{
+						const int32 NextEnd = (*Stitch)[1];
+						Stitch = &Stitches.AddDefaulted_GetRef();
+						(*Stitch)[0] = NextEnd;
+					}
+				}
+			}
+			return Stitches;
+		}
 	} // namespace Private
 
 	int32 FCollectionClothSeamConstFacade::GetNumSeamStitches() const
@@ -497,6 +611,37 @@ namespace UE::Chaos::ClothAsset
 			ClothCollection->GetSeamStitchEnd(),
 			GetElementIndex());
 	}
+	
+	void FCollectionClothSeamConstFacade::ValidateSeam() const
+	{
+#if DO_ENSURE
+		if (bEnableSeamChecks)
+		{
+			FCollectionClothConstFacade Cloth(ClothCollection);
+			TConstArrayView<FIntVector2> Ends = GetSeamStitch2DEndIndices();
+			TConstArrayView<int32> Index3D = GetSeamStitch3DIndex();
+			TConstArrayView<TArray<int32>> Lookup2D = Cloth.GetSimVertex2DLookup();
+			TConstArrayView<TArray<int32>> StitchLookup = Cloth.GetSeamStitchLookup();
+			TConstArrayView<int32> Lookup3D = Cloth.GetSimVertex3DLookup();
+
+			const int32 StitchOffset = GetSeamStitchesOffset();
+			for (int32 StitchIndex = 0; StitchIndex < GetNumSeamStitches(); ++StitchIndex)
+			{
+				if (Ends[StitchIndex][0] != INDEX_NONE && Ends[StitchIndex][1] != INDEX_NONE)
+				{
+					if (Index3D[StitchIndex] != INDEX_NONE)
+					{
+						ensureAlways(StitchLookup[Index3D[StitchIndex]].Contains(StitchOffset + StitchIndex));
+						ensureAlways(Lookup2D[Index3D[StitchIndex]].Contains(Ends[StitchIndex][0]));
+						ensureAlways(Lookup2D[Index3D[StitchIndex]].Contains(Ends[StitchIndex][1]));
+						ensureAlways(Lookup3D[Ends[StitchIndex][0]] == Index3D[StitchIndex]);
+						ensureAlways(Lookup3D[Ends[StitchIndex][1]] == Index3D[StitchIndex]);
+					}
+				}
+			}
+		}
+#endif
+	}
 
 	FCollectionClothSeamConstFacade::FCollectionClothSeamConstFacade(const TSharedRef<const FClothCollection>& ClothCollection, int32 SeamIndex)
 		: ClothCollection(ClothCollection)
@@ -740,6 +885,8 @@ namespace UE::Chaos::ClothAsset
 		}
 		VerticesToRemove.Sort();
 		GetClothCollection()->RemoveElements(ClothCollectionGroup::SimVertices3D, VerticesToRemove);
+
+		ValidateSeam();
 	}
 
 	void FCollectionClothSeamFacade::Initialize(const FCollectionClothSeamConstFacade& Other, const int32 SimVertex2DOffset, const int32 SimVertex3DOffset)
@@ -749,6 +896,142 @@ namespace UE::Chaos::ClothAsset
 		FClothCollection::CopyArrayViewDataAndApplyOffset(GetSeamStitch3DIndex(), Other.GetSeamStitch3DIndex(), SimVertex3DOffset);
 	}
 
+	void FCollectionClothSeamFacade::CleanupAndCompact()
+	{
+		FCollectionClothFacade Cloth(GetClothCollection());
+		TArrayView<FIntVector2> AllEnds = GetClothCollection()->GetElements(GetClothCollection()->GetSeamStitch2DEndIndices());
+		TArrayView<int32> Index3D = GetSeamStitch3DIndex();
+		TConstArrayView<TArray<int32>> SimVertex2DLookup = Cloth.GetSimVertex2DLookup();
+		TConstArrayView<int32> Lookup3D = Cloth.GetSimVertex3DLookup();
+		TArrayView<TArray<int32>> SeamStitchLookup = Cloth.GetSeamStitchLookupPrivate();
+
+		const int32 StitchOffset = GetSeamStitchesOffset();
+		TArray<int32> StitchesToRemoveGlobal;
+		struct FNewStitches
+		{
+			int32 Vertex3D;
+			TArray<FIntVector2> StitchEnds;
+		};
+		TArray<FNewStitches> NewStitchesToAdd;
+		for (int32 StitchIndex = 0; StitchIndex < GetNumSeamStitches(); ++StitchIndex)
+		{
+			if (Index3D[StitchIndex] == INDEX_NONE)
+			{
+				// This stitch is no longer stitching anything together. It can be removed.
+				StitchesToRemoveGlobal.Add(StitchIndex + StitchOffset);
+				continue;
+			}
+
+			FIntVector2& StitchEnds = AllEnds[StitchIndex + StitchOffset];
+			if (StitchEnds[0] == INDEX_NONE || StitchEnds[1] == INDEX_NONE)
+			{
+				// Try to fix up the stitches that go to this 3D vertex 
+				bool bFoundOtherInvalidStitches = false;
+				TArray<TSet<int32>> Vertex2DGroups = Private::BuildStitchGroupsForVertex3D(SeamStitchLookup[Index3D[StitchIndex]], SimVertex2DLookup[Index3D[StitchIndex]], AllEnds, StitchIndex + StitchOffset, bFoundOtherInvalidStitches);
+
+				if (Vertex2DGroups.Num() < 2)
+				{
+					// All 2D vertices going to this 3D vertex are stitched together. Can discard this stitch.
+					StitchesToRemoveGlobal.Add(StitchIndex + StitchOffset);
+					continue;
+				}
+
+				if (StitchEnds[1] != INDEX_NONE)
+				{
+					// Swap ends just to make next bits a little easier.
+					StitchEnds[0] = StitchEnds[1];
+					StitchEnds[1] = INDEX_NONE;
+				}
+
+				int32 GroupIndex0 = INDEX_NONE;
+				if (StitchEnds[0] != INDEX_NONE)
+				{
+					// Find which group our first end is in.
+					for (int32 GroupIndex = 0; GroupIndex < Vertex2DGroups.Num(); ++GroupIndex)
+					{
+						if (Vertex2DGroups[GroupIndex].Contains(StitchEnds[0]))
+						{
+							GroupIndex0 = GroupIndex;
+							break;
+						}
+					}
+				}
+				else
+				{
+					GroupIndex0 = 0;
+					StitchEnds[0] = *Vertex2DGroups[GroupIndex0].CreateConstIterator();
+				}
+				check(GroupIndex0 != INDEX_NONE);
+				const int32 GroupIndex1 = GroupIndex0 == 0 ? 1 : 0;
+				StitchEnds[1] = *Vertex2DGroups[GroupIndex1].CreateConstIterator();
+
+				if (Vertex2DGroups.Num() > 2 && !bFoundOtherInvalidStitches)
+				{
+					Vertex2DGroups[GroupIndex0].Append(Vertex2DGroups[GroupIndex1]);
+					Vertex2DGroups.RemoveAtSwap(GroupIndex1);
+					// There are more groups to stitch together and there aren't any other invalid stitches that correspond with this 3D index that 
+					// can be used to fix this up.
+					NewStitchesToAdd.Emplace(Index3D[StitchIndex], Private::GenerateRemainderStitchesAfterCleanup(Vertex2DGroups));
+				}
+			}
+		}
+
+		// Determine how many remainder stitches we need to add.
+		int32 NumStitchesToAdd = 0;
+		for (const FNewStitches& NewStitches : NewStitchesToAdd)
+		{
+			NumStitchesToAdd += NewStitches.StitchEnds.Num();
+		}
+
+		// Reuse indices for stitches marked to remove to create new stitches.
+		const int32 OrigNumStitches = GetNumSeamStitches();
+		if (NumStitchesToAdd > StitchesToRemoveGlobal.Num())
+		{
+			SetNumSeamStitches(OrigNumStitches + NumStitchesToAdd - StitchesToRemoveGlobal.Num());
+			// Allocating new stitches can cause the SeamStitches group pointers to change.
+			AllEnds = GetClothCollection()->GetElements(GetClothCollection()->GetSeamStitch2DEndIndices());
+			Index3D = GetSeamStitch3DIndex();
+		}
+
+		int32 NewStitchIndex = OrigNumStitches;
+		for (const FNewStitches& NewStitches : NewStitchesToAdd)
+		{
+			for (const FIntVector2& StitchEnds : NewStitches.StitchEnds)
+			{
+				if (StitchesToRemoveGlobal.Num() > 0)
+				{
+					const int32 RecycledStitchIndex = StitchesToRemoveGlobal.Pop() - StitchOffset;
+					if (Index3D[RecycledStitchIndex] != INDEX_NONE)
+					{
+						// Remove RecycledStitchIndex from this SeamStitchLookup
+						SeamStitchLookup[Index3D[RecycledStitchIndex]].Remove(RecycledStitchIndex);
+					}
+					Index3D[RecycledStitchIndex] = NewStitches.Vertex3D;
+					AllEnds[RecycledStitchIndex + StitchOffset] = StitchEnds;
+					SeamStitchLookup[NewStitches.Vertex3D].Add(RecycledStitchIndex + StitchOffset);
+				}
+				else
+				{
+					Index3D[NewStitchIndex] = NewStitches.Vertex3D;
+					AllEnds[NewStitchIndex + StitchOffset] = StitchEnds;
+					SeamStitchLookup[NewStitches.Vertex3D].Add(NewStitchIndex + StitchOffset);
+					++NewStitchIndex;
+				}
+			}
+		}
+		check(NewStitchIndex == GetNumSeamStitches());
+
+		if (!StitchesToRemoveGlobal.IsEmpty())
+		{
+			GetClothCollection()->RemoveElements(
+				ClothCollectionGroup::SeamStitches,
+				StitchesToRemoveGlobal,
+				GetClothCollection()->GetSeamStitchStart(),
+				GetClothCollection()->GetSeamStitchEnd(),
+				GetElementIndex());
+		}
+	}
+
 	void FCollectionClothSeamFacade::SetNumSeamStitches(int32 NumStitches)
 	{
 		GetClothCollection()->SetNumElements(
diff --git a/Engine/Plugins/Experimental/ChaosClothAsset/Source/ChaosClothAsset/Public/ChaosClothAsset/CollectionClothFacade.h b/Engine/Plugins/Experimental/ChaosClothAsset/Source/ChaosClothAsset/Public/ChaosClothAsset/CollectionClothFacade.h
--- a/Engine/Plugins/Experimental/ChaosClothAsset/Source/ChaosClothAsset/Public/ChaosClothAsset/CollectionClothFacade.h
+++ b/Engine/Plugins/Experimental/ChaosClothAsset/Source/ChaosClothAsset/Public/ChaosClothAsset/CollectionClothFacade.h
@@ -222,6 +222,8 @@ namespace UE::Chaos::ClothAsset
 		void RemoveSimVertices3D(const TArray<int32>& SortedDeletionList);
 		/** Compact SimVertex2DLookup to remove any references to INDEX_NONE that may have been created by deleting 2D vertices. */
 		void CompactSimVertex2DLookup();
+		/** Compact SeamStitchLookup to remove any references to INDEX_NONE that may have been created by deleting stitches. */
+		void CompactSeamStitchLookup();
 
 		//~ Pattern Sim Faces Group
 		/** SetNumSimFaces per pattern within pattern facade. */
diff --git a/Engine/Plugins/Experimental/ChaosClothAsset/Source/ChaosClothAsset/Public/ChaosClothAsset/CollectionClothSeamFacade.h b/Engine/Plugins/Experimental/ChaosClothAsset/Source/ChaosClothAsset/Public/ChaosClothAsset/CollectionClothSeamFacade.h
--- a/Engine/Plugins/Experimental/ChaosClothAsset/Source/ChaosClothAsset/Public/ChaosClothAsset/CollectionClothSeamFacade.h
+++ b/Engine/Plugins/Experimental/ChaosClothAsset/Source/ChaosClothAsset/Public/ChaosClothAsset/CollectionClothSeamFacade.h
@@ -38,6 +38,8 @@ namespace UE::Chaos::ClothAsset
 		/** Return the seam index this facade has been created with. */
 		int32 GetSeamIndex() const { return SeamIndex; }
 
+		void ValidateSeam() const;
+
 	protected:
 		friend class FCollectionClothSeamFacade;  // For other instances access
 		friend class FCollectionClothConstFacade;
@@ -77,6 +79,9 @@ namespace UE::Chaos::ClothAsset
 		/** Initialize from another seam. */
 		void Initialize(const FCollectionClothSeamConstFacade& Other, const int32 SimVertex2DOffset, const int32 SimVertex3DOffset);
 
+		/** Clean up references to invalid indices, including updating stitches to maintain topology. */
+		void CleanupAndCompact();
+
 private:
 		friend class FCollectionClothFacade;
 		FCollectionClothSeamFacade(const TSharedRef<class FClothCollection>& InClothCollection, int32 SeamIndex);
