Fix Switch CIS build error.

diff --git a/Engine/Plugins/Experimental/ChaosClothAsset/Source/ChaosClothAsset/Private/ChaosClothAsset/CollectionClothSeamFacade.cpp b/Engine/Plugins/Experimental/ChaosClothAsset/Source/ChaosClothAsset/Private/ChaosClothAsset/CollectionClothSeamFacade.cpp
--- a/Engine/Plugins/Experimental/ChaosClothAsset/Source/ChaosClothAsset/Private/ChaosClothAsset/CollectionClothSeamFacade.cpp
+++ b/Engine/Plugins/Experimental/ChaosClothAsset/Source/ChaosClothAsset/Private/ChaosClothAsset/CollectionClothSeamFacade.cpp
@@ -911,6 +911,9 @@ static FAutoConsoleVariableRef CVarEnableSeamChecks(TEXT("p.ChaosCloth.EnableSea
 		{
 			int32 Vertex3D;
 			TArray<FIntVector2> StitchEnds;
+			FNewStitches(int32 InVertex3D, TArray<FIntVector2>&& InStitchEnds)
+				: Vertex3D(InVertex3D), StitchEnds(MoveTemp(InStitchEnds))
+			{}
 		};
 		TArray<FNewStitches> NewStitchesToAdd;
 		for (int32 StitchIndex = 0; StitchIndex < GetNumSeamStitches(); ++StitchIndex)
