Implement Min LOD handling for clothing assets.

Clothing assets had a MinLod member that could be configured in the editor but implementation for the virtual methods to hook that up to the skinned mesh LOD update all just returned LOD0 in all cases. After this change we'll correctly ask for the per-platform MinLod. This doesn't currently account for quality levels, and is called out as a TODO for a later change. For that we should add an engine flag for cloth preferring quality settings similar to skeletal and static meshes and use that to select a quality LOD instead.

diff --git a/Engine/Plugins/Experimental/ChaosClothAsset/Source/ChaosClothAssetEngine/Private/ChaosClothAsset/ClothAsset.cpp b/Engine/Plugins/Experimental/ChaosClothAsset/Source/ChaosClothAssetEngine/Private/ChaosClothAsset/ClothAsset.cpp
--- a/Engine/Plugins/Experimental/ChaosClothAsset/Source/ChaosClothAssetEngine/Private/ChaosClothAsset/ClothAsset.cpp
+++ b/Engine/Plugins/Experimental/ChaosClothAsset/Source/ChaosClothAssetEngine/Private/ChaosClothAsset/ClothAsset.cpp
@@ -616,6 +616,27 @@ FString UChaosClothAsset::GetAsyncPropertyName(uint64 Property) const
 	return StaticEnum<EClothAssetAsyncProperties>()->GetNameByValue(Property).ToString();
 }
 
+int32 UChaosClothAsset::GetMinLodIdx(bool bForceLowestLODIndex) const
+{
+	// #TODO Add quality level controls alongside per-platform taking precedence when enabled
+	return GetMinLod().GetValue();
+}
+
+int32 UChaosClothAsset::GetPlatformMinLODIdx(const ITargetPlatform* InTargetPlatform) const
+{
+	// #TODO Add quality level controls alongside per-platform taking precedence when enabled
+#if WITH_EDITOR
+	return GetMinLod().GetValueForPlatform(*InTargetPlatform->IniPlatformName());
+#else
+	return 0;
+#endif
+}
+
+const FPerPlatformInt& UChaosClothAsset::GetMinLod() const
+{
+	return MinLod;
+}
+
 #if WITH_EDITOR
 void UChaosClothAsset::CacheDerivedData(FSkinnedAssetCompilationContext* Context)
 {
diff --git a/Engine/Plugins/Experimental/ChaosClothAsset/Source/ChaosClothAssetEngine/Public/ChaosClothAsset/ClothAsset.h b/Engine/Plugins/Experimental/ChaosClothAsset/Source/ChaosClothAssetEngine/Public/ChaosClothAsset/ClothAsset.h
--- a/Engine/Plugins/Experimental/ChaosClothAsset/Source/ChaosClothAssetEngine/Public/ChaosClothAsset/ClothAsset.h
+++ b/Engine/Plugins/Experimental/ChaosClothAsset/Source/ChaosClothAssetEngine/Public/ChaosClothAsset/ClothAsset.h
@@ -84,12 +84,12 @@ public:
 	virtual class UMaterialInterface* GetOverlayMaterial() const override		{ return nullptr; }
 	virtual float GetOverlayMaterialMaxDrawDistance() const override			{ return 0.f; }
 	virtual bool IsValidLODIndex(int32 Index) const override					{ return LODInfo.IsValidIndex(Index); }
-	virtual int32 GetMinLodIdx(bool bForceLowestLODIdx = false) const override	{ return 0; }
+	virtual int32 GetMinLodIdx(bool bForceLowestLODIdx = false) const override;
 	virtual bool NeedCPUData(int32 LODIndex) const override						{ return false; }
 	virtual bool GetHasVertexColors() const override							{ return false; }
-	virtual int32 GetPlatformMinLODIdx(const ITargetPlatform* TargetPlatform) const override { return 0; }
+	virtual int32 GetPlatformMinLODIdx(const ITargetPlatform* TargetPlatform) const override;
 	virtual const FPerPlatformBool& GetDisableBelowMinLodStripping() const override { return DisableBelowMinLodStripping; }
-	virtual const FPerPlatformInt& GetMinLod() const override					{ return MinLod; }
+	virtual const FPerPlatformInt& GetMinLod() const override;
 #if WITH_EDITOR
 	/* Build a LOD model for the targeted platform. */
 	virtual void BuildLODModel(const ITargetPlatform* TargetPlatform, int32 LODIndex) override;
