ChaosCloth : TrimCollectionOnCook
- attribute culling on the cloth collection during asset cooking

diff --git a/Engine/Plugins/Experimental/ChaosClothAsset/Source/ChaosClothAssetEngine/Private/ChaosClothAsset/ClothAsset.cpp b/Engine/Plugins/Experimental/ChaosClothAsset/Source/ChaosClothAssetEngine/Private/ChaosClothAsset/ClothAsset.cpp
--- a/Engine/Plugins/Experimental/ChaosClothAsset/Source/ChaosClothAssetEngine/Private/ChaosClothAsset/ClothAsset.cpp
+++ b/Engine/Plugins/Experimental/ChaosClothAsset/Source/ChaosClothAssetEngine/Private/ChaosClothAsset/ClothAsset.cpp
@@ -7,6 +7,7 @@
 #include "ChaosClothAsset/ClothGeometryTools.h"
 #include "ChaosClothAsset/ClothAssetPrivate.h"
 #include "ChaosClothAsset/ClothSimulationModel.h"
+#include "Chaos/CollectionPropertyFacade.h"
 #include "Animation/Skeleton.h"
 #if WITH_EDITORONLY_DATA
 #include "Animation/AnimationAsset.h"
@@ -31,8 +32,15 @@
 // In case of merge conflicts with DDC versions, you MUST generate a new GUID and set this new GUID as the version.
 #define CHAOS_CLOTH_ASSET_DERIVED_DATA_VERSION TEXT("5C041E93563044A69646B5E98F152B7C")
 
+
 namespace UE::Chaos::ClothAsset::Private
 {
+bool bClothCollectionOnlyCookPropertyFacade = true;
+FAutoConsoleVariableRef CVarClothCollectionOnlyCookPropertyFacade(
+	TEXT("p.ClothCollectionOnlyCookPropertyFacade"),
+	bClothCollectionOnlyCookPropertyFacade,
+	TEXT("Default setting for culling propertys on the cloth collection during the cook. Default[false]"));
+	
 ::Chaos::FChaosArchive& Serialize(::Chaos::FChaosArchive& Ar, TArray<TSharedRef<FManagedArrayCollection>>& ClothCollections)
 {
 	Ar.UsingCustomVersion(FUE5MainStreamObjectVersion::GUID);
@@ -97,6 +105,34 @@ namespace UE::Chaos::ClothAsset::Private
 		return Ar;
 	}
 }
+
+TArray<TSharedRef<FManagedArrayCollection>> TrimOnCook(const FString InAssetName, const TArray<TSharedRef<FManagedArrayCollection>>& InClothCollections)
+{
+	int32 Index = 0;
+#if WITH_EDITORONLY_DATA
+	if (bClothCollectionOnlyCookPropertyFacade)
+	{
+		TArray<TSharedRef<FManagedArrayCollection>> OutputCollections;
+		for (TSharedRef<FManagedArrayCollection> ClothCollection : InClothCollections)
+		{
+			TSharedPtr<FManagedArrayCollection> PropertyCollection(new FManagedArrayCollection());
+			::Chaos::Softs::FCollectionPropertyMutableFacade CollectionPropertyMutableFacade(PropertyCollection);
+			CollectionPropertyMutableFacade.Copy(*ClothCollection);
+			OutputCollections.Add(TSharedRef<FManagedArrayCollection>(PropertyCollection.ToSharedRef()));
+			UE_LOG(LogChaosClothAsset, Display, TEXT("TrimOnCook[ON] %s:[%d] [size:%d]"),
+				*InAssetName, Index++, PropertyCollection->GetAllocatedSize());
+		}
+		return OutputCollections;
+	}
+#endif
+	for (TSharedRef<FManagedArrayCollection> ClothCollection : InClothCollections)
+	{
+		UE_LOG(LogChaosClothAsset, Display, TEXT("TrimOnCook [OFF] %s:[%d] [size:%d]"),
+			*InAssetName, Index++, ClothCollection->GetAllocatedSize());
+	}
+	return InClothCollections;
+}
+
 }
 
 UChaosClothAsset::UChaosClothAsset(const FObjectInitializer& ObjectInitializer)
@@ -156,6 +192,7 @@ FMatrix UChaosClothAsset::GetComposedRefPoseMatrix(FName InBoneName) const
 	return LocalPose;
 }
 
+
 void UChaosClothAsset::Serialize(FArchive& Ar)
 {
 	Super::Serialize(Ar);
@@ -163,8 +200,19 @@ void UChaosClothAsset::Serialize(FArchive& Ar)
 	bool bCooked = Ar.IsCooking();
 	Ar << bCooked;
 
-	Chaos::FChaosArchive ChaosArchive(Ar);
-	UE::Chaos::ClothAsset::Private::Serialize(ChaosArchive, ClothCollections);
+
+	if (bCooked && Ar.IsSaving())
+	{
+		TArray<TSharedRef<FManagedArrayCollection>> OutputCollections = 
+			UE::Chaos::ClothAsset::Private::TrimOnCook(GetPathName(), ClothCollections);
+		Chaos::FChaosArchive ChaosArchive(Ar);
+		UE::Chaos::ClothAsset::Private::Serialize(ChaosArchive, OutputCollections);
+	}
+	else
+	{
+		Chaos::FChaosArchive ChaosArchive(Ar);
+		UE::Chaos::ClothAsset::Private::Serialize(ChaosArchive, ClothCollections);
+	}
 
 	Ar << GetRefSkeleton();
 
