Chaos Cloth Merge crash fixes

Merge node: handle when first input doesn't have any properties.
RenderFacade: handle when second input doesn't have RenderDeformer data.

diff --git a/Engine/Plugins/Experimental/ChaosClothAsset/Source/ChaosClothAsset/Private/ChaosClothAsset/CollectionClothRenderPatternFacade.cpp b/Engine/Plugins/Experimental/ChaosClothAsset/Source/ChaosClothAsset/Private/ChaosClothAsset/CollectionClothRenderPatternFacade.cpp
--- a/Engine/Plugins/Experimental/ChaosClothAsset/Source/ChaosClothAsset/Private/ChaosClothAsset/CollectionClothRenderPatternFacade.cpp
+++ b/Engine/Plugins/Experimental/ChaosClothAsset/Source/ChaosClothAsset/Private/ChaosClothAsset/CollectionClothRenderPatternFacade.cpp
@@ -226,7 +226,7 @@ namespace UE::Chaos::ClothAsset
 		FClothCollection::CopyArrayViewData(GetRenderBoneIndices(), Other.GetRenderBoneIndices());
 		FClothCollection::CopyArrayViewData(GetRenderBoneWeights(), Other.GetRenderBoneWeights());
 
-		if (GetClothCollection()->IsValid(EClothCollectionOptionalSchemas::RenderDeformer))
+		if (GetClothCollection()->IsValid(EClothCollectionOptionalSchemas::RenderDeformer) && Other.ClothCollection->IsValid(EClothCollectionOptionalSchemas::RenderDeformer))
 		{
 			FClothCollection::CopyArrayViewData(GetRenderDeformerPositionBaryCoordsAndDist(), Other.GetRenderDeformerPositionBaryCoordsAndDist());
 			FClothCollection::CopyArrayViewData(GetRenderDeformerNormalBaryCoordsAndDist(), Other.GetRenderDeformerNormalBaryCoordsAndDist());
diff --git a/Engine/Plugins/Experimental/ChaosClothAssetEditor/Source/ChaosClothAssetDataflowNodes/Private/ChaosClothAsset/MergeClothCollectionsNode.cpp b/Engine/Plugins/Experimental/ChaosClothAssetEditor/Source/ChaosClothAssetDataflowNodes/Private/ChaosClothAsset/MergeClothCollectionsNode.cpp
--- a/Engine/Plugins/Experimental/ChaosClothAssetEditor/Source/ChaosClothAssetDataflowNodes/Private/ChaosClothAsset/MergeClothCollectionsNode.cpp
+++ b/Engine/Plugins/Experimental/ChaosClothAssetEditor/Source/ChaosClothAssetDataflowNodes/Private/ChaosClothAsset/MergeClothCollectionsNode.cpp
@@ -233,6 +233,10 @@ namespace UE::Chaos::ClothAsset::Private
 			else
 			{
 				// If not we add a new property with the flags/bounds/string of the input one
+				if (!OutPropertyFacade.IsValid())
+				{
+					OutPropertyFacade.DefineSchema();
+				}
 				OutKeyIndex = OutPropertyFacade.AddProperty(InPropertyKey, InPropertyFlags);
 			}
 			if(bOverrideProperty)
