Chaos Cloth Asset: Add delta time smoothing.

This is the same as what's implemented for SKM cloth, but the default decay rate is slower (seemed to work better at least in my tests). Since the rate is slower, don't smooth from PrevDeltaTime == 0.

diff --git a/Engine/Plugins/Experimental/ChaosClothAsset/Source/ChaosClothAssetEngine/Private/ChaosClothAsset/ClothSimulationProxy.cpp b/Engine/Plugins/Experimental/ChaosClothAsset/Source/ChaosClothAssetEngine/Private/ChaosClothAsset/ClothSimulationProxy.cpp
--- a/Engine/Plugins/Experimental/ChaosClothAsset/Source/ChaosClothAssetEngine/Private/ChaosClothAsset/ClothSimulationProxy.cpp
+++ b/Engine/Plugins/Experimental/ChaosClothAsset/Source/ChaosClothAssetEngine/Private/ChaosClothAsset/ClothSimulationProxy.cpp
@@ -37,6 +37,9 @@ namespace UE::Chaos::ClothAsset
 	FAutoConsoleVariableRef CVarTransformClothSimukDataISPCEnabled(TEXT("p.ChaosClothAsset.TransformClothSimulData.ISPC"), bTransformClothSimulData_ISPC_Enabled, TEXT("Whether to use ISPC optimizations when transforming simulation data back to reference bone space."));
 #endif
 
+	float DeltaTimeDecay = 0.03;
+	FAutoConsoleVariableRef CVarDeltaTimeDecay(TEXT("p.ChaosClothAsset.DeltaTimeDecay"), DeltaTimeDecay, TEXT("Delta Time smoothing decay (1 = no smoothing)"));
+
 	static FAutoConsoleTaskPriority CPrio_ClothSimulationProxyParallelTask(
 		TEXT("TaskGraph.TaskPriorities.ClothSimulationProxyParallelTask"),
 		TEXT("Task and thread priority for the cloth simulation proxy."),
@@ -215,6 +218,10 @@ PRAGMA_ENABLE_DEPRECATION_WARNINGS
 		{
 			return;
 		}
+		// Filter delta time to smoothen time variations and prevent unwanted vibrations
+		const Softs::FSolverReal DeltaTime = (Softs::FSolverReal)ClothSimulationContext->DeltaTime;
+		const Softs::FSolverReal PrevDeltaTime = Solver->GetDeltaTime() > 0.f ? Solver->GetDeltaTime() : DeltaTime;
+		const Softs::FSolverReal SmoothedDeltaTime = PrevDeltaTime + (DeltaTime - PrevDeltaTime) * (Softs::FSolverReal)DeltaTimeDecay;
 
 		const double StartTime = FPlatformTime::Seconds();
 		const float PrevSimulationTime = SimulationTime;  // Copy the atomic to prevent a re-read
@@ -229,6 +236,7 @@ PRAGMA_ENABLE_DEPRECATION_WARNINGS
 		Solver->SetWindVelocity(ClothSimulationContext->WindVelocity);
 		Solver->SetGravity(ClothSimulationContext->WorldGravity);
 		Solver->EnableClothGravityOverride(true);
+		Solver->SetVelocityScale(!bNeedsReset ? (FReal)ClothSimulationContext->VelocityScale * (FReal)SmoothedDeltaTime / DeltaTime : 1.f);
 
 		// Check teleport modes
 		for (const TUniquePtr<FClothingSimulationCloth>& Cloth : Cloths)
@@ -247,7 +255,7 @@ PRAGMA_ENABLE_DEPRECATION_WARNINGS
 		// Step the simulation
 		if (Solver->GetEnableSolver() || !bUseCache)
 		{
-			Solver->Update((Softs::FSolverReal)ClothSimulationContext->DeltaTime);
+			Solver->Update(SmoothedDeltaTime);
 		}
 		else
 		{
