Fix for ClothCollection::GetUserDefinedAttributeNames not returning any names if it doesn't have the optional schemas

diff --git a/Engine/Plugins/Experimental/ChaosClothAsset/Source/ChaosClothAsset/Private/ChaosClothAsset/ClothCollection.cpp b/Engine/Plugins/Experimental/ChaosClothAsset/Source/ChaosClothAsset/Private/ChaosClothAsset/ClothCollection.cpp
--- a/Engine/Plugins/Experimental/ChaosClothAsset/Source/ChaosClothAsset/Private/ChaosClothAsset/ClothCollection.cpp
+++ b/Engine/Plugins/Experimental/ChaosClothAsset/Source/ChaosClothAsset/Private/ChaosClothAsset/ClothCollection.cpp
@@ -633,19 +633,14 @@ namespace UE::Chaos::ClothAsset
 		const TArray<FName>& FixedAttributeNames = FixedAttributeNamesMap.FindChecked(GroupName);  // Also checks that the group name is a recognized group name 
 		const TArray<FName> AttributeNames = ManagedArrayCollection->AttributeNames(GroupName);
 
-		const int32 MaxUserDefinedAttributes = AttributeNames.Num() - FixedAttributeNames.Num();
-		if (MaxUserDefinedAttributes > 0)
+		for (const FName& AttributeName : AttributeNames)
 		{
-			UserDefinedAttributeNames.Reserve(MaxUserDefinedAttributes);
-
-			for (const FName& AttributeName : AttributeNames)
+			if (!FixedAttributeNames.Contains(AttributeName) && ManagedArrayCollection->FindAttributeTyped<T>(AttributeName, GroupName))
 			{
-				if (!FixedAttributeNames.Contains(AttributeName) && ManagedArrayCollection->FindAttributeTyped<T>(AttributeName, GroupName))
-				{
-					UserDefinedAttributeNames.Add(AttributeName);
-				}
+				UserDefinedAttributeNames.Add(AttributeName);
 			}
 		}
+
 		return UserDefinedAttributeNames;
 	}
 	template CHAOSCLOTHASSET_API TArray<FName> FClothCollection::GetUserDefinedAttributeNames<bool>(const FName& GroupName) const;
