Skeletal Mesh - Added a GetAdditionalRequiredBonesForLeader() virtual function to USkinnedMeshComponent to add any additional required bones in functions such as USkeletalMeshComponent::ComputeRequiredBones().

Cloth Asset - Implemented GetAdditionalRequiredBonesForLeader() to add the follower's required bones to the leader component's list of required bones when the cloth component is used as a follower to a leader component.

#tests Successfully tested on a Metahuman UE5 setup at LOD2, with breakpoint to make sure other Follower Components were not affected.

diff --git a/Engine/Plugins/Experimental/ChaosClothAsset/Source/ChaosClothAssetEngine/Private/ChaosClothAsset/ClothComponent.cpp b/Engine/Plugins/Experimental/ChaosClothAsset/Source/ChaosClothAssetEngine/Private/ChaosClothAsset/ClothComponent.cpp
--- a/Engine/Plugins/Experimental/ChaosClothAsset/Source/ChaosClothAssetEngine/Private/ChaosClothAsset/ClothComponent.cpp
+++ b/Engine/Plugins/Experimental/ChaosClothAsset/Source/ChaosClothAssetEngine/Private/ChaosClothAsset/ClothComponent.cpp
@@ -6,6 +6,7 @@
 #include "ChaosClothAsset/ClothSimulationModel.h"
 #include "ChaosClothAsset/ClothSimulationProxy.h"
 #include "Chaos/CollectionPropertyFacade.h"
+#include "Components/SkeletalMeshComponent.h"
 #include "GeometryCollection/ManagedArrayCollection.h"
 #include "HAL/IConsoleManager.h"
 #include "Rendering/SkeletalMeshRenderData.h"
@@ -321,7 +322,14 @@ void UChaosClothComponent::OnAttachmentChanged()
 {
 	if (bUseAttachedParentAsPoseComponent)
 	{
-		SetLeaderPoseComponent(Cast<USkinnedMeshComponent>(GetAttachParent()));  // If the cast fail, remove the current leader
+		USkinnedMeshComponent* const AttachParentComponent = Cast<USkinnedMeshComponent>(GetAttachParent());
+		SetLeaderPoseComponent(AttachParentComponent);  // If the cast fail, remove the current leader
+
+		// When parented to a skeletal mesh, the anim setup needs re-initializing in order to use the follower's bones requirement
+		if (USkeletalMeshComponent* const SkeletalMeshComponent = Cast<USkeletalMeshComponent>(AttachParentComponent))
+		{
+			SkeletalMeshComponent->RecalcRequiredBones(SkeletalMeshComponent->GetPredictedLODLevel());
+		}
 	}
 
 	Super::OnAttachmentChanged();
@@ -390,6 +398,37 @@ void UChaosClothComponent::SetSkinnedAssetAndUpdate(USkinnedAsset* InSkinnedAsse
 	}
 }
 
+void UChaosClothComponent::GetAdditionalRequiredBonesForLeader(int32 LODIndex, TArray<FBoneIndexType>& InOutRequiredBones) const
+{
+	if (const FSkeletalMeshRenderData* const SkeletalMeshRenderData = GetSkeletalMeshRenderData())
+	{
+		if (SkeletalMeshRenderData->LODRenderData.IsValidIndex(LODIndex))
+		{
+			// Gather the follower's bones
+			TArray<FBoneIndexType> RequiredBones;
+			RequiredBones.Reserve(SkeletalMeshRenderData->LODRenderData[LODIndex].RequiredBones.Num());
+
+			for (const FBoneIndexType RequiredBone : SkeletalMeshRenderData->LODRenderData[LODIndex].RequiredBones)
+			{
+				if (LeaderBoneMap.IsValidIndex(RequiredBone))
+				{
+					const int32 FollowerRequiredLeaderBone = LeaderBoneMap[RequiredBone];
+					if (FollowerRequiredLeaderBone != INDEX_NONE)
+					{
+						RequiredBones.Add((FBoneIndexType)FollowerRequiredLeaderBone);
+					}
+				}
+			}
+
+			// Then sort array of required bones in hierarchy order
+			RequiredBones.Sort();
+
+			// Make sure all of these are in RequiredBones.
+			MergeInBoneIndexArrays(InOutRequiredBones, RequiredBones);
+		}
+	}
+}
+
 TSharedPtr<UE::Chaos::ClothAsset::FClothSimulationProxy> UChaosClothComponent::CreateClothSimulationProxy()
 {
 	using namespace UE::Chaos::ClothAsset;
diff --git a/Engine/Plugins/Experimental/ChaosClothAsset/Source/ChaosClothAssetEngine/Public/ChaosClothAsset/ClothComponent.h b/Engine/Plugins/Experimental/ChaosClothAsset/Source/ChaosClothAssetEngine/Public/ChaosClothAsset/ClothComponent.h
--- a/Engine/Plugins/Experimental/ChaosClothAsset/Source/ChaosClothAssetEngine/Public/ChaosClothAsset/ClothComponent.h
+++ b/Engine/Plugins/Experimental/ChaosClothAsset/Source/ChaosClothAssetEngine/Public/ChaosClothAsset/ClothComponent.h
@@ -141,6 +141,7 @@ protected:
 	virtual void RefreshBoneTransforms(FActorComponentTickFunction* TickFunction = nullptr) override;
 	virtual void GetUpdateClothSimulationData_AnyThread(TMap<int32, FClothSimulData>& OutClothSimulData, FMatrix& OutLocalToWorld, float& OutBlendWeight) override;
 	virtual void SetSkinnedAssetAndUpdate(USkinnedAsset* InSkinnedAsset, bool bReinitPose = true) override;
+	virtual void GetAdditionalRequiredBonesForLeader(int32 LODIndex, TArray<FBoneIndexType>& InOutRequiredBones) const override;
 	//~ End USkinnedMeshComponent Interface
 
 	/** Override this function for setting up custom simulation proxies when the component is registered. */
diff --git a/Engine/Source/Runtime/Engine/Classes/Components/SkinnedMeshComponent.h b/Engine/Source/Runtime/Engine/Classes/Components/SkinnedMeshComponent.h
--- a/Engine/Source/Runtime/Engine/Classes/Components/SkinnedMeshComponent.h
+++ b/Engine/Source/Runtime/Engine/Classes/Components/SkinnedMeshComponent.h
@@ -1985,6 +1985,14 @@ public:
 	 * This assumes both arrays are sorted and contain unique bone indices.
 	 */
 	static ENGINE_API void MergeInBoneIndexArrays(TArray<FBoneIndexType>& BaseArray, const TArray<FBoneIndexType>& InsertArray);
+
+	/**
+	 * Override this function to add any additional required bones from followers to functions such as USkeletalMeshComponent::ComputeRequiredBones().
+	 * The additional bones will need to be passed using the Leader's bone index mapping, sorted in hierarchical order, and merged into the InOutRequiredBones array.
+	 * @param LODIndex - the LOD to get the required bones for.
+	 * @param InOutRequiredBones - the array in which to merge the additional required bones.
+	 */
+	virtual void GetAdditionalRequiredBonesForLeader(int32 LODIndex, TArray<FBoneIndexType>& InOutRequiredBones) const {}
 };
 
 class FRenderStateRecreator
diff --git a/Engine/Source/Runtime/Engine/Private/Components/SkeletalMeshComponent.cpp b/Engine/Source/Runtime/Engine/Private/Components/SkeletalMeshComponent.cpp
--- a/Engine/Source/Runtime/Engine/Private/Components/SkeletalMeshComponent.cpp
+++ b/Engine/Source/Runtime/Engine/Private/Components/SkeletalMeshComponent.cpp
@@ -1881,6 +1881,15 @@ void USkeletalMeshComponent::ComputeRequiredBones(TArray<FBoneIndexType>& OutReq
 
 	// TODO - Make sure that bones with per-poly collision are also always updated.
 
+	// Get any additional required bones from followers
+	for (const TWeakObjectPtr<USkinnedMeshComponent>& FollowerPoseComponent : FollowerPoseComponents)
+	{
+		if (const USkinnedMeshComponent* const FollowerPoseComponentPtr = FollowerPoseComponent.Get())
+		{
+			FollowerPoseComponentPtr->GetAdditionalRequiredBonesForLeader(LODIndex, OutRequiredBones);
+		}
+	}
+
 	// Purge invisible bones and their children
 	// this has to be done before mirror table check/physics body checks
 	// mirror table/phys body ones has to be calculated
