Chaos Cloth Asset - Fixed crash when removing a LOD from the Cloth Editor Dataflow while visualizing it.

#tests Removed the last LOD of the demo asset, and checked the simulation behavior after removal of the LOD.

diff --git a/Engine/Plugins/Experimental/ChaosClothAsset/Source/ChaosClothAssetEngine/Private/ChaosClothAsset/ClothSimulationContext.cpp b/Engine/Plugins/Experimental/ChaosClothAsset/Source/ChaosClothAssetEngine/Private/ChaosClothAsset/ClothSimulationContext.cpp
--- a/Engine/Plugins/Experimental/ChaosClothAsset/Source/ChaosClothAssetEngine/Private/ChaosClothAsset/ClothSimulationContext.cpp
+++ b/Engine/Plugins/Experimental/ChaosClothAsset/Source/ChaosClothAssetEngine/Private/ChaosClothAsset/ClothSimulationContext.cpp
@@ -6,6 +6,7 @@
 #include "ChaosClothAsset/ClothSimulationModel.h"
 #include "Engine/World.h"
 #include "HAL/IConsoleManager.h"
+#include "Rendering/SkeletalMeshRenderData.h"
 #include "SceneInterface.h"
 
 namespace UE::Chaos::ClothAsset
@@ -94,19 +95,29 @@ namespace UE::Chaos::ClothAsset
 
 		// Update bone matrices
 		RefToLocalMatrices.Reset(NumBones);
+
+		bool bSetRefToLocalMatricesToIdentity = true;
 		if (!bIsInitialization)
 		{
-			const TArray<uint16>* RequiredExtraBones = nullptr;
-			if (TSharedPtr<const FChaosClothSimulationModel > ClothModel = ClothAsset ? ClothAsset->GetClothSimulationModel() : TSharedPtr<const FChaosClothSimulationModel >(nullptr))
+			if (const FSkeletalMeshRenderData* const RenderData = ClothAsset ? ClothAsset->GetResourceForRendering() : nullptr)
 			{
-				if (ClothModel->IsValidLodIndex(LodIndex))
+				if (RenderData->LODRenderData.IsValidIndex(LodIndex))
 				{
-					RequiredExtraBones = &ClothModel->ClothSimulationLodModels[LodIndex].RequiredExtraBoneIndices;
+					const TArray<uint16>* RequiredExtraBones = nullptr;
+					if (TSharedPtr<const FChaosClothSimulationModel > ClothModel = ClothAsset ? ClothAsset->GetClothSimulationModel() : TSharedPtr<const FChaosClothSimulationModel >(nullptr))
+					{
+						if (ClothModel->IsValidLodIndex(LodIndex))
+						{
+							RequiredExtraBones = &ClothModel->ClothSimulationLodModels[LodIndex].RequiredExtraBoneIndices;
+						}
+					}
+					ClothComponent.GetCurrentRefToLocalMatrices(RefToLocalMatrices, LodIndex, RequiredExtraBones);
+					bSetRefToLocalMatricesToIdentity = false;
 				}
 			}
-			ClothComponent.GetCurrentRefToLocalMatrices(RefToLocalMatrices, LodIndex, RequiredExtraBones);
 		}
-		else
+
+		if (bSetRefToLocalMatricesToIdentity)
 		{
 			RefToLocalMatrices.AddUninitialized(NumBones);
 
diff --git a/Engine/Plugins/Experimental/ChaosClothAsset/Source/ChaosClothAssetEngine/Private/ChaosClothAsset/ClothSimulationMesh.cpp b/Engine/Plugins/Experimental/ChaosClothAsset/Source/ChaosClothAssetEngine/Private/ChaosClothAsset/ClothSimulationMesh.cpp
--- a/Engine/Plugins/Experimental/ChaosClothAsset/Source/ChaosClothAssetEngine/Private/ChaosClothAsset/ClothSimulationMesh.cpp
+++ b/Engine/Plugins/Experimental/ChaosClothAsset/Source/ChaosClothAssetEngine/Private/ChaosClothAsset/ClothSimulationMesh.cpp
@@ -24,7 +24,7 @@ PRAGMA_ENABLE_DEPRECATION_WARNINGS
 
 	int32 FClothSimulationMesh::GetLODIndex() const
 	{
-		return ClothSimulationContext.LodIndex;
+		return FMath::Min(ClothSimulationContext.LodIndex, GetNumLODs() - 1);
 	}
 
 	int32 FClothSimulationMesh::GetOwnerLODIndex(int32 LODIndex) const
