Chaos Cloth Asset - Remove unused maps on build. Save about 1MB on reference asset.

diff --git a/Engine/Plugins/Experimental/ChaosCloth/Source/Chaos/Private/ChaosCloth/ChaosClothingSimulationConfig.cpp b/Engine/Plugins/Experimental/ChaosCloth/Source/Chaos/Private/ChaosCloth/ChaosClothingSimulationConfig.cpp
--- a/Engine/Plugins/Experimental/ChaosCloth/Source/Chaos/Private/ChaosCloth/ChaosClothingSimulationConfig.cpp
+++ b/Engine/Plugins/Experimental/ChaosCloth/Source/Chaos/Private/ChaosCloth/ChaosClothingSimulationConfig.cpp
@@ -197,8 +197,10 @@ namespace Chaos
 
 			if (ClothConfig->bUseSelfCollisionSpheres)
 			{
-				Property->AddValue("SelfCollisionSphereRadius", ClothConfig->SelfCollisionSphereRadius, NonAnimatablePropertyFlags);
-				Property->AddValue("SelfCollisionSphereStiffness", ClothConfig->SelfCollisionSphereStiffness, AnimatablePropertyFlags);
+				Property->AddValue(TEXT("SelfCollisionSphereRadius"), ClothConfig->SelfCollisionSphereRadius, NonAnimatablePropertyFlags);
+				Property->AddValue(TEXT("SelfCollisionSphereStiffness"), ClothConfig->SelfCollisionSphereStiffness, AnimatablePropertyFlags);
+				static const FString SelfCollisionSphereSetName(TEXT("SelfCollisionSphereSetName"));
+				Property->AddStringValue(SelfCollisionSphereSetName, SelfCollisionSphereSetName, NonAnimatablePropertyFlags);
 			}
 
 			// Max distance
diff --git a/Engine/Plugins/Experimental/ChaosCloth/Source/Chaos/Private/ChaosCloth/ChaosClothingSimulationSkeletalMesh.cpp b/Engine/Plugins/Experimental/ChaosCloth/Source/Chaos/Private/ChaosCloth/ChaosClothingSimulationSkeletalMesh.cpp
--- a/Engine/Plugins/Experimental/ChaosCloth/Source/Chaos/Private/ChaosCloth/ChaosClothingSimulationSkeletalMesh.cpp
+++ b/Engine/Plugins/Experimental/ChaosCloth/Source/Chaos/Private/ChaosCloth/ChaosClothingSimulationSkeletalMesh.cpp
@@ -166,8 +166,8 @@ PRAGMA_ENABLE_DEPRECATION_WARNINGS
 		{
 			const FClothLODDataCommon& ClothLODData = Asset->LodData[LODIndex];
 			const FClothPhysicalMeshData& ClothPhysicalMeshData = ClothLODData.PhysicalMeshData;
-			static const FString SelfCollisionSetName(TEXT("_SelfCollisionSpheres"));
-			VertexSets.Add(SelfCollisionSetName, &ClothPhysicalMeshData.SelfCollisionVertexSet);
+			static const FString SelfCollisionSphereSetName(TEXT("SelfCollisionSphereSetName"));
+			VertexSets.Add(SelfCollisionSphereSetName, &ClothPhysicalMeshData.SelfCollisionVertexSet);
 		}
 
 		return VertexSets;
diff --git a/Engine/Plugins/Experimental/ChaosClothAsset/Source/ChaosClothAssetEngine/Private/ChaosClothAsset/ClothSimulationModel.cpp b/Engine/Plugins/Experimental/ChaosClothAsset/Source/ChaosClothAssetEngine/Private/ChaosClothAsset/ClothSimulationModel.cpp
--- a/Engine/Plugins/Experimental/ChaosClothAsset/Source/ChaosClothAssetEngine/Private/ChaosClothAsset/ClothSimulationModel.cpp
+++ b/Engine/Plugins/Experimental/ChaosClothAsset/Source/ChaosClothAssetEngine/Private/ChaosClothAsset/ClothSimulationModel.cpp
@@ -3,13 +3,15 @@
 #include "ChaosClothAsset/ClothSimulationModel.h"
 #include "ChaosClothAsset/CollectionClothFacade.h"
 #include "ChaosClothAsset/CollectionClothSelectionFacade.h"
+#include "ChaosClothAsset/ClothAssetPrivate.h"
 #include "ChaosClothAsset/ClothCollectionGroup.h"
 #include "ChaosClothAsset/ClothGeometryTools.h"
 #include "ChaosClothAsset/ClothLodTransitionDataCache.h"
+#include "Chaos/CollectionPropertyFacade.h"
 #include "Utils/ClothingMeshUtils.h"
 #include "ReferenceSkeleton.h"
 
-namespace UE::Chaos::ClothAsset
+namespace UE::Chaos::ClothAsset::Private
 {
 	// Find the root bone for this cloth asset (common bone for all used bones)
 	static int32 CalculateReferenceBoneIndex(const TArray<FChaosClothSimulationLodModel> ClothSimulationLodModels, const FReferenceSkeleton& ReferenceSkeleton, const TArray<int32>& UsedBoneIndices)
@@ -132,6 +134,8 @@ namespace UE::Chaos::ClothAsset
 
 bool FChaosClothSimulationLodModel::Serialize(FArchive& Ar)
 {
+	using namespace UE::Chaos::ClothAsset;
+
 	// Serialize normal tagged property data
 	if (Ar.IsLoading() || Ar.IsSaving())
 	{
@@ -139,6 +143,65 @@ bool FChaosClothSimulationLodModel::Serialize(FArchive& Ar)
 		Struct->SerializeTaggedProperties(Ar, (uint8*)this, Struct, nullptr);
 	}
 
+	bool bCooked = Ar.IsCooking();
+	Ar << bCooked;
+
+#if WITH_EDITORONLY_DATA
+	if (bCooked && Ar.IsSaving())
+	{
+		TArray<FName> WeightMapsToRemove;
+		for (TPair<FName, TArray<float>>& WeightMap : WeightMaps)
+		{
+			if (WeightMap.Key.ToString().StartsWith(TEXT("_")))
+			{
+				WeightMapsToRemove.Add(WeightMap.Key);
+			}
+		}
+		for (const FName Key : WeightMapsToRemove)
+		{
+			TArray<float>& Value = WeightMaps.FindChecked(Key);
+			Value.Shrink();
+			const uint64 Size = (uint64)Value.GetAllocatedSize();
+			WeightMaps.Remove(Key);
+			UE_LOG(LogChaosClothAsset, Display, TEXT("TrimOnCook [%s]: Removed WeightMap [%s] (%llu bytes)"), *Ar.GetArchiveName(), *Key.ToString(), Size);
+		}
+
+		TArray<FName> VertexSetsToRemove;
+		for (TPair<FName, TSet<int32>>& VertexSet : VertexSets)
+		{
+			if (VertexSet.Key.ToString().StartsWith(TEXT("_")))
+			{
+				VertexSetsToRemove.Add(VertexSet.Key);
+			}
+		}
+		for (const FName Key : VertexSetsToRemove)
+		{
+			TSet<int32>& Value = VertexSets.FindChecked(Key);
+			Value.Shrink();
+			const uint64 Size = (uint64)Value.GetAllocatedSize();
+			VertexSets.Remove(Key);
+			UE_LOG(LogChaosClothAsset, Display, TEXT("TrimOnCook [%s]: Removed VertexSet [%s] (%llu bytes)"), *Ar.GetArchiveName(), *Key.ToString(), Size);
+		}
+
+		TArray<FName> FaceSetsToRemove;
+		for (TPair<FName, TSet<int32>>& FaceSet : FaceSets)
+		{
+			if (FaceSet.Key.ToString().StartsWith(TEXT("_")))
+			{
+				FaceSetsToRemove.Add(FaceSet.Key);
+			}
+		}
+		for (const FName Key : FaceSetsToRemove)
+		{
+			TSet<int32>& Value = FaceSets.FindChecked(Key);
+			Value.Shrink();
+			const uint64 Size = (uint64)Value.GetAllocatedSize();
+			FaceSets.Remove(Key);
+			UE_LOG(LogChaosClothAsset, Display, TEXT("TrimOnCook [%s]: Removed FaceSet [%s] (%llu bytes)"), *Ar.GetArchiveName(), *Key.ToString(), Size);
+		}
+	}
+#endif
+
 	// Serialize weight maps (not a tagged property)
 	Ar << WeightMaps;
 
@@ -177,11 +240,32 @@ FChaosClothSimulationModel::FChaosClothSimulationModel(const TArray<TSharedRef<c
 		Cloth.BuildSimulationMesh(LodModel.Positions, LodModel.Normals, LodModel.Indices, LodModel.PatternPositions, LodModel.PatternIndices, LodModel.PatternToWeldedIndices, &WeldedToPatternIndices);
 
 		// Copy weight maps
-		LodModel.WeightMaps.Reserve(WeightMapNames.Num());
+		const ::Chaos::Softs::FCollectionPropertyConstFacade Properties(ClothCollections[LodIndex]);
+
+		auto IsNameUsedInProperties = [&Properties](const FName& WeightMapName)->bool
+			{
+				for (int32 PropertyIndex = 0; PropertyIndex < Properties.Num(); ++PropertyIndex)
+				{
+					if (Properties.GetStringValue(PropertyIndex) == WeightMapName)
+					{
+						return true;
+					}
+				}
+				return false;
+			};
+
+		auto IsNameReserved = [&Properties](const FName& WeightMapName)->bool
+			{
+				return WeightMapName.ToString().StartsWith(TEXT("_"));
+			};
 
+		LodModel.WeightMaps.Reserve(WeightMapNames.Num());
 		for (const FName& WeightMapName : WeightMapNames)
 		{
-			LodModel.WeightMaps.Add(WeightMapName) = Cloth.GetWeightMap(WeightMapName);
+			if (IsNameUsedInProperties(WeightMapName) || IsNameReserved(WeightMapName))  // Only copy weight map used by properties - TODO: Add a switch or flag to disable this behavior when enabling runtime weightmap switching
+			{
+				LodModel.WeightMaps.Add(WeightMapName) = Cloth.GetWeightMap(WeightMapName);
+			}
 		}
 
 		// Copy vertex and face sets
@@ -191,14 +275,17 @@ FChaosClothSimulationModel::FChaosClothSimulationModel(const TArray<TSharedRef<c
 		LodModel.FaceSets.Reserve(SelectionNames.Num());
 		for (const FName& SelectionName : SelectionNames)
 		{
-			const FName SelectionGroup = Selection.GetSelectionGroup(SelectionName);
-			if (SelectionGroup == ClothCollectionGroup::SimVertices3D)
+			if (IsNameUsedInProperties(SelectionName) || IsNameReserved(SelectionName))
 			{
-				LodModel.VertexSets.Add(SelectionName) = Selection.GetSelectionSet(SelectionName);
-			}
-			else if (SelectionGroup == ClothCollectionGroup::SimFaces)
-			{
-				LodModel.FaceSets.Add(SelectionName) = Selection.GetSelectionSet(SelectionName);
+				const FName SelectionGroup = Selection.GetSelectionGroup(SelectionName);
+				if (SelectionGroup == ClothCollectionGroup::SimVertices3D)
+				{
+					LodModel.VertexSets.Add(SelectionName) = Selection.GetSelectionSet(SelectionName);
+				}
+				else if (SelectionGroup == ClothCollectionGroup::SimFaces)
+				{
+					LodModel.FaceSets.Add(SelectionName) = Selection.GetSelectionSet(SelectionName);
+				}
 			}
 		}
 
@@ -207,7 +294,10 @@ FChaosClothSimulationModel::FChaosClothSimulationModel(const TArray<TSharedRef<c
 		LodModel.FaceIntMaps.Reserve(FaceIntMapNames.Num());
 		for (const FName& FaceIntMapName : FaceIntMapNames)
 		{
-			LodModel.FaceIntMaps.Add(FaceIntMapName) = Cloth.GetUserDefinedAttribute<int32>(FaceIntMapName, ClothCollectionGroup::SimFaces);
+			if (IsNameUsedInProperties(FaceIntMapName) || IsNameReserved(FaceIntMapName))
+			{
+				LodModel.FaceIntMaps.Add(FaceIntMapName) = Cloth.GetUserDefinedAttribute<int32>(FaceIntMapName, ClothCollectionGroup::SimFaces);
+			}
 		}
 
 		// Copy bone influences (and track all used sim bones) and gather tether data.
@@ -267,7 +357,7 @@ FChaosClothSimulationModel::FChaosClothSimulationModel(const TArray<TSharedRef<c
 	}
 
 	// Initialize Reference bone index
-	ReferenceBoneIndex = CalculateReferenceBoneIndex(ClothSimulationLodModels, ReferenceSkeleton, UsedBoneIndices);
+	ReferenceBoneIndex = Private::CalculateReferenceBoneIndex(ClothSimulationLodModels, ReferenceSkeleton, UsedBoneIndices);
 
 	CalculateLODTransitionUpDownData(InOutTransitionCache);
 }
@@ -306,7 +396,7 @@ void FChaosClothSimulationModel::CalculateLODTransitionUpDownData(TArray<FChaosC
 
 		for (int32 LODIndex = 0; LODIndex < ClothSimulationLodModels.Num(); ++LODIndex)
 		{
-			const FMD5Hash CurrCacheHash = UE::Chaos::ClothAsset::CalculateTransitionDataHash(ClothSimulationLodModels[LODIndex]);
+			const FMD5Hash CurrCacheHash = UE::Chaos::ClothAsset::Private::CalculateTransitionDataHash(ClothSimulationLodModels[LODIndex]);
 			TransitionCacheHashValid[LODIndex] = (CurrCacheHash == (*InOutTransitionCache)[LODIndex].ModelHash);
 			(*InOutTransitionCache)[LODIndex].ModelHash = CurrCacheHash;
 		}
diff --git a/Engine/Plugins/Experimental/ChaosClothAssetEditor/Source/ChaosClothAssetDataflowNodes/Private/ChaosClothAsset/RemeshNode.cpp b/Engine/Plugins/Experimental/ChaosClothAssetEditor/Source/ChaosClothAssetDataflowNodes/Private/ChaosClothAsset/RemeshNode.cpp
--- a/Engine/Plugins/Experimental/ChaosClothAssetEditor/Source/ChaosClothAssetDataflowNodes/Private/ChaosClothAsset/RemeshNode.cpp
+++ b/Engine/Plugins/Experimental/ChaosClothAssetEditor/Source/ChaosClothAssetDataflowNodes/Private/ChaosClothAsset/RemeshNode.cpp
@@ -10,6 +10,8 @@
 #include "ChaosClothAsset/ClothGeometryTools.h"
 #include "ChaosClothAsset/ClothEngineTools.h"
 #include "ChaosClothAsset/ClothDataflowTools.h"
+#include "ChaosClothAsset/SimulationSelfCollisionSpheresConfigNode.h"
+#include "ChaosClothAsset/SimulationLongRangeAttachmentConfigNode.h"
 #include "Dataflow/DataflowInputOutput.h"
 #include "DynamicMesh/MeshTangents.h"
 #include "DynamicMesh/MeshNormals.h"
@@ -1206,10 +1208,15 @@ void FChaosClothAssetRemeshNode::RebuildTopologyDependentSimData(const TSharedRe
 	Chaos::Softs::FCollectionPropertyConstFacade InProperties(InClothCollection);
 
 	// Reconstruct collision spheres
-	if (InProperties.GetKeyIndex(TEXT("SelfCollisionSphereStiffness")) != INDEX_NONE)
+	const FString SelfCollisionSphereStiffnessString = GET_MEMBER_NAME_STRING_CHECKED(FChaosClothAssetSimulationSelfCollisionSpheresConfigNode, SelfCollisionSphereStiffness);
+	const FString SelfCollisionSphereRadiusString = GET_MEMBER_NAME_STRING_CHECKED(FChaosClothAssetSimulationSelfCollisionSpheresConfigNode, SelfCollisionSphereRadius);
+	const FString SelfCollisionSphereRadiusCullMultiplierString = GET_MEMBER_NAME_STRING_CHECKED(FChaosClothAssetSimulationSelfCollisionSpheresConfigNode, SelfCollisionSphereRadiusCullMultiplier);
+	const FString SelfCollisionSphereSetNameString = GET_MEMBER_NAME_STRING_CHECKED(FChaosClothAssetSimulationSelfCollisionSpheresConfigNode, SelfCollisionSphereSetName);
+
+	if (InProperties.GetKeyIndex(SelfCollisionSphereStiffnessString) != INDEX_NONE)
 	{
-		const float SelfCollisionSphereRadius = InProperties.GetValue<float>(TEXT("SelfCollisionSphereRadius"));
-		const float SelfCollisionSphereRadiusCullMultiplier = InProperties.GetValue<float>(TEXT("SelfCollisionSphereRadiusCullMultiplier"));
+		const float SelfCollisionSphereRadius = InProperties.GetValue<float>(SelfCollisionSphereRadiusString);
+		const float SelfCollisionSphereRadiusCullMultiplier = InProperties.GetValue<float>(SelfCollisionSphereRadiusCullMultiplierString);
 		const float CullDiameterSq = FMath::Square(SelfCollisionSphereRadius * SelfCollisionSphereRadiusCullMultiplier * 2.f);
 
 		if (OutClothFacade.IsValid() && CullDiameterSq > 0.f)
@@ -1221,16 +1228,21 @@ void FChaosClothAssetRemeshNode::RebuildTopologyDependentSimData(const TSharedRe
 			FCollectionClothSelectionFacade Selection(OutClothCollection);
 			Selection.DefineSchema();
 
-			static const FName SelectionSetName(TEXT("_SelfCollisionSpheres"));
+			const FName SelectionSetName(*InProperties.GetStringValue(SelfCollisionSphereSetNameString, SelfCollisionSphereSetNameString));
 			Selection.FindOrAddSelectionSet(SelectionSetName, ClothCollectionGroup::SimVertices3D) = VertexSet;
 		}
 	}
 
 	// Reconstruct long-range attachments
-	if (InProperties.GetKeyIndex(TEXT("TetherStiffness")) != INDEX_NONE)
+	const FString TetherStiffnessString = GET_MEMBER_NAME_STRING_CHECKED(FChaosClothAssetSimulationLongRangeAttachmentConfigNode, TetherStiffness);
+	const FString FixedEndWeightMapString = GET_MEMBER_NAME_STRING_CHECKED(FChaosClothAssetSimulationLongRangeAttachmentConfigNode, FixedEndWeightMap);
+	FString UseGeodesicTethersString = GET_MEMBER_NAME_STRING_CHECKED(FChaosClothAssetSimulationLongRangeAttachmentConfigNode, bUseGeodesicTethers);
+	UseGeodesicTethersString.RemoveFromStart(TEXT("b"), ESearchCase::CaseSensitive);  // Property collection names doesn't use the b prefix for booleans
+
+	if (InProperties.GetKeyIndex(TetherStiffnessString) != INDEX_NONE)
 	{
-		const bool bUseGeodesicTethers = InProperties.GetValue<bool>(TEXT("UseGeodesicTethers"));
-		const FName FixedEndWeightMap(InProperties.GetStringValue(TEXT("FixedEndWeightMap")));
+		const bool bUseGeodesicTethers = InProperties.GetValue<bool>(UseGeodesicTethersString);
+		const FName FixedEndWeightMap(InProperties.GetStringValue(FixedEndWeightMapString));
 
 		UE::Chaos::ClothAsset::FClothEngineTools::GenerateTethers(OutClothCollection, FixedEndWeightMap, bUseGeodesicTethers);
 	}
diff --git a/Engine/Plugins/Experimental/ChaosClothAssetEditor/Source/ChaosClothAssetDataflowNodes/Private/ChaosClothAsset/SimulationSelfCollisionSpheresConfigNode.cpp b/Engine/Plugins/Experimental/ChaosClothAssetEditor/Source/ChaosClothAssetDataflowNodes/Private/ChaosClothAsset/SimulationSelfCollisionSpheresConfigNode.cpp
--- a/Engine/Plugins/Experimental/ChaosClothAssetEditor/Source/ChaosClothAssetDataflowNodes/Private/ChaosClothAsset/SimulationSelfCollisionSpheresConfigNode.cpp
+++ b/Engine/Plugins/Experimental/ChaosClothAssetEditor/Source/ChaosClothAssetDataflowNodes/Private/ChaosClothAsset/SimulationSelfCollisionSpheresConfigNode.cpp
@@ -2,11 +2,11 @@
 
 #include "ChaosClothAsset/SimulationSelfCollisionSpheresConfigNode.h"
 #include "Chaos/CollectionPropertyFacade.h"
-
 #include "ChaosClothAsset/CollectionClothFacade.h"
 #include "ChaosClothAsset/CollectionClothSelectionFacade.h"
 #include "ChaosClothAsset/ClothCollectionGroup.h"
 #include "ChaosClothAsset/ClothGeometryTools.h"
+#include "Dataflow/DataflowInputOutput.h"
 
 #include UE_INLINE_GENERATED_CPP_BY_NAME(SimulationSelfCollisionSpheresConfigNode)
 
@@ -14,6 +14,7 @@ FChaosClothAssetSimulationSelfCollisionSpheresConfigNode::FChaosClothAssetSimula
 	: FChaosClothAssetSimulationBaseConfigNode(InParam, InGuid)
 {
 	RegisterCollectionConnections();
+	RegisterOutputConnection(&SelfCollisionSphereSetName);
 }
 
 void FChaosClothAssetSimulationSelfCollisionSpheresConfigNode::AddProperties(FPropertyHelper& PropertyHelper) const
@@ -21,6 +22,7 @@ void FChaosClothAssetSimulationSelfCollisionSpheresConfigNode::AddProperties(FPr
 	PropertyHelper.SetProperty(this, &SelfCollisionSphereRadius);
 	PropertyHelper.SetProperty(this, &SelfCollisionSphereRadiusCullMultiplier);
 	PropertyHelper.SetProperty(this, &SelfCollisionSphereStiffness);
+	PropertyHelper.SetPropertyString(this, &SelfCollisionSphereSetName, {}, ECollectionPropertyFlags::None);  // Non animatable
 }
 
 void FChaosClothAssetSimulationSelfCollisionSpheresConfigNode::EvaluateClothCollection(Dataflow::FContext& Context, const TSharedRef<FManagedArrayCollection>& ClothCollection) const
@@ -38,8 +40,18 @@ void FChaosClothAssetSimulationSelfCollisionSpheresConfigNode::EvaluateClothColl
 		FCollectionClothSelectionFacade Selection(ClothCollection);
 		Selection.DefineSchema();
 
-		static const FName SelectionSetName(TEXT("_SelfCollisionSpheres"));
-		Selection.FindOrAddSelectionSet(SelectionSetName, ClothCollectionGroup::SimVertices3D) = VertexSet;
+		Selection.FindOrAddSelectionSet(FName(*SelfCollisionSphereSetName), ClothCollectionGroup::SimVertices3D) = VertexSet;
 	}
+}
 
+void FChaosClothAssetSimulationSelfCollisionSpheresConfigNode::Evaluate(Dataflow::FContext& Context, const FDataflowOutput* Out) const
+{
+	if (Out->IsA<FString>(&SelfCollisionSphereSetName))
+	{
+		SetValue(Context, SelfCollisionSphereSetName, &SelfCollisionSphereSetName);
+	}
+	else
+	{
+		Super::Evaluate(Context, Out);
+	}
 }
diff --git a/Engine/Plugins/Experimental/ChaosClothAssetEditor/Source/ChaosClothAssetDataflowNodes/Public/ChaosClothAsset/SimulationSelfCollisionSpheresConfigNode.h b/Engine/Plugins/Experimental/ChaosClothAssetEditor/Source/ChaosClothAssetDataflowNodes/Public/ChaosClothAsset/SimulationSelfCollisionSpheresConfigNode.h
--- a/Engine/Plugins/Experimental/ChaosClothAssetEditor/Source/ChaosClothAssetDataflowNodes/Public/ChaosClothAsset/SimulationSelfCollisionSpheresConfigNode.h
+++ b/Engine/Plugins/Experimental/ChaosClothAssetEditor/Source/ChaosClothAssetDataflowNodes/Public/ChaosClothAsset/SimulationSelfCollisionSpheresConfigNode.h
@@ -28,10 +28,15 @@ public:
 	UPROPERTY(EditAnywhere, Category = "Self-Collision Spheres Properties", meta = (ClampMin = "1"))
 	float SelfCollisionSphereRadiusCullMultiplier = 1.f;
 
+	/** The name of the output selection generated from the unculled vertices that receives the collision spheres. */
+	UPROPERTY(VisibleAnywhere, Category = "Self-Collision Spheres Properties", Meta = (DataflowOutput))
+	FString SelfCollisionSphereSetName = GET_MEMBER_NAME_STRING_CHECKED(FChaosClothAssetSimulationSelfCollisionSpheresConfigNode, SelfCollisionSphereSetName);
+
 	FChaosClothAssetSimulationSelfCollisionSpheresConfigNode(const Dataflow::FNodeParameters& InParam, FGuid InGuid = FGuid::NewGuid());
 
 private:
 	virtual void AddProperties(FPropertyHelper& PropertyHelper) const override;
 
 	virtual void EvaluateClothCollection(Dataflow::FContext& Context, const TSharedRef<FManagedArrayCollection>& ClothCollection) const override;
+	virtual void Evaluate(Dataflow::FContext& Context, const FDataflowOutput* Out) const override;
 };
diff --git a/Engine/Source/Runtime/Experimental/Chaos/Private/Chaos/PBDSelfCollisionSphereConstraints.cpp b/Engine/Source/Runtime/Experimental/Chaos/Private/Chaos/PBDSelfCollisionSphereConstraints.cpp
--- a/Engine/Source/Runtime/Experimental/Chaos/Private/Chaos/PBDSelfCollisionSphereConstraints.cpp
+++ b/Engine/Source/Runtime/Experimental/Chaos/Private/Chaos/PBDSelfCollisionSphereConstraints.cpp
@@ -5,10 +5,8 @@
 #include "Chaos/Framework/Parallel.h"
 #include "Chaos/HierarchicalSpatialHash.h"
 
-namespace Chaos::Softs {
-
-	const FString FPBDSelfCollisionSphereConstraints::SelectionSetName(TEXT("_SelfCollisionSpheres"));
-
+namespace Chaos::Softs
+{
 	namespace Private
 	{
 		struct FSphereSpatialEntry
diff --git a/Engine/Source/Runtime/Experimental/Chaos/Public/Chaos/PBDSelfCollisionSphereConstraints.h b/Engine/Source/Runtime/Experimental/Chaos/Public/Chaos/PBDSelfCollisionSphereConstraints.h
--- a/Engine/Source/Runtime/Experimental/Chaos/Public/Chaos/PBDSelfCollisionSphereConstraints.h
+++ b/Engine/Source/Runtime/Experimental/Chaos/Public/Chaos/PBDSelfCollisionSphereConstraints.h
@@ -51,8 +51,6 @@ namespace Chaos::Softs
 		typedef FPBDSelfCollisionSphereConstraintsBase Base;
 
 	public:
-		CHAOS_API static const FString SelectionSetName;
-
 		static bool IsEnabled(const FCollectionPropertyConstFacade& PropertyCollection)
 		{
 			return IsSelfCollisionSphereStiffnessEnabled(PropertyCollection, false);
@@ -65,7 +63,7 @@ namespace Chaos::Softs
 			const FCollectionPropertyConstFacade& PropertyCollection)
 			: Base(InOffset,
 				InNumParticles,
-				VertexSets.FindRef(SelectionSetName, nullptr),
+				VertexSets.FindRef(GetSelfCollisionSphereSetNameString(PropertyCollection, SelfCollisionSphereSetNameName.ToString()), nullptr),
 				nullptr, // No reference positions--using vertex sets instead
 				(FSolverReal)FMath::Max(GetSelfCollisionSphereRadius(PropertyCollection, 0.5f), 0.f),
 				(FSolverReal)FMath::Clamp(GetSelfCollisionSphereStiffness(PropertyCollection, 0.5f), 0.f, 1.f)),
@@ -82,9 +80,9 @@ namespace Chaos::Softs
 		using Base::Stiffness;
 		using Base::VertexSetNoOffset;
 
+		UE_CHAOS_DECLARE_INDEXLESS_PROPERTYCOLLECTION_NAME(SelfCollisionSphereSetName, bool);  // Selection set name string property, the bool value is not actually used
 		UE_CHAOS_DECLARE_PROPERTYCOLLECTION_NAME(SelfCollisionSphereRadius, float);
 		UE_CHAOS_DECLARE_PROPERTYCOLLECTION_NAME(SelfCollisionSphereStiffness, float);
-
 	};
 }
 #endif
