Chaos Cloth Asset:
- Move the ChaosClothAsset and ChaosClothAssetEditor plug-ins out of the Experimental folder and mark them as Beta.
- Move the ChaosClothAssetDataflowNodes and ChaosClothAssetTools modules out of the ChaosClothAssetEditor plug-in to the ChaosClothAsset plug-in.

#virtualized

diff --git a/Engine/Plugins/ChaosClothAsset/ChaosClothAsset.uplugin b/Engine/Plugins/ChaosClothAsset/ChaosClothAsset.uplugin
new file mode 100644
--- /dev/null
+++ b/Engine/Plugins/ChaosClothAsset/ChaosClothAsset.uplugin
@@ -0,0 +1,75 @@
+{
+	"FileVersion": 1,
+	"Version": 1,
+	"VersionName": "0.1",
+	"FriendlyName": "Chaos Cloth Asset",
+	"Description": "Pattern based cloth asset using the Chaos Cloth simulation.",
+	"Category": "Physics",
+	"CreatedBy": "Epic Games, Inc.",
+	"CreatedByURL": "https://epicgames.com",
+	"DocsURL": "",
+	"MarketplaceURL": "",
+	"SupportURL": "",
+	"EnabledByDefault": false,
+	"CanContainContent": true,
+	"IsBetaVersion": true,
+	"Installed": false,
+	"Modules": [
+		{
+			"Name": "ChaosClothAsset",
+			"Type": "Runtime",
+			"LoadingPhase": "Default"
+		},
+		{
+			"Name": "ChaosClothAssetEngine",
+			"Type": "Runtime",
+			"LoadingPhase": "Default"
+		},
+		{
+			"Name": "ChaosClothAssetTools",
+			"Type": "UncookedOnly",
+			"LoadingPhase": "Default",
+			"PlatformAllowList": [ "Win64" ]
+		},
+		{
+			"Name": "ChaosClothAssetDataflowNodes",
+			"Type": "Editor",
+			"LoadingPhase": "Default",
+			"PlatformAllowList": [ "Win64" ]
+		}
+	],
+	"Plugins": [
+		{
+			"Name": "ChaosCloth",
+			"Enabled": true
+		},
+		{
+			"Name": "Dataflow",
+			"Enabled": true
+		},
+		{
+			"Name": "DatasmithImporter",
+			"Enabled": true
+		},
+		{
+			"Name": "GeometryCache",
+			"Enabled": true
+		},
+		{
+			"Name": "GeometryProcessing",
+			"Enabled": true
+		},
+		{
+			"Name": "MeshModelingToolsetExp",
+			"Enabled": true,
+			"TargetAllowList": [ "Editor" ]
+		},
+		{
+			"Name": "USDImporter",
+			"Enabled": true,
+			"SupportedTargetPlatforms": [
+				"Win64"
+			]
+		}
+	]
+}
diff --git a/Engine/Plugins/ChaosClothAsset/Source/ChaosClothAsset/Public/ChaosClothAsset/CollectionClothFacade.h b/Engine/Plugins/ChaosClothAsset/Source/ChaosClothAsset/Public/ChaosClothAsset/CollectionClothFacade.h
new file mode 100644
--- /dev/null
+++ b/Engine/Plugins/ChaosClothAsset/Source/ChaosClothAsset/Public/ChaosClothAsset/CollectionClothFacade.h
@@ -0,0 +1,376 @@
+// Copyright Epic Games, Inc. All Rights Reserved.
+
+#pragma once
+
+#include "Math/Vector.h"
+#include "ChaosClothAsset/ClothCollectionOptionalSchemas.h"
+#include "ChaosClothAsset/CollectionClothFabricFacade.h"
+#include "ChaosClothAsset/CollectionClothRenderPatternFacade.h"
+#include "ChaosClothAsset/CollectionClothSeamFacade.h"
+#include "ChaosClothAsset/CollectionClothSimPatternFacade.h"
+#include "ChaosClothAsset/IsUserAttributeType.h"
+
+namespace Chaos
+{
+class FChaosArchive;
+}
+
+namespace UE::Chaos::ClothAsset
+{
+	struct FDefaultSolver
+	{
+		inline static const FVector3f Gravity = FVector3f(0.0f, 0.0f, -980.665f);
+		inline static constexpr float AirDamping = 0.035f;
+		inline static constexpr int32 SubSteps = 1;
+		inline static constexpr float TimeStep = 0.033f;
+	};
+
+	/**
+	 * Cloth Asset collection facade class focused on draping and pattern information.
+	 * Const access (read only) version.
+	 */
+	class CHAOSCLOTHASSET_API FCollectionClothConstFacade
+	{
+	public:
+		explicit FCollectionClothConstFacade(const TSharedRef<const FManagedArrayCollection>& ManagedArrayCollection);
+
+		FCollectionClothConstFacade() = delete;
+
+		FCollectionClothConstFacade(const FCollectionClothConstFacade&) = delete;
+		FCollectionClothConstFacade& operator=(const FCollectionClothConstFacade&) = delete;
+
+		FCollectionClothConstFacade(FCollectionClothConstFacade&&) = default;
+		FCollectionClothConstFacade& operator=(FCollectionClothConstFacade&&) = default;
+
+		virtual ~FCollectionClothConstFacade() = default;
+
+		/** Return whether the facade is defined on the collection. */
+		bool IsValid(EClothCollectionOptionalSchemas OptionalSchemas = EClothCollectionOptionalSchemas::None) const;
+
+		/**
+		 * Return whether the facade has a non-empty sim and render mesh data.
+		 */
+		bool HasValidData() const;
+
+		uint32 CalculateTypeHash(bool bIncludeWeightMaps, uint32 PreviousHash = 0) const;
+		uint32 CalculateWeightMapTypeHash(uint32 PreviousHash = 0) const;
+		template<typename T, TEMPLATE_REQUIRES(TIsUserAttributeType<T>::Value)>
+		uint32 CalculateUserDefinedAttributesTypeHash(const FName& GroupName, uint32 PreviousHash = 0) const;
+
+		//~ LOD (single per collection) Group
+		/** Return the physics asset path names used for this collection. */
+		const FString& GetPhysicsAssetPathName() const;
+		/** Return the skeleton asset path names used for this collection. */
+		const FString& GetSkeletalMeshPathName() const;
+		
+
+		//~ Solver (single per collection) Group
+		/** Return true if the solver group has one element*/
+		bool HasSolverElement() const;
+		/** Return the solver gravity vector used for this collection. */
+		const FVector3f& GetSolverGravity() const;
+		/** Return the solver air damping used for this collection. */
+		float GetSolverAirDamping() const;
+		/** Return the solver time step used for this collection. */
+		float GetSolverTimeStep() const;
+		/** Return the solver sub steps used for this collection. */
+		int32 GetSolverSubSteps() const;
+
+		//~ Sim Vertices 2D Group
+		/** Return the total number of 2D simulation vertices for this collection. */
+		int32 GetNumSimVertices2D() const;
+		TConstArrayView<FVector2f> GetSimPosition2D() const;
+		TConstArrayView<int32> GetSimVertex3DLookup() const;
+
+		//~ Sim Vertices 3D Group
+		/** Return the total number of 3D simulation vertices for this collection. */
+		int32 GetNumSimVertices3D() const;
+		TConstArrayView<FVector3f> GetSimPosition3D() const;
+		TConstArrayView<FVector3f> GetSimNormal() const;
+		TConstArrayView<TArray<int32>> GetSimBoneIndices() const;
+		TConstArrayView<TArray<float>> GetSimBoneWeights() const;
+		TConstArrayView<TArray<int32>> GetTetherKinematicIndex() const;
+		TConstArrayView<TArray<float>> GetTetherReferenceLength() const;
+		TConstArrayView<TArray<int32>> GetSimVertex2DLookup() const;
+		TConstArrayView<TArray<int32>> GetSeamStitchLookup() const;
+
+		//~ Sim Faces Group
+		/** Return the total number of simulation faces for this collection across all patterns. */
+		int32 GetNumSimFaces() const;
+		TConstArrayView<FIntVector3> GetSimIndices2D() const;
+		TConstArrayView<FIntVector3> GetSimIndices3D() const;
+
+		//~ Sim Patterns Group
+		/** Return the number of patterns in this collection. */
+		int32 GetNumSimPatterns() const;
+		/** Return a pattern facade for the specified pattern index. */
+		FCollectionClothSimPatternConstFacade GetSimPattern(int32 PatternIndex) const;
+		/** Convenience to find which sim pattern a 2D vertex belongs to */
+		int32 FindSimPatternByVertex2D(int32 Vertex2DIndex) const;
+		/** Convenience to find which sim pattern a sim face belongs to */
+		int32 FindSimPatternByFaceIndex(int32 FaceIndex) const;
+
+		//~ Render Patterns Group
+		/** Return the number of patterns in this collection. */
+		int32 GetNumRenderPatterns() const;
+		/** Return a pattern facade for the specified pattern index. */
+		FCollectionClothRenderPatternConstFacade GetRenderPattern(int32 PatternIndex) const;
+		/** Return a view of all the render deformer number of influences used on this collection across all patterns. */
+		TConstArrayView<int32> GetRenderDeformerNumInfluences() const;
+		/** Return a view of all the render materials used on this collection across all patterns. */
+		TConstArrayView<FString> GetRenderMaterialPathName() const;
+		/** Convenience to find which render pattern a render vertex belongs to */
+		int32 FindRenderPatternByVertex(int32 VertexIndex) const;
+		/** Convenience to find which render pattern a render face belongs to */
+		int32 FindRenderPatternByFaceIndex(int32 FaceIndex) const;
+
+		//~ Seam Group
+		/** Return the number of seams in this collection. */
+		int32 GetNumSeams() const;
+		/** Return a seam facade for the specified seam index. */
+		FCollectionClothSeamConstFacade GetSeam(int32 SeamIndex) const;
+
+		//~ Fabric Group
+		/** Return the number of fabrics in this collection. */
+		int32 GetNumFabrics() const;
+		/** Return a fabric facade for the specified fabric index. */
+		FCollectionClothFabricConstFacade GetFabric(int32 FabricIndex) const;
+
+		//~ Render Vertices Group
+		/** Return the total number of render vertices for this collection. */
+		int32 GetNumRenderVertices() const;
+		TConstArrayView<FVector3f> GetRenderPosition() const;
+		TConstArrayView<FVector3f> GetRenderNormal() const;
+		TConstArrayView<FVector3f> GetRenderTangentU() const;
+		TConstArrayView<FVector3f> GetRenderTangentV() const;
+		TConstArrayView<TArray<FVector2f>> GetRenderUVs() const;
+		TConstArrayView<FLinearColor> GetRenderColor() const;
+		TConstArrayView<TArray<int32>> GetRenderBoneIndices() const;
+		TConstArrayView<TArray<float>> GetRenderBoneWeights() const;
+		TConstArrayView<TArray<FVector4f>> GetRenderDeformerPositionBaryCoordsAndDist() const;
+		TConstArrayView<TArray<FVector4f>> GetRenderDeformerNormalBaryCoordsAndDist() const;
+		TConstArrayView<TArray<FVector4f>> GetRenderDeformerTangentBaryCoordsAndDist() const;
+		TConstArrayView<TArray<FIntVector3>> GetRenderDeformerSimIndices3D() const;
+		TConstArrayView<TArray<float>> GetRenderDeformerWeight() const;
+		TConstArrayView<float> GetRenderDeformerSkinningBlend() const;
+
+		//~ Render Faces Group
+		int32 GetNumRenderFaces() const;
+		TConstArrayView<FIntVector3> GetRenderIndices() const;
+
+		//~ Weight Maps
+		/** Return whether this cloth collection has the specified weight map. */
+		bool HasWeightMap(const FName& Name) const;
+		/** Return the name of all user weight maps on this cloth collection. */
+		TArray<FName> GetWeightMapNames() const;
+		TConstArrayView<float> GetWeightMap(const FName& Name) const;
+
+		//~ Other User-Defined Attributes (not instantiated for bools)
+		template<typename T, TEMPLATE_REQUIRES(TIsUserAttributeType<T>::Value)>
+		bool HasUserDefinedAttribute(const FName& Name, const FName& GroupName) const;
+		template<typename T, TEMPLATE_REQUIRES(TIsUserAttributeType<T>::Value)>
+		TArray<FName> GetUserDefinedAttributeNames(const FName& GroupName) const;
+		template<typename T, TEMPLATE_REQUIRES(TIsUserAttributeType<T>::Value)>
+		TConstArrayView<T> GetUserDefinedAttribute(const FName& Name, const FName& GroupName) const;
+
+		void BuildSimulationMesh(TArray<FVector3f>& Positions, TArray<FVector3f>& Normals, TArray<uint32>& Indices, TArray<FVector2f>& PatternsPositions, TArray<uint32>& PatternsIndices, 
+			TArray<uint32>& PatternToWeldedIndices, TArray<TArray<int32>>* OptionalWeldedToPatternIndices = nullptr) const;
+
+	protected:
+		friend class FCollectionClothFacade;
+		TSharedRef<const class FClothCollection> ClothCollection;
+
+		friend class FCollectionClothFacade;  // To enable access from a different instance
+		friend class FCollectionClothSeamFacade;
+		friend class FCollectionClothSeamConstFacade;
+		friend class FCollectionClothFabricFacade;
+		friend class FCollectionClothFabricConstFacade;
+		
+		explicit FCollectionClothConstFacade(const TSharedRef<const class FClothCollection>& ClothCollection);
+	};
+
+	/**
+	 * Cloth Asset collection facade class focused on draping and pattern information.
+	 * Non-const access (read/write) version.
+	 */
+	class CHAOSCLOTHASSET_API FCollectionClothFacade final : public FCollectionClothConstFacade
+	{
+	public:
+		explicit FCollectionClothFacade(const TSharedRef<FManagedArrayCollection>& ManagedArrayCollection);
+
+		FCollectionClothFacade() = delete;
+
+		FCollectionClothFacade(const FCollectionClothFacade&) = delete;
+		FCollectionClothFacade& operator=(const FCollectionClothFacade&) = delete;
+
+		FCollectionClothFacade(FCollectionClothFacade&&) = default;
+		FCollectionClothFacade& operator=(FCollectionClothFacade&&) = default;
+		virtual ~FCollectionClothFacade() override = default;
+
+		/** Create this facade's groups and attributes. */
+		void DefineSchema(EClothCollectionOptionalSchemas OptionalSchemas = EClothCollectionOptionalSchemas::None);
+
+		/** Remove all LODs from this cloth. */
+		void Reset();
+
+		/** Initialize the cloth using another cloth collection. */
+		void Initialize(const FCollectionClothConstFacade& Other);
+
+		/** Append data from another cloth collection. */
+		void Append(const FCollectionClothConstFacade& Other);
+
+		/** Append to this cloth another cloth. */
+		//void Append(const FCollectionClothConstFacade& Other);
+
+		//~ LOD (single per collection) Group
+		/** Set the physics asset path name. */
+		void SetPhysicsAssetPathName(const FString& PathName);
+		/** Set the skeletal mesh asset path name and the reference skeleton that will be used with this asset. */
+		void SetSkeletalMeshPathName(const FString& PathName);
+
+		//~ Solver (max 1 per collection) Group
+		/** Set the solver gravity */
+		void SetSolverGravity(const FVector3f& SolverGravity);
+		/** Set the solver air damping */
+		void SetSolverAirDamping(const float SolverAirDamping);
+		/** Set the solver time step */
+		void SetSolverTimeStep(const float SolverTimeStep);
+		/** Set the solver substeps */
+		void SetSolverSubSteps(const int32 SolverSubSteps);
+
+		//~ Pattern Sim Vertices 2D Group
+		/** SetNumSimVertices2D per pattern within pattern facade. */
+		TArrayView<FVector2f> GetSimPosition2D();
+
+		//~ Pattern Sim Vertices 3D Group
+		TArrayView<FVector3f> GetSimPosition3D();
+		TArrayView<FVector3f> GetSimNormal();
+		TArrayView<TArray<int32>> GetSimBoneIndices();
+		TArrayView<TArray<float>> GetSimBoneWeights();
+		TArrayView<TArray<int32>> GetTetherKinematicIndex();
+		TArrayView<TArray<float>> GetTetherReferenceLength();
+
+		/** This will remove the 3D vertices, but the associated seams and 2D vertices will still exist, and point to INDEX_NONE */
+		void RemoveSimVertices3D(int32 NumSimVertices);
+		void RemoveAllSimVertices3D() { RemoveSimVertices3D(GetNumSimVertices3D()); }
+		void RemoveSimVertices3D(const TArray<int32>& SortedDeletionList);
+		/** Compact SimVertex2DLookup to remove any references to INDEX_NONE that may have been created by deleting 2D vertices. */
+		void CompactSimVertex2DLookup();
+		/** Compact SeamStitchLookup to remove any references to INDEX_NONE that may have been created by deleting stitches. */
+		void CompactSeamStitchLookup();
+
+		//~ Pattern Sim Faces Group
+		/** SetNumSimFaces per pattern within pattern facade. */
+		TArrayView<FIntVector3> GetSimIndices2D();
+		TArrayView<FIntVector3> GetSimIndices3D();
+
+		//~ Sim Patterns Group
+		/** Set the new number of patterns to this cloth LOD. */
+		void SetNumSimPatterns(int32 NumPatterns);
+		/** Add a new pattern to this cloth LOD and return its index in the LOD pattern list. */
+		int32 AddSimPattern();
+		/** Return a pattern facade for the specified pattern index. */
+		FCollectionClothSimPatternFacade GetSimPattern(int32 PatternIndex);
+		/** Add a new pattern to this cloth LOD, and return the cloth pattern facade set to its index. */
+		FCollectionClothSimPatternFacade AddGetSimPattern() { return GetSimPattern(AddSimPattern()); }
+		/** Remove a sorted list of sim patterns. */
+		void RemoveSimPatterns(const TArray<int32>& SortedDeletionList);
+
+		//~ Render Patterns Group
+		/** Set the new number of patterns to this cloth LOD. */
+		void SetNumRenderPatterns(int32 NumPatterns);
+		/** Add a new pattern to this cloth LOD and return its index in the LOD pattern list. */
+		int32 AddRenderPattern();
+		/** Return a pattern facade for the specified pattern index. */
+		FCollectionClothRenderPatternFacade GetRenderPattern(int32 PatternIndex);
+		/** Add a new pattern to this cloth LOD, and return the cloth pattern facade set to its index. */
+		FCollectionClothRenderPatternFacade AddGetRenderPattern() { return GetRenderPattern(AddRenderPattern()); }
+		/** Remove a sorted list of render patterns. */
+		void RemoveRenderPatterns(const TArray<int32>& SortedDeletionList);
+		/** Return a view of all the render deformer number of influences used on this collection across all patterns. */
+		TArrayView<int32> GetRenderDeformerNumInfluences();
+		/** Return a view of all the render materials used on this collection across all patterns. */
+		TArrayView<FString> GetRenderMaterialPathName();
+
+		//~ Seam Group
+		/** Set the new number of seams to this cloth. */
+		void SetNumSeams(int32 NumSeams);
+		/** Add a new seam to this cloth and return its index in the seam list. */
+		int32 AddSeam();
+		/** Return a seam facade for the specified seam index. */
+		FCollectionClothSeamFacade GetSeam(int32 SeamIndex);
+		/** Add a new seam to this cloth and return the seam facade set to its index. */
+		FCollectionClothSeamFacade AddGetSeam() { return GetSeam(AddSeam()); }
+		/** Remove a sorted list of seams. */
+		void RemoveSeams(const TArray<int32>& SortedDeletionList);
+		
+		//~ Fabric Group
+		/** Set the new number of fabrics to this cloth. */
+		void SetNumFabrics(int32 NumFabrics);
+		/** Add a new fabric to this cloth and return its index in the fabric list. */
+		int32 AddFabric();
+		/** Return a fabric facade for the specified fabric index. */
+		FCollectionClothFabricFacade GetFabric(int32 FabricIndex);
+		/** Add a new fabric to this cloth and return the fabric facade set to its index. */
+		FCollectionClothFabricFacade AddGetFabric() { return GetFabric(AddFabric()); }
+		/** Remove a sorted list of fabrics. */
+		void RemoveFabrics(const TArray<int32>& SortedDeletionList);
+
+		//~ Render Vertices Group
+		/** SetNumRenderVertices per pattern within pattern facade. */
+		TArrayView<FVector3f> GetRenderPosition();
+		TArrayView<FVector3f> GetRenderNormal();
+		TArrayView<FVector3f> GetRenderTangentU();
+		TArrayView<FVector3f> GetRenderTangentV();
+		TArrayView<TArray<FVector2f>> GetRenderUVs();
+		TArrayView<FLinearColor> GetRenderColor();
+		TArrayView<TArray<int32>> GetRenderBoneIndices();
+		TArrayView<TArray<float>> GetRenderBoneWeights();
+		TArrayView<TArray<FVector4f>> GetRenderDeformerPositionBaryCoordsAndDist();
+		TArrayView<TArray<FVector4f>> GetRenderDeformerNormalBaryCoordsAndDist();
+		TArrayView<TArray<FVector4f>> GetRenderDeformerTangentBaryCoordsAndDist();
+		TArrayView<TArray<FIntVector3>> GetRenderDeformerSimIndices3D();
+		TArrayView<TArray<float>> GetRenderDeformerWeight();
+		TArrayView<float> GetRenderDeformerSkinningBlend();
+
+		//~ Render Faces Group
+		/** SetNumRenderFaces per pattern within pattern facade. */
+		TArrayView<FIntVector3> GetRenderIndices();
+
+		//~ Weight Maps
+		/** Add a new weight map to this cloth. Access is then done per pattern. */
+		void AddWeightMap(const FName& Name);
+		/** Remove a weight map from this cloth. */
+		void RemoveWeightMap(const FName& Name);
+		TArrayView<float> GetWeightMap(const FName& Name);
+
+		//~ Other User-Defined Attributes (not instantiated for bools)
+		/** GroupName must be an existing group as defined in ClothCollectionGroup. Returns success */
+		template<typename T, TEMPLATE_REQUIRES(TIsUserAttributeType<T>::Value)>
+		bool AddUserDefinedAttribute(const FName& Name, const FName& GroupName);
+		void RemoveUserDefinedAttribute(const FName& Name, const FName& GroupName);
+		template<typename T, TEMPLATE_REQUIRES(TIsUserAttributeType<T>::Value)>
+		TArrayView<T> GetUserDefinedAttribute(const FName& Name, const FName& GroupName);
+		static bool IsValidClothCollectionGroupName(const FName& GroupName);
+
+	private:
+		
+		void SetDefaults();
+
+		TSharedRef<class FClothCollection> GetClothCollection() { return ConstCastSharedRef<class FClothCollection>(ClothCollection); }
+
+		friend class FCollectionClothSeamFacade;
+		friend class FCollectionClothSimPatternFacade;
+		friend class FCollectionClothFabricFacade;
+
+		explicit FCollectionClothFacade(const TSharedRef<class FClothCollection>& InClothCollection);
+
+		// These methods are private because they're managed by the FCollectionClothSeamFacade.
+		//~ Sim Vertices 2D Group
+		TArrayView<int32> GetSimVertex3DLookupPrivate();
+
+		//~ Sim Vertices 3D Group
+		TArrayView<TArray<int32>> GetSeamStitchLookupPrivate();
+		TArrayView<TArray<int32>> GetSimVertex2DLookupPrivate();
+	};
+}  // End namespace UE::Chaos::ClothAsset
diff --git a/Engine/Plugins/ChaosClothAsset/Source/ChaosClothAsset/Public/ChaosClothAsset/CollectionClothRenderPatternFacade.h b/Engine/Plugins/ChaosClothAsset/Source/ChaosClothAsset/Public/ChaosClothAsset/CollectionClothRenderPatternFacade.h
new file mode 100644
--- /dev/null
+++ b/Engine/Plugins/ChaosClothAsset/Source/ChaosClothAsset/Public/ChaosClothAsset/CollectionClothRenderPatternFacade.h
@@ -0,0 +1,146 @@
+// Copyright Epic Games, Inc. All Rights Reserved.
+
+#pragma once
+
+#include "Templates/SharedPointer.h"
+
+struct FManagedArrayCollection;
+
+namespace UE::Chaos::ClothAsset
+{
+	/**
+	 * Cloth Asset collection render pattern facade class to access cloth render pattern data.
+	 * Constructed from FCollectionClothConstFacade.
+	 * Const access (read only) version.
+	 */
+	class CHAOSCLOTHASSET_API FCollectionClothRenderPatternConstFacade
+	{
+	public:
+		FCollectionClothRenderPatternConstFacade() = delete;
+
+		FCollectionClothRenderPatternConstFacade(const FCollectionClothRenderPatternConstFacade&) = delete;
+		FCollectionClothRenderPatternConstFacade& operator=(const FCollectionClothRenderPatternConstFacade&) = delete;
+
+		FCollectionClothRenderPatternConstFacade(FCollectionClothRenderPatternConstFacade&&) = default;
+		FCollectionClothRenderPatternConstFacade& operator=(FCollectionClothRenderPatternConstFacade&&) = default;
+
+		virtual ~FCollectionClothRenderPatternConstFacade() = default;
+
+		/** Return the render deformer number of influences for this pattern. */
+		int32 GetRenderDeformerNumInfluences() const;
+		/** Return the render material for this pattern. */
+		const FString& GetRenderMaterialPathName() const;
+
+		//~ Render Vertices Group
+		// Note: Use the FCollectionClothConstFacade accessors instead of these for the array indices to match the RenderIndices values
+		/** Return the total number of render vertices for this pattern. */
+		int32 GetNumRenderVertices() const;
+		/** Return the render vertices offset for this pattern in the render vertices for the collection. */
+		int32 GetRenderVerticesOffset() const;
+		TConstArrayView<FVector3f> GetRenderPosition() const;
+		TConstArrayView<FVector3f> GetRenderNormal() const;
+		TConstArrayView<FVector3f> GetRenderTangentU() const;
+		TConstArrayView<FVector3f> GetRenderTangentV() const;
+		TConstArrayView<TArray<FVector2f>> GetRenderUVs() const;
+		TConstArrayView<FLinearColor> GetRenderColor() const;
+		TConstArrayView<TArray<int32>> GetRenderBoneIndices() const;
+		TConstArrayView<TArray<float>> GetRenderBoneWeights() const;
+		TConstArrayView<TArray<FVector4f>> GetRenderDeformerPositionBaryCoordsAndDist() const;
+		TConstArrayView<TArray<FVector4f>> GetRenderDeformerNormalBaryCoordsAndDist() const;
+		TConstArrayView<TArray<FVector4f>> GetRenderDeformerTangentBaryCoordsAndDist() const;
+		TConstArrayView<TArray<FIntVector3>> GetRenderDeformerSimIndices3D() const;
+		TConstArrayView<TArray<float>> GetRenderDeformerWeight() const;
+		TConstArrayView<float> GetRenderDeformerSkinningBlend() const;
+
+		//~ Render Faces Group
+		// Note: RenderIndices points to the collection arrays, not the pattern arrays
+		int32 GetNumRenderFaces() const;
+		/** Return the render faces offset for this pattern in the render faces */
+		int32 GetRenderFacesOffset() const;
+		TConstArrayView<FIntVector3> GetRenderIndices() const;
+
+		bool IsEmpty() const;
+
+		/** Return the Pattern index this facade has been created with. */
+		int32 GetPatternIndex() const { return PatternIndex; }
+
+	protected:
+		friend class FCollectionClothRenderPatternFacade;  // For other instances access
+		friend class FCollectionClothConstFacade;
+		FCollectionClothRenderPatternConstFacade(const TSharedRef<const class FClothCollection>& InClothCollection, int32 InPatternIndex);
+
+		static constexpr int32 GetBaseElementIndex() { return 0; }
+		int32 GetElementIndex() const { return GetBaseElementIndex() + PatternIndex; }
+
+		TSharedRef<const class FClothCollection> ClothCollection;
+		int32 PatternIndex;
+	};
+
+	/**
+	 * Cloth Asset collection render pattern facade class to access cloth render pattern data.
+	 * Constructed from FCollectionClothFacade.
+	 * Non-const access (read/write) version.
+	 */
+	class CHAOSCLOTHASSET_API FCollectionClothRenderPatternFacade final : public FCollectionClothRenderPatternConstFacade
+	{
+	public:
+		FCollectionClothRenderPatternFacade() = delete;
+
+		FCollectionClothRenderPatternFacade(const FCollectionClothRenderPatternFacade&) = delete;
+		FCollectionClothRenderPatternFacade& operator=(const FCollectionClothRenderPatternFacade&) = delete;
+
+		FCollectionClothRenderPatternFacade(FCollectionClothRenderPatternFacade&&) = default;
+		FCollectionClothRenderPatternFacade& operator=(FCollectionClothRenderPatternFacade&&) = default;
+
+		virtual ~FCollectionClothRenderPatternFacade() override = default;
+
+		/** Remove all geometry from this cloth pattern. */
+		void Reset();
+
+		/** Initialize from another render pattern. Assumes all indices match between source and target. */
+		void Initialize(const FCollectionClothRenderPatternConstFacade& Other, int32 SimVertex3DOffset);
+
+		/** Initialize from another render pattern. Assumes all indices match between source and target. */
+		UE_DEPRECATED(5.4, "Use Initialize with the SimVertex3DOffset instead")
+		void Initialize(const FCollectionClothRenderPatternConstFacade& Other) { Initialize(Other, 0); }
+
+		/** Set the render deformer number of influences for this pattern. */
+		void SetRenderDeformerNumInfluences(int32 NumInfluences);
+		/** Set the render material for this pattern. */
+		void SetRenderMaterialPathName(const FString& PathName);
+
+		//~ Render Vertices Group
+		// Note: Use the FCollectionClothFacade accessors instead of these for the array indices to match the RenderIndices values
+		/** Grow or shrink the space reserved for render vertices for this pattern within the cloth collection and return its start index. */
+		void SetNumRenderVertices(int32 NumRenderVertices);
+		void RemoveRenderVertices(const TArray<int32>& SortedDeletionList);
+		TArrayView<FVector3f> GetRenderPosition();
+		TArrayView<FVector3f> GetRenderNormal();
+		TArrayView<FVector3f> GetRenderTangentU();
+		TArrayView<FVector3f> GetRenderTangentV();
+		TArrayView<TArray<FVector2f>> GetRenderUVs();
+		TArrayView<FLinearColor> GetRenderColor();
+		TArrayView<TArray<int32>> GetRenderBoneIndices();
+		TArrayView<TArray<float>> GetRenderBoneWeights();
+		TArrayView<TArray<FVector4f>> GetRenderDeformerPositionBaryCoordsAndDist();
+		TArrayView<TArray<FVector4f>> GetRenderDeformerNormalBaryCoordsAndDist();
+		TArrayView<TArray<FVector4f>> GetRenderDeformerTangentBaryCoordsAndDist();
+		TArrayView<TArray<FIntVector3>> GetRenderDeformerSimIndices3D();
+		TArrayView<TArray<float>> GetRenderDeformerWeight();
+		TArrayView<float> GetRenderDeformerSkinningBlend();
+
+		//~ Render Faces Group
+		/** Grow or shrink the space reserved for render faces for this pattern within the cloth collection and return its start index. */
+		void SetNumRenderFaces(int32 NumRenderFaces);
+		void RemoveRenderFaces(const TArray<int32>& SortedDeletionList);
+		TArrayView<FIntVector3> GetRenderIndices();
+
+	protected:
+		friend class FCollectionClothFacade;
+		FCollectionClothRenderPatternFacade(const TSharedRef<class FClothCollection>& InClothCollection, int32 InPatternIndex);
+
+		void SetDefaults();
+
+		TSharedRef<class FClothCollection> GetClothCollection() { return ConstCastSharedRef<class FClothCollection>(ClothCollection); }
+	};
+}  // End namespace UE::Chaos::ClothAsset
diff --git a/Engine/Plugins/ChaosClothAsset/Source/ChaosClothAssetEngine/Private/ChaosClothAsset/ClothAsset.cpp b/Engine/Plugins/ChaosClothAsset/Source/ChaosClothAssetEngine/Private/ChaosClothAsset/ClothAsset.cpp
new file mode 100644
--- /dev/null
+++ b/Engine/Plugins/ChaosClothAsset/Source/ChaosClothAssetEngine/Private/ChaosClothAsset/ClothAsset.cpp
@@ -0,0 +1,939 @@
+// Copyright Epic Games, Inc. All Rights Reserved.
+
+#include "ChaosClothAsset/ClothAsset.h"
+#include "ChaosClothAsset/ClothAssetBuilder.h"
+#include "ChaosClothAsset/CollectionClothFacade.h"
+#include "ChaosClothAsset/ClothComponent.h"
+#include "ChaosClothAsset/ClothGeometryTools.h"
+#include "ChaosClothAsset/ClothAssetPrivate.h"
+#include "ChaosClothAsset/ClothSimulationModel.h"
+#include "Chaos/CollectionPropertyFacade.h"
+#include "Animation/Skeleton.h"
+#if WITH_EDITORONLY_DATA
+#include "Animation/AnimationAsset.h"
+#endif
+#include "Engine/RendererSettings.h"
+#include "Engine/SkinnedAssetAsyncCompileUtils.h"
+#include "Features/IModularFeatures.h"
+#include "GeometryCollection/ManagedArrayCollection.h"
+#include "PhysicsEngine/PhysicsAsset.h"
+#include "Rendering/SkeletalMeshModel.h"
+#include "Rendering/SkeletalMeshRenderData.h"
+#include "UObject/Package.h"
+#include "EngineUtils.h"
+#if WITH_EDITOR
+#include "IMeshBuilderModule.h"
+#include "DerivedDataCacheInterface.h"
+#endif
+
+#include UE_INLINE_GENERATED_CPP_BY_NAME(ClothAsset)
+
+// If Chaos cloth asset derived data needs to be rebuilt (new format, serialization differences, etc.) replace the version GUID below with a new one. 
+// In case of merge conflicts with DDC versions, you MUST generate a new GUID and set this new GUID as the version.
+#define CHAOS_CLOTH_ASSET_DERIVED_DATA_VERSION TEXT("5C041E93563044A69646B5E98F152B7C")
+
+
+namespace UE::Chaos::ClothAsset::Private
+{
+bool bClothCollectionOnlyCookPropertyFacade = true;
+FAutoConsoleVariableRef CVarClothCollectionOnlyCookPropertyFacade(
+	TEXT("p.ClothCollectionOnlyCookPropertyFacade"),
+	bClothCollectionOnlyCookPropertyFacade,
+	TEXT("Default setting for culling propertys on the cloth collection during the cook. Default[false]"));
+	
+::Chaos::FChaosArchive& Serialize(::Chaos::FChaosArchive& Ar, TArray<TSharedRef<FManagedArrayCollection>>& ClothCollections)
+{
+	Ar.UsingCustomVersion(FUE5MainStreamObjectVersion::GUID);
+
+	if (Ar.IsLoading() && Ar.CustomVer(FUE5MainStreamObjectVersion::GUID) < FUE5MainStreamObjectVersion::ClothCollectionSingleLodSchema)
+	{
+		// Cloth assets before this version had a single ClothCollection with a completely different schema.
+		ClothCollections.Empty(1);
+		TSharedRef<FManagedArrayCollection>& ClothCollection = ClothCollections.Emplace_GetRef(MakeShared<FManagedArrayCollection>());
+		ClothCollection->Serialize(Ar);
+
+		// Now we're just going to hard reset and define a new schema.
+		ClothCollection->Reset();
+		FCollectionClothFacade ClothFacade(ClothCollection);
+		ClothFacade.DefineSchema();
+
+		return Ar;
+	}
+	else
+	{
+		// This is following Serialize for Arrays
+		ClothCollections.CountBytes(Ar);
+		int32 SerializeNum = Ar.IsLoading() ? 0 : ClothCollections.Num();
+		Ar << SerializeNum;
+		if (SerializeNum == 0)
+		{
+			// if we are loading, then we have to reset the size to 0, in case it isn't currently 0
+			if (Ar.IsLoading())
+			{
+				ClothCollections.Empty();
+			}
+			return Ar;
+		}
+		check(SerializeNum >= 0);
+
+		if (Ar.IsError() || SerializeNum < 0)
+		{
+			Ar.SetError();
+			return Ar;
+		}
+		if (Ar.IsLoading())
+		{
+			// Required for resetting ArrayNum
+			ClothCollections.Empty(SerializeNum);
+
+			for (int32 i = 0; i < SerializeNum; i++)
+			{
+				TSharedRef<FManagedArrayCollection>& ClothCollection = ClothCollections.Emplace_GetRef(MakeShared<FManagedArrayCollection>());
+				ClothCollection->Serialize(Ar);
+			}
+		}
+		else
+		{
+			check(SerializeNum == ClothCollections.Num());
+
+			for (int32 i = 0; i < SerializeNum; i++)
+			{
+				ClothCollections[i]->Serialize(Ar);
+			}
+		}
+
+		return Ar;
+	}
+}
+
+TArray<TSharedRef<FManagedArrayCollection>> TrimOnCook(const FString InAssetName, const TArray<TSharedRef<FManagedArrayCollection>>& InClothCollections)
+{
+	int32 Index = 0;
+#if WITH_EDITORONLY_DATA
+	if (bClothCollectionOnlyCookPropertyFacade)
+	{
+		TArray<TSharedRef<FManagedArrayCollection>> OutputCollections;
+		for (TSharedRef<FManagedArrayCollection> ClothCollection : InClothCollections)
+		{
+			TSharedPtr<FManagedArrayCollection> PropertyCollection(new FManagedArrayCollection());
+			::Chaos::Softs::FCollectionPropertyMutableFacade CollectionPropertyMutableFacade(PropertyCollection);
+			CollectionPropertyMutableFacade.Copy(*ClothCollection);
+			OutputCollections.Add(TSharedRef<FManagedArrayCollection>(PropertyCollection.ToSharedRef()));
+			UE_LOG(LogChaosClothAsset, Display, TEXT("TrimOnCook[ON] %s:[%d] [size:%d]"),
+				*InAssetName, Index++, PropertyCollection->GetAllocatedSize());
+		}
+		return OutputCollections;
+	}
+#endif
+	for (TSharedRef<FManagedArrayCollection> ClothCollection : InClothCollections)
+	{
+		UE_LOG(LogChaosClothAsset, Display, TEXT("TrimOnCook [OFF] %s:[%d] [size:%d]"),
+			*InAssetName, Index++, ClothCollection->GetAllocatedSize());
+	}
+	return InClothCollections;
+}
+
+}
+
+UChaosClothAsset::UChaosClothAsset(const FObjectInitializer& ObjectInitializer)
+	: Super(ObjectInitializer)
+	, DisableBelowMinLodStripping(FPerPlatformBool(false))
+#if WITH_EDITORONLY_DATA
+	, MeshModel(MakeShareable(new FSkeletalMeshModel()))
+#endif
+{
+	// Setup a single LOD's Cloth Collection
+	TSharedRef<FManagedArrayCollection>& ClothCollection = ClothCollections.Emplace_GetRef(MakeShared<FManagedArrayCollection>());
+	UE::Chaos::ClothAsset::FCollectionClothFacade ClothFacade(ClothCollection);
+	ClothFacade.DefineSchema();
+
+	// Add the LODInfo for the default LOD 0
+	LODInfo.SetNum(1);
+
+	// Set default skeleton (must be done after having added the LOD)
+	constexpr bool bRebuildModels = false;
+	constexpr bool bRebindMeshes = false;
+	SetReferenceSkeleton(nullptr, bRebuildModels, bRebindMeshes);
+}
+
+UChaosClothAsset::UChaosClothAsset(FVTableHelper& Helper)
+	: Super(Helper)
+{
+}
+
+PRAGMA_DISABLE_DEPRECATION_WARNINGS
+UChaosClothAsset::~UChaosClothAsset() = default;
+PRAGMA_ENABLE_DEPRECATION_WARNINGS
+
+FSkeletalMeshLODInfo* UChaosClothAsset::GetLODInfo(int32 Index)
+{
+	return LODInfo.IsValidIndex(Index) ? &LODInfo[Index] : nullptr;
+}
+
+const FSkeletalMeshLODInfo* UChaosClothAsset::GetLODInfo(int32 Index) const
+{
+	return LODInfo.IsValidIndex(Index) ? &LODInfo[Index] : nullptr;
+}
+
+FMatrix UChaosClothAsset::GetComposedRefPoseMatrix(FName InBoneName) const
+{
+	FMatrix LocalPose(FMatrix::Identity);
+
+	if (InBoneName != NAME_None)
+	{
+		const int32 BoneIndex = GetRefSkeleton().FindBoneIndex(InBoneName);
+		if (BoneIndex != INDEX_NONE)
+		{
+			return GetComposedRefPoseMatrix(BoneIndex);
+		}
+		// TODO: Might need to add sockets like on the SkeletalMesh
+	}
+
+	return LocalPose;
+}
+
+
+void UChaosClothAsset::Serialize(FArchive& Ar)
+{
+	Super::Serialize(Ar);
+
+	bool bCooked = Ar.IsCooking();
+	Ar << bCooked;
+
+
+	if (bCooked && Ar.IsSaving())
+	{
+		TArray<TSharedRef<FManagedArrayCollection>> OutputCollections = 
+			UE::Chaos::ClothAsset::Private::TrimOnCook(GetPathName(), ClothCollections);
+		Chaos::FChaosArchive ChaosArchive(Ar);
+		UE::Chaos::ClothAsset::Private::Serialize(ChaosArchive, OutputCollections);
+	}
+	else
+	{
+		Chaos::FChaosArchive ChaosArchive(Ar);
+		UE::Chaos::ClothAsset::Private::Serialize(ChaosArchive, ClothCollections);
+	}
+
+	Ar << GetRefSkeleton();
+
+	if (bCooked && !IsTemplate() && !Ar.IsCountingMemory())
+	{
+		if (Ar.IsLoading())
+		{
+			SetResourceForRendering(MakeUnique<FSkeletalMeshRenderData>());
+		}
+		GetResourceForRendering()->Serialize(Ar, this);
+
+		if (!ClothSimulationModel.IsValid())
+		{
+			ClothSimulationModel = MakeShared<FChaosClothSimulationModel>();
+		}
+		UScriptStruct* const Struct = FChaosClothSimulationModel::StaticStruct();
+		Struct->SerializeTaggedProperties(Ar, (uint8*)ClothSimulationModel.Get(), Struct, nullptr);
+	}
+}
+
+void UChaosClothAsset::PostLoad()
+{
+	Super::PostLoad();
+}
+
+#if WITH_EDITOR
+void UChaosClothAsset::PostEditChangeProperty(FPropertyChangedEvent& PropertyChangedEvent)
+{
+	if (PropertyChangedEvent.GetPropertyName() == GET_MEMBER_NAME_CHECKED(UChaosClothAsset, PhysicsAsset))
+	{
+		ReregisterComponents();
+	}
+	Super::PostEditChangeProperty(PropertyChangedEvent);
+}
+#endif // #if WITH_EDITOR
+
+void UChaosClothAsset::BeginPostLoadInternal(FSkinnedAssetPostLoadContext& Context)
+{
+#if WITH_EDITOR
+	TRACE_CPUPROFILER_EVENT_SCOPE(UChaosClothAsset::BeginPostLoadInternal);
+
+	checkf(IsInGameThread(), TEXT("Cannot execute function UChaosClothAsset::BeginPostLoadInternal asynchronously. Asset: %s"), *GetFullName());
+	SetInternalFlags(EInternalObjectFlags::Async);
+
+	// Lock all properties that should not be modified/accessed during async post-load
+	AcquireAsyncProperty();
+
+	// This scope allows us to use any locked properties without causing stalls
+	FSkinnedAssetAsyncBuildScope AsyncBuildScope(this);
+
+	// Make sure that the collection is still compatible and valid
+	bool bAnyInvalidLods = false;
+	if (ClothCollections.IsEmpty())
+	{
+		UE_LOG(LogChaosClothAsset, Warning, TEXT("Invalid Cloth Collection (no LODs) found while loading Cloth Asset %s."), *GetFullName());
+		TSharedRef<FManagedArrayCollection>& ClothCollection = ClothCollections.Emplace_GetRef(MakeShared<FManagedArrayCollection>());
+		UE::Chaos::ClothAsset::FCollectionClothFacade ClothFacade(ClothCollection);
+		ClothFacade.DefineSchema();
+		bAnyInvalidLods = true;
+	}
+
+	const int32 NumLods = ClothCollections.Num();
+	check(NumLods >= 1);  // The default LOD 0 should be present now if it ever was missing
+	LODInfo.SetNum(NumLods);  // Always keep a matching number of LODInfos
+
+	bool bAnyInvalidSkeletons = false;
+	for (int32 LODIndex = 0; LODIndex < NumLods; ++LODIndex)
+	{
+		TSharedRef<FManagedArrayCollection>& ClothCollection = ClothCollections[LODIndex];
+
+		UE::Chaos::ClothAsset::FCollectionClothFacade ClothFacade(ClothCollection);
+		if (!ClothFacade.IsValid())
+		{
+			UE_LOG(LogChaosClothAsset, Warning, TEXT("Invalid Cloth Collection found at LOD %i while loading Cloth Asset %s."), LODIndex, *GetFullName());
+			ClothCollection = MakeShared<FManagedArrayCollection>();
+			ClothFacade = UE::Chaos::ClothAsset::FCollectionClothFacade(ClothCollection);
+			ClothFacade.DefineSchema();
+			bAnyInvalidLods = true;
+			bAnyInvalidSkeletons = true;
+		}
+		else if (ClothFacade.GetSkeletalMeshPathName().IsEmpty())
+		{
+			bAnyInvalidSkeletons = true;
+		}
+	}
+	if (bAnyInvalidLods)
+	{
+		SetPhysicsAsset(PhysicsAsset);  // Re-update the collection with the physics asset information if any
+	}
+	if (bAnyInvalidSkeletons)
+	{
+		constexpr bool bRebuildModels = false;
+		constexpr bool bRebindMeshes = true;  // Best to rebind the mesh when reloading broken data
+		SetReferenceSkeleton(nullptr, bRebuildModels, bRebindMeshes);
+	}
+
+	BuildClothSimulationModel();  // TODO: Cache ClothSimulationModel?
+
+	BuildMeshModel();
+#endif // #if WITH_EDITOR
+}
+
+void UChaosClothAsset::ExecutePostLoadInternal(FSkinnedAssetPostLoadContext& Context)
+{
+#if WITH_EDITOR
+	TRACE_CPUPROFILER_EVENT_SCOPE(UChaosClothAsset::ExecutePostLoadInternal);
+
+	// This scope allows us to use any locked properties without causing stalls
+	FSkinnedAssetAsyncBuildScope AsyncBuildScope(this);
+
+	if (!GetOutermost()->bIsCookedForEditor)
+	{
+		if (GetResourceForRendering() == nullptr)
+		{
+			CacheDerivedData(&Context);
+			Context.bHasCachedDerivedData = true;
+		}
+	}
+#endif // WITH_EDITOR
+}
+
+void UChaosClothAsset::FinishPostLoadInternal(FSkinnedAssetPostLoadContext& Context)
+{
+#if WITH_EDITOR
+	TRACE_CPUPROFILER_EVENT_SCOPE(UChaosClothAsset::FinishPostLoadInternal);
+
+	checkf(IsInGameThread(), TEXT("Cannot execute function UChaosClothAsset::FinishPostLoadInternal asynchronously. Asset: %s"), *this->GetFullName());
+	ClearInternalFlags(EInternalObjectFlags::Async);
+
+	// This scope allows us to use any locked properties without causing stalls
+	FSkinnedAssetAsyncBuildScope AsyncBuildScope(this);
+#endif
+
+	if (FApp::CanEverRender())
+	{
+		InitResources();
+	}
+	else
+	{
+		// Update any missing data when cooking.
+		UpdateUVChannelData(false);
+	}
+
+	CalculateInvRefMatrices();
+	CalculateBounds();
+
+#if WITH_EDITOR
+	ReleaseAsyncProperty();
+#endif
+}
+
+void UChaosClothAsset::BeginDestroy()
+{
+	check(IsInGameThread());
+
+	Super::BeginDestroy();
+
+	// Release the mesh's render resources now
+	ReleaseResources();
+}
+
+bool UChaosClothAsset::IsReadyForFinishDestroy()
+{
+	if (!Super::IsReadyForFinishDestroy())
+	{
+		return false;
+	}
+
+	ReleaseResources();
+
+	// see if we have hit the resource flush fence
+	return ReleaseResourcesFence.IsFenceComplete();
+}
+
+void UChaosClothAsset::InitResources()
+{
+	LLM_SCOPE_BYNAME(TEXT("ClothAsset/InitResources"));
+
+	// Build the material channel data used by the texture streamer
+	UpdateUVChannelData(false);
+
+	if (SkeletalMeshRenderData.IsValid())
+	{
+		TArray<UMorphTarget*> DummyMorphTargets;  // Not even used by InitResources() at the moment
+		SkeletalMeshRenderData->InitResources(false, DummyMorphTargets, this);
+	}
+}
+
+void UChaosClothAsset::ReleaseResources()
+{
+	if (SkeletalMeshRenderData && SkeletalMeshRenderData->IsInitialized())
+	{
+		if (GIsEditor && !GIsPlayInEditorWorld)
+		{
+			// Flush the rendering command to be sure there is no command left that can create/modify a rendering resource
+			FlushRenderingCommands();
+		}
+
+		SkeletalMeshRenderData->ReleaseResources();
+
+		// Insert a fence to signal when these commands completed
+		ReleaseResourcesFence.BeginFence();
+	}
+}
+
+void UChaosClothAsset::CalculateInvRefMatrices()
+{
+	auto GetRefPoseMatrix = [this](int32 BoneIndex)->FMatrix
+	{
+		check(BoneIndex >= 0 && BoneIndex < RefSkeleton.GetRawBoneNum());
+		FTransform BoneTransform = RefSkeleton.GetRawRefBonePose()[BoneIndex];
+		BoneTransform.NormalizeRotation();  // Make sure quaternion is normalized!
+		return BoneTransform.ToMatrixWithScale();
+	};
+
+	const int32 NumRealBones = RefSkeleton.GetRawBoneNum();
+
+	RefBasesInvMatrix.Empty(NumRealBones);
+	RefBasesInvMatrix.AddUninitialized(NumRealBones);
+
+	// Reset cached mesh-space ref pose
+	TArray<FMatrix> ComposedRefPoseMatrices;
+	ComposedRefPoseMatrices.SetNumUninitialized(NumRealBones);
+
+	// Precompute the Mesh.RefBasesInverse
+	for (int32 BoneIndex = 0; BoneIndex < NumRealBones; ++BoneIndex)
+	{
+		// Render the default pose
+		ComposedRefPoseMatrices[BoneIndex] = GetRefPoseMatrix(BoneIndex);
+
+		// Construct mesh-space skeletal hierarchy
+		if (BoneIndex > 0)
+		{
+			int32 Parent = RefSkeleton.GetRawParentIndex(BoneIndex);
+			ComposedRefPoseMatrices[BoneIndex] = ComposedRefPoseMatrices[BoneIndex] * ComposedRefPoseMatrices[Parent];
+		}
+
+		FVector XAxis, YAxis, ZAxis;
+		ComposedRefPoseMatrices[BoneIndex].GetScaledAxes(XAxis, YAxis, ZAxis);
+		if (XAxis.IsNearlyZero(UE_SMALL_NUMBER) &&
+			YAxis.IsNearlyZero(UE_SMALL_NUMBER) &&
+			ZAxis.IsNearlyZero(UE_SMALL_NUMBER))
+		{
+			// This is not allowed, warn them
+			UE_LOG(
+				LogChaosClothAsset,
+				Warning,
+				TEXT("Reference Pose for asset %s for joint (%s) includes NIL matrix. Zero scale isn't allowed on ref pose."),
+				*GetPathName(),
+				*RefSkeleton.GetBoneName(BoneIndex).ToString());
+		}
+
+		// Precompute inverse so we can use from-refpose-skin vertices
+		RefBasesInvMatrix[BoneIndex] = FMatrix44f(ComposedRefPoseMatrices[BoneIndex].Inverse());
+	}
+}
+
+void UChaosClothAsset::CalculateBounds()
+{
+	using namespace UE::Chaos::ClothAsset;
+
+	FBox BoundingBox(ForceInit);
+
+	for (const TSharedRef<FManagedArrayCollection>& ClothCollection : ClothCollections)
+	{
+		const FCollectionClothConstFacade Cloth(ClothCollection);
+		const TConstArrayView<FVector3f> RenderPositionArray = Cloth.GetRenderPosition();
+
+		for (const FVector3f& RenderPosition : RenderPositionArray)
+		{
+			BoundingBox += (FVector)RenderPosition;
+		}
+	}
+
+	Bounds = FBoxSphereBounds(BoundingBox);
+}
+
+void UChaosClothAsset::UpdateSkeleton(bool bRebuildClothSimulationModel)
+{
+	CalculateInvRefMatrices();
+
+	if (bRebuildClothSimulationModel)
+	{
+		// Rebuild simulation model  // TODO: How does this work with skinning
+		BuildClothSimulationModel();
+	}
+}
+
+void UChaosClothAsset::Build(TArray<FChaosClothAssetLodTransitionDataCache>* InOutTransitionCache)
+{
+	using namespace UE::Chaos::ClothAsset;
+
+#if WITH_EDITOR
+	FSkinnedAssetAsyncBuildScope AsyncBuildScope(this);
+
+	FSkinnedAssetBuildContext Context;
+	BeginBuildInternal(Context);
+#else
+	ReleaseResources();
+#endif
+
+	// Set a new Guid to invalidate the DDC
+	AssetGuid = FGuid::NewGuid();
+
+	// Rebuild matrices
+	CalculateInvRefMatrices();
+
+	// Update bounds
+	CalculateBounds();
+
+	// Add LODs to the render data
+	const int32 NumLods = FMath::Max(ClothCollections.Num(), 1);  // The render data will always look for at least one default LOD 0
+
+	// Rebuild LOD Infos
+	LODInfo.Reset(NumLods);
+	LODInfo.AddDefaulted(NumLods);  // TODO: Expose some properties to fill up the LOD infos
+
+	// Build simulation model
+	BuildClothSimulationModel(InOutTransitionCache);
+
+#if WITH_EDITOR
+	// Rebuild LOD Model
+	BuildMeshModel();
+
+	// Load/save render data from/to DDC
+	ExecuteBuildInternal(Context);
+#endif
+
+	if (FApp::CanEverRender())
+	{
+		InitResources();
+	}
+
+#if WITH_EDITOR
+	FinishBuildInternal(Context);
+#endif
+
+	// Re-register any components using this asset to restart the simulation with the updated asset
+	ReregisterComponents();
+}
+
+#if WITH_EDITOR
+void UChaosClothAsset::ExecuteBuildInternal(FSkinnedAssetBuildContext& Context)
+{
+	TRACE_CPUPROFILER_EVENT_SCOPE(UChaosClothAsset::ExecuteBuildInternal);
+
+	// This scope allows us to use any locked properties without causing stalls
+	FSkinnedAssetAsyncBuildScope AsyncBuildScope(this);
+
+	// rebuild render data from imported model
+	CacheDerivedData(&Context);
+
+	// Build the material channel data used by the texture streamer
+	UpdateUVChannelData(true);
+}
+
+void UChaosClothAsset::BeginBuildInternal(FSkinnedAssetBuildContext& Context)
+{
+	TRACE_CPUPROFILER_EVENT_SCOPE(UChaosClothAsset::BeginBuildInternal);
+
+	SetInternalFlags(EInternalObjectFlags::Async);
+	
+	// Unregister all instances of this component
+	Context.RecreateRenderStateContext = MakeUnique<FSkinnedMeshComponentRecreateRenderStateContext>(this, false);
+
+	// Release the render data resources.
+	ReleaseResources();
+
+	// Flush the resource release commands to the rendering thread to ensure that the build doesn't occur while a resource is still
+	// allocated, and potentially accessing the UChaosClothAsset.
+	ReleaseResourcesFence.Wait();
+
+	// Lock all properties that should not be modified/accessed during async post-load
+	AcquireAsyncProperty();
+}
+
+void UChaosClothAsset::FinishBuildInternal(FSkinnedAssetBuildContext& Context)
+{
+	TRACE_CPUPROFILER_EVENT_SCOPE(UChaosClothAsset::FinishBuildInternal);
+
+	ClearInternalFlags(EInternalObjectFlags::Async);
+
+	ReleaseAsyncProperty();
+}
+#endif // #if WITH_EDITOR
+
+#if WITH_EDITOR
+void UChaosClothAsset::BuildMeshModel()
+{
+	using namespace UE::Chaos::ClothAsset;
+
+	const int32 NumLods = ClothCollections.Num();
+
+	// Clear current LOD models
+	check(MeshModel);  // MeshModel should always be created in the Cloth Asset constructor WITH_EDITORONLY_DATA
+	MeshModel->LODModels.Reset(NumLods);
+
+	// Get the running platform
+	ITargetPlatformManagerModule& TargetPlatformManager = GetTargetPlatformManagerRef();
+	ITargetPlatform* const TargetPlatform = TargetPlatformManager.GetRunningTargetPlatform();
+
+	// Rebuild each LOD models
+	for (int32 LodIndex = 0; LodIndex < NumLods; ++LodIndex)
+	{
+		MeshModel->LODModels.Add(new FSkeletalMeshLODModel());
+		BuildLODModel(TargetPlatform, LodIndex);
+	}
+}
+#endif  // #if WITH_EDITOR
+
+void UChaosClothAsset::BuildClothSimulationModel(TArray<FChaosClothAssetLodTransitionDataCache>* InOutTransitionCache)
+{
+	ClothSimulationModel = MakeShared<FChaosClothSimulationModel>(const_cast<const UChaosClothAsset*>(this)->GetClothCollections(), 
+		GetRefSkeleton(), InOutTransitionCache);
+}
+
+const FMeshUVChannelInfo* UChaosClothAsset::GetUVChannelData(int32 MaterialIndex) const
+{
+	if (GetMaterials().IsValidIndex(MaterialIndex))
+	{
+		// TODO: enable ensure when UVChannelData is setup
+		//ensure(GetMaterials()[MaterialIndex].UVChannelData.bInitialized);
+		return &GetMaterials()[MaterialIndex].UVChannelData;
+	}
+
+	return nullptr;
+}
+
+FSkeletalMeshRenderData* UChaosClothAsset::GetResourceForRendering() const
+{
+	WaitUntilAsyncPropertyReleased(EClothAssetAsyncProperties::RenderData);
+	return SkeletalMeshRenderData.Get();
+}
+
+void UChaosClothAsset::SetResourceForRendering(TUniquePtr<FSkeletalMeshRenderData>&& InSkeletalMeshRenderData)
+{
+	WaitUntilAsyncPropertyReleased(EClothAssetAsyncProperties::RenderData);
+	SkeletalMeshRenderData = MoveTemp(InSkeletalMeshRenderData);
+}
+
+void UChaosClothAsset::WaitUntilAsyncPropertyReleased(EClothAssetAsyncProperties AsyncProperties, ESkinnedAssetAsyncPropertyLockType LockType) const
+{
+	// Cast strongly typed enum to uint64
+	WaitUntilAsyncPropertyReleasedInternal((uint64)AsyncProperties, LockType);
+}
+
+FString UChaosClothAsset::GetAsyncPropertyName(uint64 Property) const
+{
+	return StaticEnum<EClothAssetAsyncProperties>()->GetNameByValue(Property).ToString();
+}
+
+int32 UChaosClothAsset::GetMinLodIdx(bool bForceLowestLODIndex) const
+{
+	// #TODO Add quality level controls alongside per-platform taking precedence when enabled
+	return GetMinLod().GetValue();
+}
+
+int32 UChaosClothAsset::GetPlatformMinLODIdx(const ITargetPlatform* InTargetPlatform) const
+{
+	// #TODO Add quality level controls alongside per-platform taking precedence when enabled
+#if WITH_EDITOR
+	return GetMinLod().GetValueForPlatform(*InTargetPlatform->IniPlatformName());
+#else
+	return 0;
+#endif
+}
+
+const FPerPlatformInt& UChaosClothAsset::GetMinLod() const
+{
+	return MinLod;
+}
+
+#if WITH_EDITOR
+void UChaosClothAsset::CacheDerivedData(FSkinnedAssetCompilationContext* Context)
+{
+	TRACE_CPUPROFILER_EVENT_SCOPE(UChaosClothAsset::CacheDerivedData);
+	check(Context);
+
+	// Cache derived data for the running platform.
+	ITargetPlatform* RunningPlatform = GetTargetPlatformManagerRef().GetRunningTargetPlatform();
+	check(RunningPlatform);
+
+	// Create the render data
+	SetResourceForRendering(MakeUnique<FSkeletalMeshRenderData>());
+
+	// Load render data from DDC, or generate it and save to DDC
+	GetResourceForRendering()->Cache(RunningPlatform, this, Context);
+
+}
+
+void UChaosClothAsset::BuildLODModel(const ITargetPlatform* TargetPlatform, int32 LODIndex)
+{
+	check(MeshModel && MeshModel->LODModels.IsValidIndex(LODIndex));
+	FBuilder::BuildLod(MeshModel->LODModels[LODIndex], *this, LODIndex, TargetPlatform);
+}
+
+FString UChaosClothAsset::BuildDerivedDataKey(const ITargetPlatform* TargetPlatform)
+{
+	FString KeySuffix(TEXT(""));
+	KeySuffix += AssetGuid.ToString();
+
+	FString TmpPartialKeySuffix;
+	//Synchronize the user data that are part of the key
+	GetImportedModel()->SyncronizeLODUserSectionsData();
+
+	// Model GUID is not generated so exclude GetImportedModel()->GetIdString() from DDC key.
+
+	// Add the hashed string generated from the model data 
+	TmpPartialKeySuffix = GetImportedModel()->GetLODModelIdString();
+	KeySuffix += TmpPartialKeySuffix;
+
+	//Add the max gpu bone per section
+	const int32 MaxGPUSkinBones = FGPUBaseSkinVertexFactory::GetMaxGPUSkinBones(TargetPlatform);
+	KeySuffix += FString::FromInt(MaxGPUSkinBones);
+	// Add unlimited bone influences mode
+	IMeshBuilderModule::GetForPlatform(TargetPlatform).AppendToDDCKey(KeySuffix, true);
+	const bool bUnlimitedBoneInfluences = FGPUBaseSkinVertexFactory::GetUnlimitedBoneInfluences(TargetPlatform);
+	KeySuffix += bUnlimitedBoneInfluences ? "1" : "0";
+
+	// Include the global default bone influences limit in case any LODs don't set an explicit limit (highly likely)
+	KeySuffix += FString::FromInt(GetDefault<URendererSettings>()->DefaultBoneInfluenceLimit.GetValueForPlatform(*TargetPlatform->IniPlatformName()));
+
+	// Add LODInfoArray
+	TmpPartialKeySuffix = TEXT("");
+	TArray<FSkeletalMeshLODInfo>& LODInfos = GetLODInfoArray();
+	for (int32 LODIndex = 0; LODIndex < GetLODNum(); ++LODIndex)
+	{
+		check(LODInfos.IsValidIndex(LODIndex));
+		FSkeletalMeshLODInfo& LOD = LODInfos[LODIndex];
+		LOD.BuildGUID = LOD.ComputeDeriveDataCacheKey(nullptr);	// TODO: FSkeletalMeshLODGroupSettings
+		TmpPartialKeySuffix += LOD.BuildGUID.ToString(EGuidFormats::Digits);
+	}
+	KeySuffix += TmpPartialKeySuffix;
+
+	//KeySuffix += GetHasVertexColors() ? "1" : "0";
+	//KeySuffix += GetVertexColorGuid().ToString(EGuidFormats::Digits);
+
+	//if (GetEnableLODStreaming(TargetPlatform))
+	//{
+	//	const int32 MaxNumStreamedLODs = GetMaxNumStreamedLODs(TargetPlatform);
+	//	const int32 MaxNumOptionalLODs = GetMaxNumOptionalLODs(TargetPlatform);
+	//	KeySuffix += *FString::Printf(TEXT("1%08x%08x"), MaxNumStreamedLODs, MaxNumOptionalLODs);
+	//}
+	//else
+	//{
+	//	KeySuffix += TEXT("0zzzzzzzzzzzzzzzz");
+	//}
+
+	//if (TargetPlatform->GetPlatformInfo().PlatformGroupName == TEXT("Desktop")
+	//	&& GStripSkeletalMeshLodsDuringCooking != 0
+	//	&& GSkeletalMeshKeepMobileMinLODSettingOnDesktop != 0)
+	//{
+	//	KeySuffix += TEXT("_MinMLOD");
+	//}
+
+	return FDerivedDataCacheInterface::BuildCacheKey(
+		TEXT("CHAOSCLOTH"),
+		CHAOS_CLOTH_ASSET_DERIVED_DATA_VERSION,
+		*KeySuffix
+	);
+}
+
+bool UChaosClothAsset::IsInitialBuildDone() const
+{
+	//We are consider built if we have a valid lod model
+	return GetImportedModel() != nullptr &&
+		GetImportedModel()->LODModels.Num() > 0 &&
+		GetImportedModel()->LODModels[0].Sections.Num() > 0;
+}
+#endif // WITH_EDITOR
+
+void UChaosClothAsset::SetPhysicsAsset(UPhysicsAsset* InPhysicsAsset)
+{
+	using namespace UE::Chaos::ClothAsset;
+
+	PhysicsAsset = InPhysicsAsset;
+	for (TSharedRef<FManagedArrayCollection>& ClothCollection : ClothCollections)
+	{
+		FCollectionClothFacade Cloth(ClothCollection);
+		Cloth.SetPhysicsAssetPathName(PhysicsAsset ? PhysicsAsset->GetPathName() : FString());
+	}
+}
+
+void UChaosClothAsset::SetReferenceSkeleton(const FReferenceSkeleton* ReferenceSkeleton, bool bRebuildModels, bool bRebindMeshes)
+{
+	using namespace UE::Chaos::ClothAsset;
+
+	// Update the reference skeleton
+	if (ReferenceSkeleton)
+	{
+		RefSkeleton = *ReferenceSkeleton;
+	}
+	else
+	{
+		// Create a default reference skeleton
+		RefSkeleton.Empty(1);
+		FReferenceSkeletonModifier ReferenceSkeletonModifier(RefSkeleton, nullptr);
+
+		FMeshBoneInfo MeshBoneInfo;
+		constexpr const TCHAR* RootName = TEXT("Root");
+		MeshBoneInfo.ParentIndex = INDEX_NONE;
+#if WITH_EDITORONLY_DATA
+		MeshBoneInfo.ExportName = RootName;
+#endif
+		MeshBoneInfo.Name = FName(RootName);
+		ReferenceSkeletonModifier.Add(MeshBoneInfo, FTransform::Identity);
+
+		bRebindMeshes = true; // Force the binding when a default reference skeleton is being created 
+	}
+
+	// Rebind the meshes
+	if (bRebindMeshes)
+	{
+		for (TSharedRef<FManagedArrayCollection>& ClothCollection : ClothCollections)
+		{
+			FClothGeometryTools::BindMeshToRootBone(ClothCollection, true, true);
+		}
+	}
+
+	// Rebind the models
+	if (bRebuildModels)
+	{
+		Build();
+	}
+}
+
+void UChaosClothAsset::UpdateSkeletonFromCollection(bool bRebuildModels)
+{
+	using namespace UE::Chaos::ClothAsset;
+
+	check(ClothCollections.Num());
+	FCollectionClothConstFacade ClothFacade(ClothCollections[0]);
+	check(ClothFacade.IsValid());
+
+	const FString& SkeletalMeshPathName = ClothFacade.GetSkeletalMeshPathName();
+	USkeletalMesh* const SkeletalMesh = SkeletalMeshPathName.IsEmpty() ? nullptr :
+		LoadObject<USkeletalMesh>(nullptr, *SkeletalMeshPathName, nullptr, LOAD_None, nullptr);
+
+	SetSkeleton(SkeletalMesh ? SkeletalMesh->GetSkeleton() : nullptr); // For completion only, this is not being used and might mismatch the skeletal mesh's reference skeleton
+
+	constexpr bool bRebindMeshes = false;  // The collection should contain the correct binding at the time SkeletalMeshPathName was set
+	SetReferenceSkeleton(SkeletalMesh ? &SkeletalMesh->GetRefSkeleton() : nullptr, bRebuildModels, bRebindMeshes);
+}
+
+void UChaosClothAsset::CopySimMeshToRenderMesh(UMaterialInterface* Material)
+{
+	using namespace UE::Chaos::ClothAsset;
+	check(ClothCollections.Num());
+
+	// Add a default material if none is specified
+	const FString RenderMaterialPathName = Material ?
+		Material->GetPathName() :
+		FString(TEXT("/Engine/EditorMaterials/Cloth/CameraLitDoubleSided.CameraLitDoubleSided"));
+
+	bool bAnyLodHasRenderMesh = false;
+	for (TSharedRef<FManagedArrayCollection>& ClothCollection : ClothCollections)
+	{
+		constexpr bool bSingleRenderPattern = true;
+		FClothGeometryTools::CopySimMeshToRenderMesh(ClothCollection, RenderMaterialPathName, bSingleRenderPattern);
+		bAnyLodHasRenderMesh = bAnyLodHasRenderMesh || FClothGeometryTools::HasRenderMesh(ClothCollection);
+	}
+
+	// Set new material
+	Materials.Reset(1);
+	if (bAnyLodHasRenderMesh)
+	{
+		if (UMaterialInterface* const LoadedMaterial = LoadObject<UMaterialInterface>(nullptr, *RenderMaterialPathName, nullptr, LOAD_None, nullptr))
+		{
+			Materials.Emplace(LoadedMaterial, true, false, LoadedMaterial->GetFName());
+		}
+	}
+}
+
+void UChaosClothAsset::ReregisterComponents()
+{
+	// Recreate the simulation proxies with the updated physics asset
+	for (TObjectIterator<UChaosClothComponent> ObjectIterator; ObjectIterator; ++ObjectIterator)
+	{
+		if (UChaosClothComponent* const Component = *ObjectIterator)
+		{
+			if (Component->GetClothAsset() == this)
+			{
+				const FComponentReregisterContext Context(Component);  // Context goes out of scope, causing the Component to be re-registered
+			}
+		}
+	}
+}
+
+void UChaosClothAsset::BindSimMeshToRootBone()
+{
+	using namespace UE::Chaos::ClothAsset;
+	check(ClothCollections.Num());
+
+	for (TSharedRef<FManagedArrayCollection>& ClothCollection : ClothCollections)
+	{
+		FClothGeometryTools::BindMeshToRootBone(ClothCollection, true, false);
+	}
+}
+
+#if WITH_EDITORONLY_DATA
+
+void UChaosClothAsset::SetPreviewSceneSkeletalMesh(USkeletalMesh* Mesh)
+{
+	PreviewSceneSkeletalMesh = Mesh;
+}
+
+USkeletalMesh* UChaosClothAsset::GetPreviewSceneSkeletalMesh() const
+{
+	// Load the SkeletalMesh asset if it's not already loaded
+	return PreviewSceneSkeletalMesh.LoadSynchronous();
+}
+
+void UChaosClothAsset::SetPreviewSceneAnimation(UAnimationAsset* Animation)
+{
+	PreviewSceneAnimation = Animation;
+}
+
+UAnimationAsset* UChaosClothAsset::GetPreviewSceneAnimation() const
+{
+	// Load the animation asset if it's not already loaded
+	return PreviewSceneAnimation.LoadSynchronous();
+}
+
+#endif
diff --git a/Engine/Plugins/ChaosClothAsset/Source/ChaosClothAssetEngine/Public/ChaosClothAsset/ClothAsset.h b/Engine/Plugins/ChaosClothAsset/Source/ChaosClothAssetEngine/Public/ChaosClothAsset/ClothAsset.h
new file mode 100644
--- /dev/null
+++ b/Engine/Plugins/ChaosClothAsset/Source/ChaosClothAssetEngine/Public/ChaosClothAsset/ClothAsset.h
@@ -0,0 +1,343 @@
+// Copyright Epic Games, Inc. All Rights Reserved.
+
+#pragma once
+
+#include "Engine/SkinnedAsset.h"
+#include "ReferenceSkeleton.h"
+#include "RenderCommandFence.h"
+#include "ClothAsset.generated.h"
+
+class FSkeletalMeshRenderData;
+class FSkeletalMeshModel;
+class FSkinnedAssetCompilationContext;
+class USkeletalMesh;
+class UDataflow;
+struct FChaosClothSimulationModel;
+struct FSkeletalMeshLODInfo;
+struct FManagedArrayCollection;
+struct FChaosClothAssetLodTransitionDataCache;
+class UAnimationAsset;
+
+UENUM()
+enum class EClothAssetAsyncProperties : uint64
+{
+	None = 0,
+	RenderData = 1 << 0,
+	ThumbnailInfo = 1 << 1,
+	All = MAX_uint64
+};
+ENUM_CLASS_FLAGS(EClothAssetAsyncProperties);
+
+/**
+ * Cloth asset for pattern based simulation.
+ */
+UCLASS(hidecategories = Object, BlueprintType)
+class CHAOSCLOTHASSETENGINE_API UChaosClothAsset : public USkinnedAsset
+{
+	GENERATED_BODY()
+public:
+	UChaosClothAsset(const FObjectInitializer& ObjectInitializer);
+	UChaosClothAsset(FVTableHelper& Helper);  // This is declared so we can use TUniquePtr<FClothSimulationModel> with just a forward declare of that class
+	virtual ~UChaosClothAsset() override;
+
+	//~ Begin UObject interface
+	virtual void Serialize(FArchive& Ar) override;
+	virtual void BeginDestroy() override;
+	virtual bool IsReadyForFinishDestroy() override;
+	virtual void PostLoad() override;
+#if WITH_EDITOR
+	virtual void PostEditChangeProperty(FPropertyChangedEvent& PropertyChangedEvent) override;
+#endif
+	//~ End UObject interface
+
+	//~ Begin USkinnedAsset interface
+	virtual FReferenceSkeleton& GetRefSkeleton()								{ return RefSkeleton; }
+	virtual const FReferenceSkeleton& GetRefSkeleton() const					{ return RefSkeleton; }
+	virtual FSkeletalMeshLODInfo* GetLODInfo(int32 Index);
+	virtual const FSkeletalMeshLODInfo* GetLODInfo(int32 Index) const;
+	UFUNCTION(BlueprintGetter)
+	virtual class UPhysicsAsset* GetShadowPhysicsAsset() const override			{ return ShadowPhysicsAsset; }
+	virtual FMatrix GetComposedRefPoseMatrix(FName InBoneName) const override;
+	virtual FMatrix GetComposedRefPoseMatrix(int32 InBoneIndex) const override	{ return CachedComposedRefPoseMatrices[InBoneIndex]; }
+	virtual const FMeshUVChannelInfo* GetUVChannelData(int32 MaterialIndex) const override;
+	virtual bool GetSupportRayTracing() const override							{ return bSupportRayTracing; }
+	virtual int32 GetRayTracingMinLOD() const override							{ return RayTracingMinLOD; }
+	virtual TArray<FMatrix44f>& GetRefBasesInvMatrix() override					{ return RefBasesInvMatrix; }
+	virtual const TArray<FMatrix44f>& GetRefBasesInvMatrix() const override		{ return RefBasesInvMatrix; }
+	virtual TArray<FSkeletalMeshLODInfo>& GetLODInfoArray() override			{ return LODInfo; }
+	virtual const TArray<FSkeletalMeshLODInfo>& GetLODInfoArray() const override { return LODInfo; }
+	virtual FSkeletalMeshRenderData* GetResourceForRendering() const override;
+	virtual int32 GetDefaultMinLod() const										{ return 0; }
+	virtual UPhysicsAsset* GetPhysicsAsset() const								{ return PhysicsAsset; }
+	virtual TArray<FSkeletalMaterial>& GetMaterials() override					{ return Materials; }
+	virtual const TArray<FSkeletalMaterial>& GetMaterials() const override		{ return Materials; }
+	virtual int32 GetLODNum() const override									{ return LODInfo.Num(); }
+	virtual bool IsMaterialUsed(int32 MaterialIndex) const override				{ return true; }
+	virtual FBoxSphereBounds GetBounds() const override							{ return Bounds; }
+	virtual TArray<class USkeletalMeshSocket*> GetActiveSocketList() const override { static TArray<class USkeletalMeshSocket*> Dummy; return Dummy; }
+	virtual USkeletalMeshSocket* FindSocket(FName InSocketName) const override	{ return nullptr; }
+	virtual USkeletalMeshSocket* FindSocketInfo(FName InSocketName, FTransform& OutTransform, int32& OutBoneIndex, int32& OutIndex) const override { return nullptr; }
+	virtual USkeleton* GetSkeleton() override									{ return Skeleton; }  // Note: The USkeleton isn't a reliable source of reference skeleton
+	virtual const USkeleton* GetSkeleton() const override						{ return Skeleton; }
+	virtual void SetSkeleton(USkeleton* InSkeleton) override					{ Skeleton = InSkeleton; }
+	virtual UMeshDeformer* GetDefaultMeshDeformer() const override				{ return nullptr; }
+	virtual class UMaterialInterface* GetOverlayMaterial() const override		{ return nullptr; }
+	virtual float GetOverlayMaterialMaxDrawDistance() const override			{ return 0.f; }
+	virtual bool IsValidLODIndex(int32 Index) const override					{ return LODInfo.IsValidIndex(Index); }
+	virtual int32 GetMinLodIdx(bool bForceLowestLODIdx = false) const override;
+	virtual bool NeedCPUData(int32 LODIndex) const override						{ return false; }
+	virtual bool GetHasVertexColors() const override							{ return false; }
+	virtual int32 GetPlatformMinLODIdx(const ITargetPlatform* TargetPlatform) const override;
+	virtual const FPerPlatformBool& GetDisableBelowMinLodStripping() const override { return DisableBelowMinLodStripping; }
+	virtual const FPerPlatformInt& GetMinLod() const override;
+#if WITH_EDITOR
+	/* Build a LOD model for the targeted platform. */
+	virtual void BuildLODModel(const ITargetPlatform* TargetPlatform, int32 LODIndex) override;
+	virtual FString BuildDerivedDataKey(const ITargetPlatform* TargetPlatform) override;
+	virtual bool IsInitialBuildDone() const override;
+	virtual bool GetEnableLODStreaming(const class ITargetPlatform* TargetPlatform) const override	{ return false; }
+	virtual int32 GetMaxNumStreamedLODs(const class ITargetPlatform* TargetPlatform) const override	{ return 0; }
+	virtual int32 GetMaxNumOptionalLODs(const ITargetPlatform* TargetPlatform) const override { return 0; }
+#endif
+#if WITH_EDITORONLY_DATA
+	virtual FSkeletalMeshModel* GetImportedModel() const override				{ return MeshModel.Get(); }
+#endif
+	//~ End USkinnedAsset interface
+
+	/** Return the enclosed Cloth Collection object. */
+	TArray<TSharedRef<FManagedArrayCollection>>& GetClothCollections() { return ClothCollections; }
+
+	/** Return the enclosed Cloth Collection object, const version. */
+	const TArray<TSharedRef<const FManagedArrayCollection>>& GetClothCollections() const { return reinterpret_cast<const TArray<TSharedRef<const FManagedArrayCollection>>&>(ClothCollections); }
+
+	/** Return the cloth simulation ready LOD model data. */
+	TSharedPtr<const FChaosClothSimulationModel> GetClothSimulationModel() const { return ClothSimulationModel; }
+
+	/** Build this asset static render and simulation data. This needs to be done every time the asset has changed. */
+	void Build(TArray<FChaosClothAssetLodTransitionDataCache>* InOutTransitionCache = nullptr);
+
+	/**
+	 * Copy the draped simulation mesh patterns into the render mesh data.
+	 * This is useful to visualize the simulation mesh, or when the simulation mesh can be used for both simulation and rendering.
+	 * @param Material The material used to render. If none is specified the default debug cloth material will be used.
+	 */
+	void CopySimMeshToRenderMesh(UMaterialInterface* Material = nullptr);
+
+	/** Set the physics asset for this cloth. */
+	void SetPhysicsAsset(UPhysicsAsset* InPhysicsAsset);
+
+	/**
+	 * Set the skeleton asset using the LOD0 reference skeleton of the skeletal mesh specified in the cloth collection, or set a default reference skeleton if there isn't any.
+	 * @param bRebuildModels Whether to rebuild the simulation and rendering models (editor builds only).
+	 */
+	void UpdateSkeletonFromCollection(bool bRebuildModels);
+
+	/**
+	 * Set the specified reference skeleton for this cloth, and rebuild the models.
+	 * Use nullptr to set to a default single bone root reference skeleton.
+	 * @param bRebuildModels Whether to rebuild the mesh static models. This is always required, but could be skipped if done soon after.
+	 * @param bRebindMeshes Bindings could be invalid after a change of reference skeleton, use this option to clear the existing binding and bind the sim and render meshes to the root bone.
+	 */
+	void SetReferenceSkeleton(const FReferenceSkeleton* ReferenceSkeleton, bool bRebuildModels = true, bool bRebindMeshes = true);
+
+	/** Set the bone hierachy to use for this cloth. */
+	UE_DEPRECATED(5.3, "Use SetReferenceSkeleton(const FReferenceSkeleton*, bool, bool) instead")
+	PRAGMA_DISABLE_DEPRECATION_WARNINGS
+	void SetReferenceSkeleton(const FReferenceSkeleton& InReferenceSkeleton, bool bRebuildClothSimulationModel = true) { RefSkeleton = InReferenceSkeleton; UpdateSkeleton(bRebuildClothSimulationModel); }
+	PRAGMA_ENABLE_DEPRECATION_WARNINGS
+
+	/** Set the skinning weights for all of the sim vertices to be bound to the root node of the reference skeleton. */
+	UE_DEPRECATED(5.3, "Use FClothGeometryTools::BindMeshToRootBone or SetReferenceSkeleton(const FReferenceSkeleton*, bool, bool) instead.")
+	void BindSimMeshToRootBone();
+
+	//
+	// Dataflow
+	//
+	UE_DEPRECATED(5.3, "Do not use. Will be made private in 5.4")
+	UPROPERTY(EditAnywhere, Category = "Dataflow")
+	TObjectPtr<UDataflow> DataflowAsset;
+
+	UE_DEPRECATED(5.3, "Do not use. Will be made private in 5.4")
+	UPROPERTY(EditAnywhere, Category = "Dataflow")
+	FString DataflowTerminal = "ClothAssetTerminal";
+
+
+#if WITH_EDITORONLY_DATA
+
+	void SetPreviewSceneSkeletalMesh(USkeletalMesh* Mesh);
+	USkeletalMesh* GetPreviewSceneSkeletalMesh() const;
+
+	void SetPreviewSceneAnimation(UAnimationAsset* Animation);
+	UAnimationAsset* GetPreviewSceneAnimation() const;
+
+#endif
+
+private:
+	//~ Begin USkinnedAsset interface
+	/** Initial step for the Post Load process - Can't be done in parallel. */
+	virtual void BeginPostLoadInternal(FSkinnedAssetPostLoadContext& Context) override;
+	/** Thread-safe part of the Post Load */
+	virtual void ExecutePostLoadInternal(FSkinnedAssetPostLoadContext& Context) override;
+	/** Complete the postload process - Can't be done in parallel. */
+	virtual void FinishPostLoadInternal(FSkinnedAssetPostLoadContext& Context) override;
+	/** Convert async property from enum value to string. */
+	virtual FString GetAsyncPropertyName(uint64 Property) const override;
+	//~ End USkinnedAsset interface
+
+	/**
+	 * Wait for the asset to finish compilation to protect internal skinned asset data from race conditions during async build.
+	 * This should be called before accessing all async accessible properties.
+	 */
+	void WaitUntilAsyncPropertyReleased(EClothAssetAsyncProperties AsyncProperties, ESkinnedAssetAsyncPropertyLockType LockType = ESkinnedAssetAsyncPropertyLockType::ReadWrite) const;
+
+	/** Pre-calculate refpose-to-local transforms. */
+	void CalculateInvRefMatrices();
+
+	/** Re-calculate the bounds for this asset. */
+	void CalculateBounds();
+
+	/** Update the bone informations after a change of skeleton. */
+	UE_DEPRECATED(5.3, "Use Build() instead")
+	void UpdateSkeleton(bool bRebuildClothSimulationModel = true);
+
+#if WITH_EDITOR
+	/** Build the SkeletalMeshLODModel for this asset. */
+	void BuildMeshModel();
+#endif
+
+	/** Build the clothing simulation meshes from the Cloth Collection. */
+	void BuildClothSimulationModel(TArray<FChaosClothAssetLodTransitionDataCache>* InOutTransitionCache = nullptr);
+
+	/** Initialize all render resources. */
+	void InitResources();
+
+	/** Safely release the render data. */
+	void ReleaseResources();
+
+	/** Set render data. */
+	void SetResourceForRendering(TUniquePtr<FSkeletalMeshRenderData>&& InSkeletalMeshRenderData);
+
+#if WITH_EDITOR
+	/** Load render data from DDC if the data is cached, otherwise generate render data and save into DDC */
+	void CacheDerivedData(FSkinnedAssetCompilationContext* Context);
+
+	/** Initial step for the building process - Can't be done in parallel. USkinnedAsset Interface. */
+	virtual void BeginBuildInternal(FSkinnedAssetBuildContext& Context) override;
+
+	/** Thread-safe part. USkinnedAsset Interface. */
+	virtual void ExecuteBuildInternal(FSkinnedAssetBuildContext& Context) override;
+
+	/** Complete the building process - Can't be done in parallel. USkinnedAsset Interface. */
+	virtual void FinishBuildInternal(FSkinnedAssetBuildContext& Context) override;
+#endif
+
+	/** Reregister all components using this asset to reset the simulation in case anything has changed. */
+	void ReregisterComponents();
+
+	/** List of materials for this cloth asset. */
+	UPROPERTY(EditAnywhere, Category = Materials)
+	TArray<FSkeletalMaterial> Materials;
+
+	/** Skeleton asset used at creation time. This is of limited use since this USkeleton's reference skeleton might not necessarily match the one created for this asset. */
+	UPROPERTY(EditAnywhere, Setter = SetSkeleton, Category = Skeleton)
+	TObjectPtr<USkeleton> Skeleton;
+
+	/** Physics asset used for collision. */
+	UPROPERTY(EditAnywhere, Category = Collision)
+	TObjectPtr<UPhysicsAsset> PhysicsAsset;
+
+	/** Struct containing information for each LOD level, such as materials to use, and when use the LOD. */
+	UPROPERTY(EditAnywhere, EditFixedSize, Category = LevelOfDetails)
+	TArray<FSkeletalMeshLODInfo> LODInfo;
+
+	UPROPERTY(EditAnywhere, Category = LODSettings)
+	FPerPlatformBool DisableBelowMinLodStripping;
+
+	UPROPERTY(EditAnywhere, Category = LODSettings, meta = (DisplayName = "Minimum LOD"))
+	FPerPlatformInt MinLod;
+
+	/** Enable raytracing for this asset. */
+	UPROPERTY(EditAnywhere, Category = RayTracing)
+	uint8 bSupportRayTracing : 1;
+
+	/** Minimum raytracing LOD for this asset. */
+	UPROPERTY(EditAnywhere, Category = RayTracing)
+	int32 RayTracingMinLOD;
+
+	/** Whether to blend positions between the skinned/simulated transitions of the cloth render mesh. */
+	UPROPERTY(EditAnywhere, Category = ClothDeformer)
+	bool bSmoothTransition = true;
+
+	/** Whether to use multiple triangle influences on the proxy wrap deformer to help smoothe deformations. */
+	UPROPERTY(EditAnywhere, Category = ClothDeformer)
+	bool bUseMultipleInfluences = false;
+
+	/** The radius from which to get the multiple triangle influences from the simulated proxy mesh. */
+	UPROPERTY(EditAnywhere, Category = ClothDeformer)
+	float SkinningKernelRadius = 30.f;
+
+	/**
+	 * Physics asset whose shapes will be used for shadowing when components have bCastCharacterCapsuleDirectShadow or bCastCharacterCapsuleIndirectShadow enabled.
+	 * Only spheres and sphyl shapes in the physics asset can be supported.  The more shapes used, the higher the cost of the capsule shadows will be.
+	 */
+	UPROPERTY(EditAnywhere, AssetRegistrySearchable, BlueprintGetter = GetShadowPhysicsAsset, Category = Lighting)
+	TObjectPtr<class UPhysicsAsset> ShadowPhysicsAsset;
+
+	/** A unique identifier as used by the section rendering code. */
+	UPROPERTY()
+	FGuid AssetGuid;
+
+	/** Bounds for this asset. */
+	FBoxSphereBounds Bounds;
+
+	/** Mesh-space ref pose, where parent matrices are applied to ref pose matrices. */
+	TArray<FMatrix> CachedComposedRefPoseMatrices;
+
+	/** Cloth Collection containing this asset data. One per LOD. */
+	TArray<TSharedRef<FManagedArrayCollection>> ClothCollections;
+
+	/** Reference skeleton created from the provided skeleton asset. */
+	FReferenceSkeleton RefSkeleton;
+
+	/** Reference skeleton precomputed bases. */
+	TArray<FMatrix44f> RefBasesInvMatrix;
+
+	/** Rendering data. */
+	TUniquePtr<FSkeletalMeshRenderData> SkeletalMeshRenderData;
+
+	/** A fence which is used to keep track of the rendering thread releasing the static mesh resources. */
+	FRenderCommandFence ReleaseResourcesFence;
+
+#if WITH_EDITORONLY_DATA
+	/** Source mesh geometry information (not used at runtime). */
+	TSharedPtr<FSkeletalMeshModel> MeshModel;
+#endif
+
+	/** Simulation mesh Lods as fed to the solver for constraints creation. Ownership gets transferred to the proxy when it is changed during a simulation. */
+	TSharedPtr<FChaosClothSimulationModel> ClothSimulationModel;
+
+
+#if WITH_EDITORONLY_DATA
+
+	/*
+	* The following PreviewScene properties are modeled after PreviewSkeletalMesh in USkeleton
+	*	- they are inside WITH_EDITORONLY_DATA because they are not used at game runtime
+	*	- TSoftObjectPtrs since that will make it possible to avoid loading these assets until the PreviewScene asks for them
+	*	- DuplicateTransient so that if you copy a ClothAsset it won't copy these preview properties
+	*	- AssetRegistrySearchable makes it so that if the user searches the name of a PreviewScene asset in the Asset Browser, it will return any ClothAssets that use it
+	*/
+
+	/** Optional Skeletal Mesh that the cloth asset is attached to in the Preview Scene in the Cloth Editor */
+	UPROPERTY(DuplicateTransient, AssetRegistrySearchable)
+	TSoftObjectPtr<USkeletalMesh> PreviewSceneSkeletalMesh;
+
+	/** Optional animation attached to PreviewSceneSkeletalMesh in the Preview Scene in the Cloth Editor */
+	UPROPERTY(DuplicateTransient, AssetRegistrySearchable)
+	TSoftObjectPtr<UAnimationAsset> PreviewSceneAnimation;
+
+#endif
+
+#if WITH_EDITOR
+	struct FBuilder;
+#endif
+};
diff --git a/Engine/Plugins/ChaosClothAssetEditor/ChaosClothAssetEditor.uplugin b/Engine/Plugins/ChaosClothAssetEditor/ChaosClothAssetEditor.uplugin
new file mode 100644
--- /dev/null
+++ b/Engine/Plugins/ChaosClothAssetEditor/ChaosClothAssetEditor.uplugin
@@ -0,0 +1,56 @@
+{
+	"FileVersion": 1,
+	"Version": 1,
+	"VersionName": "0.1",
+	"FriendlyName": "Chaos Cloth Asset Editor",
+	"Description": "Editor for modifying cloth assets",
+	"Category": "Other",
+	"CreatedBy": "Epic Games, Inc.",
+	"CreatedByURL": "https://epicgames.com",
+	"DocsURL": "",
+	"MarketplaceURL": "",
+	"SupportURL": "",
+	"CanContainContent": true,
+	"IsBetaVersion": true,
+	"Installed": false,
+	"Modules": [
+		{
+			"Name": "ChaosClothAssetEditor",
+			"Type": "Editor",
+			"LoadingPhase": "Default",
+			"PlatformAllowList": [ "Win64" ]
+		},
+		{
+			"Name": "ChaosClothAssetEditorTools",
+			"Type": "Editor",
+			"LoadingPhase": "Default",
+			"PlatformAllowList": [ "Win64" ]
+		}
+	],
+	"Plugins": [
+		{
+			"Name": "ChaosClothAsset",
+			"Enabled": true
+		},
+		{
+			"Name": "ChaosCloth",
+			"Enabled": true
+		},
+		{
+			"Name": "BaseCharacterFXEditor",
+			"Enabled": true,
+			"TargetAllowList": [ "Editor" ]
+		},
+		{
+			"Name": "Dataflow",
+			"Enabled": true
+		},
+		{
+			"Name": "DatasmithCloTranslator",
+			"Enabled": true,
+			"SupportedTargetPlatforms": [
+				"Win64"
+			]
+		}
+	]
+}
\ No newline at end of file
diff --git a/Engine/Plugins/ChaosClothAssetEditor/Source/ChaosClothAssetEditor/Private/ChaosClothAsset/ClothEditorCommands.cpp b/Engine/Plugins/ChaosClothAssetEditor/Source/ChaosClothAssetEditor/Private/ChaosClothAsset/ClothEditorCommands.cpp
new file mode 100644
--- /dev/null
+++ b/Engine/Plugins/ChaosClothAssetEditor/Source/ChaosClothAssetEditor/Private/ChaosClothAsset/ClothEditorCommands.cpp
@@ -0,0 +1,96 @@
+// Copyright Epic Games, Inc. All Rights Reserved.
+
+#include "ChaosClothAsset/ClothEditorCommands.h"
+#include "ChaosClothAsset/ClothEditorStyle.h"
+#include "ChaosClothAsset/ClothWeightMapPaintTool.h"
+#include "ChaosClothAsset/ClothMeshSelectionTool.h"
+
+#define LOCTEXT_NAMESPACE "FChaosClothAssetEditorCommands"
+
+namespace UE::Chaos::ClothAsset
+{
+
+const FString FChaosClothAssetEditorCommands::BeginRemeshToolIdentifier = TEXT("BeginRemeshTool");
+const FString FChaosClothAssetEditorCommands::BeginAttributeEditorToolIdentifier = TEXT("BeginAttributeEditorTool");
+const FString FChaosClothAssetEditorCommands::BeginWeightMapPaintToolIdentifier = TEXT("BeginWeightMapPaintTool");
+const FString FChaosClothAssetEditorCommands::AddWeightMapNodeIdentifier = TEXT("AddWeightMapNode");
+const FString FChaosClothAssetEditorCommands::BeginTransferSkinWeightsToolIdentifier = TEXT("BeginTransferSkinWeightsTool");
+const FString FChaosClothAssetEditorCommands::AddTransferSkinWeightsNodeIdentifier = TEXT("AddTransferSkinWeightsNode");
+const FString FChaosClothAssetEditorCommands::BeginMeshSelectionToolIdentifier = TEXT("BeginMeshSelectionTool");
+const FString FChaosClothAssetEditorCommands::AddMeshSelectionNodeIdentifier = TEXT("AddMeshSelectionNode");
+const FString FChaosClothAssetEditorCommands::ToggleSimulationSuspendedIdentifier = TEXT("ToggleSimulationSuspended");
+const FString FChaosClothAssetEditorCommands::SoftResetSimulationIdentifier = TEXT("SoftResetSimulation");
+const FString FChaosClothAssetEditorCommands::HardResetSimulationIdentifier = TEXT("HardResetSimulation");
+const FString FChaosClothAssetEditorCommands::TogglePreviewWireframeIdentifier = TEXT("TogglePreviewWireframe");
+const FString FChaosClothAssetEditorCommands::ToggleConstructionViewWireframeIdentifier = TEXT("ToggleConstructionViewWireframe");
+const FString FChaosClothAssetEditorCommands::ToggleConstructionViewSeamsIdentifier = TEXT("ToggleConstructionViewSeams");
+const FString FChaosClothAssetEditorCommands::ToggleConstructionViewSeamsCollapseIdentifier = TEXT("ToggleConstructionViewSeamsCollapse");
+
+FChaosClothAssetEditorCommands::FChaosClothAssetEditorCommands()
+	: TBaseCharacterFXEditorCommands<FChaosClothAssetEditorCommands>("ChaosClothAssetEditor",
+		LOCTEXT("ContextDescription", "Cloth Editor"), 
+		NAME_None, // Parent
+		FChaosClothAssetEditorStyle::Get().GetStyleSetName())
+{
+}
+
+void FChaosClothAssetEditorCommands::RegisterCommands()
+{
+	TBaseCharacterFXEditorCommands::RegisterCommands();
+
+	UI_COMMAND(OpenClothEditor, "Cloth Editor", "Open the Cloth Editor window", EUserInterfaceActionType::Button, FInputChord());
+
+	UI_COMMAND(BeginRemeshTool, "Remesh", "Remesh the selected mesh", EUserInterfaceActionType::Button, FInputChord());
+	UI_COMMAND(BeginAttributeEditorTool, "AttrEd", "Edit/configure mesh attributes", EUserInterfaceActionType::Button, FInputChord());
+	
+	UI_COMMAND(BeginWeightMapPaintTool, "Add Weight Map", "Paint weight maps on the mesh", EUserInterfaceActionType::None, FInputChord());
+	UI_COMMAND(AddWeightMapNode, "Add Weight Map", "Paint weight maps on the mesh", EUserInterfaceActionType::Button, FInputChord());
+	UI_COMMAND(BeginMeshSelectionTool, "Select", "Select mesh elements", EUserInterfaceActionType::None, FInputChord());
+	UI_COMMAND(AddMeshSelectionNode, "Select", "Select mesh elements", EUserInterfaceActionType::Button, FInputChord());
+
+	UI_COMMAND(BeginTransferSkinWeightsTool, "Transfer Skin Weights", "Transfer skinning weights from a SkeletalMesh", EUserInterfaceActionType::None, FInputChord());
+	UI_COMMAND(AddTransferSkinWeightsNode, "Transfer Skin Weights", "Transfer skinning weights from a SkeletalMesh", EUserInterfaceActionType::Button, FInputChord());
+
+	UI_COMMAND(SetConstructionMode2D, "2D Sim", "Switches the viewport to 2D simulation mesh view", EUserInterfaceActionType::RadioButton, FInputChord());
+	UI_COMMAND(SetConstructionMode3D, "3D Sim", "Switches the viewport to 3D simulation mesh view", EUserInterfaceActionType::RadioButton, FInputChord());
+	UI_COMMAND(SetConstructionModeRender, "Render", "Switches the viewport to render mesh view", EUserInterfaceActionType::RadioButton, FInputChord());
+
+	UI_COMMAND(TogglePreviewWireframe, "TogglePreviewWireframe", "Toggle preview wireframe", EUserInterfaceActionType::ToggleButton, FInputChord());
+	UI_COMMAND(ToggleConstructionViewWireframe, "ToggleConstructionViewWireframe", "Toggle construction view wireframe", EUserInterfaceActionType::ToggleButton, FInputChord());
+
+	UI_COMMAND(ToggleConstructionViewSeams, "ToggleConstructionViewSeams", "Toggle construction view seams", EUserInterfaceActionType::ToggleButton, FInputChord());
+	UI_COMMAND(ToggleConstructionViewSeamsCollapse, "ToggleConstructionViewSeamsCollapse", "Toggle drawing one line between seams in construction view", EUserInterfaceActionType::ToggleButton, FInputChord());
+
+	UI_COMMAND(TogglePatternColor, "TogglePatternColor", "Toggle Pattern color in the construction view", EUserInterfaceActionType::ToggleButton, FInputChord());
+	UI_COMMAND(ToggleMeshStats, "ToggleMeshStats", "Toggle mesh information display", EUserInterfaceActionType::ToggleButton, FInputChord());
+	
+	UI_COMMAND(SoftResetSimulation, "SoftResetSimulation", "Soft reset simulation", EUserInterfaceActionType::ToggleButton, FInputChord());
+	UI_COMMAND(HardResetSimulation, "HardResetClothSimulation", "Hard reset simulation", EUserInterfaceActionType::ToggleButton, FInputChord(EModifierKey::Control | EModifierKey::Alt, EKeys::C));
+	UI_COMMAND(ToggleSimulationSuspended, "ToggleSimulationSuspended", "Toggle simulation suspended", EUserInterfaceActionType::ToggleButton, FInputChord());
+
+	UI_COMMAND(LODAuto, "LOD Auto", "Automatically select LOD", EUserInterfaceActionType::RadioButton, FInputChord());
+	UI_COMMAND(LOD0, "LOD 0", "Force select LOD 0", EUserInterfaceActionType::RadioButton, FInputChord());
+}
+
+void FChaosClothAssetEditorCommands::GetToolDefaultObjectList(TArray<UInteractiveTool*>& ToolCDOs)
+{
+	ToolCDOs.Add(GetMutableDefault<UClothEditorWeightMapPaintTool>());
+}
+
+void FChaosClothAssetEditorCommands::UpdateToolCommandBinding(UInteractiveTool* Tool, TSharedPtr<FUICommandList> UICommandList, bool bUnbind)
+{
+	if (FChaosClothAssetEditorCommands::IsRegistered())
+	{
+		if (bUnbind)
+		{
+			FChaosClothAssetEditorCommands::Get().UnbindActiveCommands(UICommandList);
+		}
+		else
+		{
+			FChaosClothAssetEditorCommands::Get().BindCommandsForCurrentTool(UICommandList, Tool);
+		}
+	}
+}
+} // namespace UE::Chaos::ClothAsset
+
+#undef LOCTEXT_NAMESPACE
diff --git a/Engine/Plugins/ChaosClothAssetEditor/Source/ChaosClothAssetEditor/Private/ChaosClothAsset/SClothEditor3DViewportToolBar.cpp b/Engine/Plugins/ChaosClothAssetEditor/Source/ChaosClothAssetEditor/Private/ChaosClothAsset/SClothEditor3DViewportToolBar.cpp
new file mode 100644
--- /dev/null
+++ b/Engine/Plugins/ChaosClothAssetEditor/Source/ChaosClothAssetEditor/Private/ChaosClothAsset/SClothEditor3DViewportToolBar.cpp
@@ -0,0 +1,236 @@
+// Copyright Epic Games, Inc. All Rights Reserved.
+
+#include "ChaosClothAsset/SClothEditor3DViewportToolBar.h"
+#include "ChaosClothAsset/SClothEditor3DViewport.h"
+#include "ChaosClothAsset/ClothEditor3DViewportClient.h"
+#include "ChaosClothAsset/ClothEditorSimulationVisualization.h"
+#include "ChaosClothAsset/ClothEditorCommands.h"
+#include "ChaosClothAsset/ClothEditorStyle.h"
+#include "ChaosClothAsset/ClothEditorPreviewScene.h"
+#include "ChaosClothAsset/ClothAsset.h"
+#include "ChaosClothAsset/ClothComponent.h"
+#include "EditorViewportCommands.h"
+#include "Styling/AppStyle.h"
+#include "Framework/MultiBox/MultiBoxBuilder.h"
+#include "Widgets/Layout/SBorder.h"
+#include "Widgets/SBoxPanel.h"
+#include "Widgets/Text/SRichTextBlock.h"
+#include "Editor/UnrealEd/Public/SEditorViewportToolBarMenu.h"
+#include "SEditorViewportViewMenu.h"
+
+
+#define LOCTEXT_NAMESPACE "SChaosClothAssetEditor3DViewportToolBar"
+
+void SChaosClothAssetEditor3DViewportToolBar::Construct(const FArguments& InArgs, TSharedPtr<class SChaosClothAssetEditor3DViewport> InChaosClothAssetEditor3DViewport)
+{
+	SCommonEditorViewportToolbarBase::Construct(SCommonEditorViewportToolbarBase::FArguments(), InChaosClothAssetEditor3DViewport);
+
+	ChaosClothAssetEditor3DViewportPtr = InChaosClothAssetEditor3DViewport;
+	CommandList = InArgs._CommandList;
+
+	const FMargin ToolbarSlotPadding(4.0f, 1.0f);
+	TSharedPtr<SHorizontalBox> MainBoxPtr;
+
+	ChildSlot
+	[
+		SNew( SBorder )
+		.BorderImage(FAppStyle::Get().GetBrush("EditorViewportToolBar.Background"))
+		.Cursor(EMouseCursor::Default)
+		[
+			SNew(SVerticalBox)
+			+SVerticalBox::Slot()
+			.AutoHeight()
+			[
+				SAssignNew( MainBoxPtr, SHorizontalBox )
+			]
+			+SVerticalBox::Slot()
+			.Padding(FMargin(4.0f, 3.0f, 0.0f, 0.0f))
+			[
+				// Display text (e.g., item being previewed)
+				SNew(SRichTextBlock)
+				.DecoratorStyleSet(&FAppStyle::Get())
+				.Text(this, &SChaosClothAssetEditor3DViewportToolBar::GetDisplayString)
+				.TextStyle(&FAppStyle::Get().GetWidgetStyle<FTextBlockStyle>("AnimViewport.MessageText"))
+			]
+		]
+	];
+
+	MainBoxPtr->AddSlot()
+		.AutoWidth()
+		.Padding(ToolbarSlotPadding)
+		[
+			MakeOptionsMenu()
+		];
+
+	MainBoxPtr->AddSlot()
+		.AutoWidth()
+		.Padding(ToolbarSlotPadding)
+		.HAlign(HAlign_Left)
+		[
+			MakeDisplayToolBar(InArgs._Extenders)
+		];
+
+	MainBoxPtr->AddSlot()
+		.AutoWidth()
+		.Padding(ToolbarSlotPadding)
+		.HAlign(HAlign_Left)
+		[
+			SNew(SEditorViewportToolbarMenu)
+			.Label(this, &SChaosClothAssetEditor3DViewportToolBar::GetLODMenuLabel)
+			.Cursor(EMouseCursor::Default)
+			.ParentToolBar(SharedThis(this))
+			.OnGetMenuContent(this, &SChaosClothAssetEditor3DViewportToolBar::MakeLODMenu)
+		];
+
+	MainBoxPtr->AddSlot()
+		.Padding(ToolbarSlotPadding)
+		.HAlign(HAlign_Right)
+		[
+			MakeToolBar(InArgs._Extenders)
+		];
+}
+
+TSharedRef<SWidget> SChaosClothAssetEditor3DViewportToolBar::MakeOptionsMenu()
+{
+	return SNew(SEditorViewportToolbarMenu)
+		.ParentToolBar(SharedThis(this))
+		.Cursor(EMouseCursor::Default)
+		.Image("EditorViewportToolBar.OptionsDropdown")
+		.OnGetMenuContent(this, &SChaosClothAssetEditor3DViewportToolBar::GenerateClothViewportOptionsMenu);
+}
+
+TSharedRef<SWidget> SChaosClothAssetEditor3DViewportToolBar::MakeDisplayToolBar(const TSharedPtr<FExtender> InExtenders)
+{
+	TSharedRef<SEditorViewport> ViewportRef = StaticCastSharedPtr<SEditorViewport>(ChaosClothAssetEditor3DViewportPtr.Pin()).ToSharedRef();
+
+	return SNew(SEditorViewportViewMenu, ViewportRef, SharedThis(this))
+		.Cursor(EMouseCursor::Default)
+		.MenuExtenders(InExtenders);
+}
+
+TSharedRef<SWidget> SChaosClothAssetEditor3DViewportToolBar::MakeToolBar(const TSharedPtr<FExtender> InExtenders)
+{
+	using namespace UE::Chaos::ClothAsset;
+
+	// The following is modeled after portions of STransformViewportToolBar, which gets 
+	// used in SCommonEditorViewportToolbarBase.
+
+	// The buttons are hooked up to actual functions via command bindings in SChaosClothAssetEditor3DViewport::BindCommands(),
+	// and the toolbar gets built in SChaosClothAssetEditor3DViewport::MakeViewportToolbar().
+
+	FSlimHorizontalToolBarBuilder ToolbarBuilder(CommandList, FMultiBoxCustomization::None, InExtenders);
+
+	// Use a custom style
+	FName ToolBarStyle = "EditorViewportToolBar";
+	ToolbarBuilder.SetStyle(&FAppStyle::Get(), ToolBarStyle);
+	ToolbarBuilder.SetLabelVisibility(EVisibility::Collapsed);
+
+	ToolbarBuilder.BeginSection("Visualization");
+	ToolbarBuilder.BeginBlockGroup();
+	{
+		// TODO: Add button to toggle between sim and render meshes
+
+		ToolbarBuilder.AddToolBarButton(FChaosClothAssetEditorCommands::Get().TogglePreviewWireframe,
+			NAME_None,
+			TAttribute<FText>(),
+			TAttribute<FText>(),
+			TAttribute<FSlateIcon>(),
+			TEXT("ToggleRenderMeshWireframe"));
+	}
+	ToolbarBuilder.EndBlockGroup();
+	ToolbarBuilder.EndSection();
+
+	ToolbarBuilder.BeginSection("Sim Controls");
+	ToolbarBuilder.BeginBlockGroup();
+	{
+		ToolbarBuilder.AddToolBarButton(FChaosClothAssetEditorCommands::Get().HardResetSimulation,
+			NAME_None,
+			TAttribute<FText>(),
+			TAttribute<FText>(),
+			TAttribute<FSlateIcon>(),
+			FName(*FChaosClothAssetEditorCommands::HardResetSimulationIdentifier));
+
+		ToolbarBuilder.AddToolBarButton(FChaosClothAssetEditorCommands::Get().SoftResetSimulation,
+			NAME_None,
+			TAttribute<FText>(),
+			TAttribute<FText>(),
+			TAttribute<FSlateIcon>(),
+			FName(FChaosClothAssetEditorCommands::SoftResetSimulationIdentifier));
+
+		ToolbarBuilder.AddToolBarButton(FChaosClothAssetEditorCommands::Get().ToggleSimulationSuspended,
+			NAME_None,
+			TAttribute<FText>(),
+			TAttribute<FText>(),
+			TAttribute<FSlateIcon>(),
+			FName(FChaosClothAssetEditorCommands::ToggleSimulationSuspendedIdentifier));
+	}
+	ToolbarBuilder.EndBlockGroup();
+	ToolbarBuilder.EndSection();
+
+	return ToolbarBuilder.MakeWidget();
+}
+
+FText SChaosClothAssetEditor3DViewportToolBar::GetDisplayString() const
+{
+	using namespace UE::Chaos::ClothAsset;
+	TSharedRef<FChaosClothAssetEditor3DViewportClient> ViewportClient = StaticCastSharedPtr<FChaosClothAssetEditor3DViewportClient>(ChaosClothAssetEditor3DViewportPtr.Pin()->GetViewportClient()).ToSharedRef();
+	if (FClothEditorSimulationVisualization* const Visualization = ViewportClient->GetSimulationVisualization().Pin().Get())
+	{
+		return Visualization->GetDisplayString(ViewportClient->GetPreviewClothComponent());
+	}
+	return FText();
+}
+
+FText SChaosClothAssetEditor3DViewportToolBar::GetLODMenuLabel() const
+{
+	using namespace UE::Chaos::ClothAsset;
+
+	FText Label = LOCTEXT("LODMenu_AutoLabel", "LOD Auto"); 
+
+	TSharedRef<FChaosClothAssetEditor3DViewportClient> ViewportClient = 
+		StaticCastSharedPtr<FChaosClothAssetEditor3DViewportClient>(ChaosClothAssetEditor3DViewportPtr.Pin()->GetViewportClient()).ToSharedRef();
+
+	const int32 LODSelectionType = ViewportClient->GetLODModel();
+	if (LODSelectionType >= 0)
+	{
+		const FString TitleLabel = FString::Printf(TEXT("LOD %d"), LODSelectionType);
+		Label = FText::FromString(TitleLabel);
+	}
+	return Label;
+}
+
+TSharedRef<SWidget> SChaosClothAssetEditor3DViewportToolBar::MakeLODMenu() const
+{
+	using namespace UE::Chaos::ClothAsset;
+	TSharedRef<FChaosClothAssetEditor3DViewportClient> ViewportClient =
+		StaticCastSharedPtr<FChaosClothAssetEditor3DViewportClient>(ChaosClothAssetEditor3DViewportPtr.Pin()->GetViewportClient()).ToSharedRef();
+
+	constexpr bool bInShouldCloseWindowAfterMenuSelection = true;
+	FMenuBuilder MenuBuilder(bInShouldCloseWindowAfterMenuSelection, CommandList);
+
+	MenuBuilder.PushCommandList(CommandList.ToSharedRef());
+	const int32 NumLODs = ViewportClient->GetNumLODs();
+
+	MenuBuilder.BeginSection("ClothAssetPreviewLODs", LOCTEXT("ShowLOD_PreviewLabel", "Preview LODs"));
+	{
+		MenuBuilder.AddMenuEntry(FChaosClothAssetEditorCommands::Get().LODAuto);
+		MenuBuilder.AddMenuEntry(FChaosClothAssetEditorCommands::Get().LOD0);
+
+		for (int32 LODIndex = 1; LODIndex < NumLODs; ++LODIndex)
+		{
+			FString TitleLabel = FString::Printf(TEXT("LOD %d"), LODIndex);
+
+			FUIAction Action(FExecuteAction::CreateSP(ViewportClient, &FChaosClothAssetEditor3DViewportClient::SetLODModel, LODIndex),
+				FCanExecuteAction(),
+				FIsActionChecked::CreateSP(ViewportClient, &FChaosClothAssetEditor3DViewportClient::IsLODModelSelected, LODIndex));
+
+			MenuBuilder.AddMenuEntry(FText::FromString(TitleLabel), FText::GetEmpty(), FSlateIcon(), Action, NAME_None, EUserInterfaceActionType::RadioButton);
+		}
+	}
+	MenuBuilder.EndSection();
+	MenuBuilder.PopCommandList();
+	return MenuBuilder.MakeWidget();
+}
+
+
+#undef LOCTEXT_NAMESPACE
