#UE Change UI style FindPlugin call to look for ChaosClothAsset instead of ChaosClothAssetEditor since the svgs are in ChaosClothAsset now

diff --git a/Engine/Plugins/ChaosClothAsset/Source/ChaosClothAssetTools/Private/ChaosClothAsset/ClothAssetEditorStyle.cpp b/Engine/Plugins/ChaosClothAsset/Source/ChaosClothAssetTools/Private/ChaosClothAsset/ClothAssetEditorStyle.cpp
--- a/Engine/Plugins/ChaosClothAsset/Source/ChaosClothAssetTools/Private/ChaosClothAsset/ClothAssetEditorStyle.cpp
+++ b/Engine/Plugins/ChaosClothAsset/Source/ChaosClothAssetTools/Private/ChaosClothAsset/ClothAssetEditorStyle.cpp
@@ -11,13 +11,17 @@ namespace UE::Chaos::ClothAsset
 	FClothAssetEditorStyle::FClothAssetEditorStyle()
 		: FSlateStyleSet("ClothAssetEditorStyle")
 	{
-		SetContentRoot(IPluginManager::Get().FindPlugin("ChaosClothAssetEditor")->GetBaseDir() / TEXT("Resources"));
+		TSharedPtr<IPlugin> ChaosClothAssetPlugin = IPluginManager::Get().FindPlugin("ChaosClothAsset");
+		if (ChaosClothAssetPlugin.IsValid())
+		{
+			SetContentRoot(ChaosClothAssetPlugin->GetBaseDir() / TEXT("Resources"));
 
-		Set("ClassIcon.ChaosClothAsset", new FSlateVectorImageBrush(RootToContentDir(TEXT("ClothAsset_16.svg")), FVector2D(16)));
-		Set("ClassThumbnail.ChaosClothAsset", new FSlateVectorImageBrush(RootToContentDir(TEXT("ClothAsset_64.svg")), FVector2D(64)));
+			Set("ClassIcon.ChaosClothAsset", new FSlateVectorImageBrush(RootToContentDir(TEXT("ClothAsset_16.svg")), FVector2D(16)));
+			Set("ClassThumbnail.ChaosClothAsset", new FSlateVectorImageBrush(RootToContentDir(TEXT("ClothAsset_64.svg")), FVector2D(64)));
 
-		Set("ClassIcon.ChaosClothPreset", new FSlateVectorImageBrush(RootToContentDir(TEXT("ClothPreset_16.svg")), FVector2D(16)));
-		Set("ClassThumbnail.ChaosClothPreset", new FSlateVectorImageBrush(RootToContentDir(TEXT("ClothPreset_64.svg")), FVector2D(64)));
+			Set("ClassIcon.ChaosClothPreset", new FSlateVectorImageBrush(RootToContentDir(TEXT("ClothPreset_16.svg")), FVector2D(16)));
+			Set("ClassThumbnail.ChaosClothPreset", new FSlateVectorImageBrush(RootToContentDir(TEXT("ClothPreset_64.svg")), FVector2D(64)));
+		}
 
 		FSlateStyleRegistry::RegisterSlateStyle(*this);
 	}
