CIS - Fixed ClothAsset icons' plugin path after recent ClothAsset plugin changes.

diff --git a/Engine/Plugins/ChaosClothAsset/Source/ChaosClothAssetTools/Private/ChaosClothAsset/ClothComponentEditorStyle.cpp b/Engine/Plugins/ChaosClothAsset/Source/ChaosClothAssetTools/Private/ChaosClothAsset/ClothComponentEditorStyle.cpp
--- a/Engine/Plugins/ChaosClothAsset/Source/ChaosClothAssetTools/Private/ChaosClothAsset/ClothComponentEditorStyle.cpp
+++ b/Engine/Plugins/ChaosClothAsset/Source/ChaosClothAssetTools/Private/ChaosClothAsset/ClothComponentEditorStyle.cpp
@@ -11,12 +11,16 @@ namespace UE::Chaos::ClothAsset
 	FClothComponentEditorStyle::FClothComponentEditorStyle()
 		: FSlateStyleSet("ClothComponentEditorStyle")
 	{
-		SetContentRoot(IPluginManager::Get().FindPlugin("ChaosClothAssetEditor")->GetBaseDir() / TEXT("Resources"));
+		TSharedPtr<IPlugin> ChaosClothAssetPlugin = IPluginManager::Get().FindPlugin("ChaosClothAsset");
+		if (ChaosClothAssetPlugin.IsValid())
+		{
+			SetContentRoot(ChaosClothAssetPlugin->GetBaseDir() / TEXT("Resources"));
 
-		Set("ClassIcon.ChaosClothComponent", new FSlateVectorImageBrush(RootToContentDir(TEXT("ClothAsset_16.svg")), FVector2D(16)));
-		Set("ClassThumbnail.ChaosClothComponent", new FSlateVectorImageBrush(RootToContentDir(TEXT("ClothAsset_64.svg")), FVector2D(64)));
+			Set("ClassIcon.ChaosClothComponent", new FSlateVectorImageBrush(RootToContentDir(TEXT("ClothAsset_16.svg")), FVector2D(16)));
+			Set("ClassThumbnail.ChaosClothComponent", new FSlateVectorImageBrush(RootToContentDir(TEXT("ClothAsset_64.svg")), FVector2D(64)));
 
-		FSlateStyleRegistry::RegisterSlateStyle(*this);
+			FSlateStyleRegistry::RegisterSlateStyle(*this);
+		}
 	}
 
 	FClothComponentEditorStyle::~FClothComponentEditorStyle()
