Implement the EClothAssetAsyncProperties that are accessed by the PostLoadInternal worker in the ChaosClothAsset.

diff --git a/Engine/Plugins/ChaosClothAsset/Source/ChaosClothAssetEngine/Private/ChaosClothAsset/ClothAsset.cpp b/Engine/Plugins/ChaosClothAsset/Source/ChaosClothAssetEngine/Private/ChaosClothAsset/ClothAsset.cpp
--- a/Engine/Plugins/ChaosClothAsset/Source/ChaosClothAssetEngine/Private/ChaosClothAsset/ClothAsset.cpp
+++ b/Engine/Plugins/ChaosClothAsset/Source/ChaosClothAssetEngine/Private/ChaosClothAsset/ClothAsset.cpp
@@ -144,7 +144,7 @@ UChaosClothAsset::UChaosClothAsset(const FObjectInitializer& ObjectInitializer)
 #endif
 {
 	// Setup a single LOD's Cloth Collection
-	TSharedRef<FManagedArrayCollection>& ClothCollection = ClothCollections.Emplace_GetRef(MakeShared<FManagedArrayCollection>());
+	TSharedRef<FManagedArrayCollection>& ClothCollection = GetClothCollections().Emplace_GetRef(MakeShared<FManagedArrayCollection>());
 	UE::Chaos::ClothAsset::FCollectionClothFacade ClothFacade(ClothCollection);
 	ClothFacade.DefineSchema();
 
@@ -205,14 +205,14 @@ void UChaosClothAsset::Serialize(FArchive& Ar)
 	if (bCooked && Ar.IsSaving())
 	{
 		TArray<TSharedRef<FManagedArrayCollection>> OutputCollections = 
-			UE::Chaos::ClothAsset::Private::TrimOnCook(GetPathName(), ClothCollections);
+			UE::Chaos::ClothAsset::Private::TrimOnCook(GetPathName(), GetClothCollections());
 		Chaos::FChaosArchive ChaosArchive(Ar);
 		UE::Chaos::ClothAsset::Private::Serialize(ChaosArchive, OutputCollections);
 	}
 	else
 	{
 		Chaos::FChaosArchive ChaosArchive(Ar);
-		UE::Chaos::ClothAsset::Private::Serialize(ChaosArchive, ClothCollections);
+		UE::Chaos::ClothAsset::Private::Serialize(ChaosArchive, GetClothCollections());
 	}
 
 	Ar << GetRefSkeleton();
@@ -266,23 +266,23 @@ void UChaosClothAsset::BeginPostLoadInternal(FSkinnedAssetPostLoadContext& Conte
 
 	// Make sure that the collection is still compatible and valid
 	bool bAnyInvalidLods = false;
-	if (ClothCollections.IsEmpty())
+	if (GetClothCollections().IsEmpty())
 	{
 		UE_LOG(LogChaosClothAsset, Warning, TEXT("Invalid Cloth Collection (no LODs) found while loading Cloth Asset %s."), *GetFullName());
-		TSharedRef<FManagedArrayCollection>& ClothCollection = ClothCollections.Emplace_GetRef(MakeShared<FManagedArrayCollection>());
+		TSharedRef<FManagedArrayCollection>& ClothCollection = GetClothCollections().Emplace_GetRef(MakeShared<FManagedArrayCollection>());
 		UE::Chaos::ClothAsset::FCollectionClothFacade ClothFacade(ClothCollection);
 		ClothFacade.DefineSchema();
 		bAnyInvalidLods = true;
 	}
 
-	const int32 NumLods = ClothCollections.Num();
+	const int32 NumLods = GetClothCollections().Num();
 	check(NumLods >= 1);  // The default LOD 0 should be present now if it ever was missing
 	LODInfo.SetNum(NumLods);  // Always keep a matching number of LODInfos
 
 	bool bAnyInvalidSkeletons = false;
 	for (int32 LODIndex = 0; LODIndex < NumLods; ++LODIndex)
 	{
-		TSharedRef<FManagedArrayCollection>& ClothCollection = ClothCollections[LODIndex];
+		TSharedRef<FManagedArrayCollection>& ClothCollection = GetClothCollections()[LODIndex];
 
 		UE::Chaos::ClothAsset::FCollectionClothFacade ClothFacade(ClothCollection);
 		if (!ClothFacade.IsValid())
@@ -310,6 +310,9 @@ void UChaosClothAsset::BeginPostLoadInternal(FSkinnedAssetPostLoadContext& Conte
 		SetReferenceSkeleton(nullptr, bRebuildModels, bRebindMeshes);
 	}
 
+	// We're done touching the ClothCollections, so can unlock for read
+	ReleaseAsyncProperty((uint64)EClothAssetAsyncProperties::ClothCollection, ESkinnedAssetAsyncPropertyLockType::WriteOnly);
+
 	BuildClothSimulationModel();  // TODO: Cache ClothSimulationModel?
 
 	BuildMeshModel();
@@ -423,13 +426,13 @@ void UChaosClothAsset::CalculateInvRefMatrices()
 {
 	auto GetRefPoseMatrix = [this](int32 BoneIndex)->FMatrix
 	{
-		check(BoneIndex >= 0 && BoneIndex < RefSkeleton.GetRawBoneNum());
-		FTransform BoneTransform = RefSkeleton.GetRawRefBonePose()[BoneIndex];
+		check(BoneIndex >= 0 && BoneIndex < GetRefSkeleton().GetRawBoneNum());
+		FTransform BoneTransform = GetRefSkeleton().GetRawRefBonePose()[BoneIndex];
 		BoneTransform.NormalizeRotation();  // Make sure quaternion is normalized!
 		return BoneTransform.ToMatrixWithScale();
 	};
 
-	const int32 NumRealBones = RefSkeleton.GetRawBoneNum();
+	const int32 NumRealBones = GetRefSkeleton().GetRawBoneNum();
 
 	RefBasesInvMatrix.Empty(NumRealBones);
 	RefBasesInvMatrix.AddUninitialized(NumRealBones);
@@ -447,7 +450,7 @@ void UChaosClothAsset::CalculateInvRefMatrices()
 		// Construct mesh-space skeletal hierarchy
 		if (BoneIndex > 0)
 		{
-			int32 Parent = RefSkeleton.GetRawParentIndex(BoneIndex);
+			int32 Parent = GetRefSkeleton().GetRawParentIndex(BoneIndex);
 			ComposedRefPoseMatrices[BoneIndex] = ComposedRefPoseMatrices[BoneIndex] * ComposedRefPoseMatrices[Parent];
 		}
 
@@ -463,7 +466,7 @@ void UChaosClothAsset::CalculateInvRefMatrices()
 				Warning,
 				TEXT("Reference Pose for asset %s for joint (%s) includes NIL matrix. Zero scale isn't allowed on ref pose."),
 				*GetPathName(),
-				*RefSkeleton.GetBoneName(BoneIndex).ToString());
+				*GetRefSkeleton().GetBoneName(BoneIndex).ToString());
 		}
 
 		// Precompute inverse so we can use from-refpose-skin vertices
@@ -477,7 +480,7 @@ void UChaosClothAsset::CalculateBounds()
 
 	FBox BoundingBox(ForceInit);
 
-	for (const TSharedRef<FManagedArrayCollection>& ClothCollection : ClothCollections)
+	for (const TSharedRef<const FManagedArrayCollection>& ClothCollection : const_cast<const UChaosClothAsset*>(this)->GetClothCollections())
 	{
 		const FCollectionClothConstFacade Cloth(ClothCollection);
 		const TConstArrayView<FVector3f> RenderPositionArray = Cloth.GetRenderPosition();
@@ -525,7 +528,7 @@ void UChaosClothAsset::Build(TArray<FChaosClothAssetLodTransitionDataCache>* InO
 	CalculateBounds();
 
 	// Add LODs to the render data
-	const int32 NumLods = FMath::Max(ClothCollections.Num(), 1);  // The render data will always look for at least one default LOD 0
+	const int32 NumLods = FMath::Max(GetClothCollections().Num(), 1);  // The render data will always look for at least one default LOD 0
 
 	// Rebuild LOD Infos
 	LODInfo.Reset(NumLods);
@@ -605,11 +608,11 @@ void UChaosClothAsset::BuildMeshModel()
 {
 	using namespace UE::Chaos::ClothAsset;
 
-	const int32 NumLods = ClothCollections.Num();
+	const int32 NumLods = const_cast<const UChaosClothAsset*>(this)->GetClothCollections().Num();
 
 	// Clear current LOD models
-	check(MeshModel);  // MeshModel should always be created in the Cloth Asset constructor WITH_EDITORONLY_DATA
-	MeshModel->LODModels.Reset(NumLods);
+	check(GetImportedModel());  // MeshModel should always be created in the Cloth Asset constructor WITH_EDITORONLY_DATA
+	GetImportedModel()->LODModels.Reset(NumLods);
 
 	// Get the running platform
 	ITargetPlatformManagerModule& TargetPlatformManager = GetTargetPlatformManagerRef();
@@ -618,7 +621,7 @@ void UChaosClothAsset::BuildMeshModel()
 	// Rebuild each LOD models
 	for (int32 LodIndex = 0; LodIndex < NumLods; ++LodIndex)
 	{
-		MeshModel->LODModels.Add(new FSkeletalMeshLODModel());
+		GetImportedModel()->LODModels.Add(new FSkeletalMeshLODModel());
 		BuildLODModel(TargetPlatform, LodIndex);
 	}
 }
@@ -706,8 +709,8 @@ void UChaosClothAsset::CacheDerivedData(FSkinnedAssetCompilationContext* Context
 
 void UChaosClothAsset::BuildLODModel(const ITargetPlatform* TargetPlatform, int32 LODIndex)
 {
-	check(MeshModel && MeshModel->LODModels.IsValidIndex(LODIndex));
-	FBuilder::BuildLod(MeshModel->LODModels[LODIndex], *this, LODIndex, TargetPlatform);
+	check(GetImportedModel() && GetImportedModel()->LODModels.IsValidIndex(LODIndex));
+	FBuilder::BuildLod(GetImportedModel()->LODModels[LODIndex], *this, LODIndex, TargetPlatform);
 }
 
 FString UChaosClothAsset::BuildDerivedDataKey(const ITargetPlatform* TargetPlatform)
@@ -790,7 +793,7 @@ void UChaosClothAsset::SetPhysicsAsset(UPhysicsAsset* InPhysicsAsset)
 	using namespace UE::Chaos::ClothAsset;
 
 	PhysicsAsset = InPhysicsAsset;
-	for (TSharedRef<FManagedArrayCollection>& ClothCollection : ClothCollections)
+	for (TSharedRef<FManagedArrayCollection>& ClothCollection : GetClothCollections())
 	{
 		FCollectionClothFacade Cloth(ClothCollection);
 		Cloth.SetPhysicsAssetPathName(PhysicsAsset ? PhysicsAsset->GetPathName() : FString());
@@ -804,13 +807,13 @@ void UChaosClothAsset::SetReferenceSkeleton(const FReferenceSkeleton* ReferenceS
 	// Update the reference skeleton
 	if (ReferenceSkeleton)
 	{
-		RefSkeleton = *ReferenceSkeleton;
+		GetRefSkeleton() = *ReferenceSkeleton;
 	}
 	else
 	{
 		// Create a default reference skeleton
-		RefSkeleton.Empty(1);
-		FReferenceSkeletonModifier ReferenceSkeletonModifier(RefSkeleton, nullptr);
+		GetRefSkeleton().Empty(1);
+		FReferenceSkeletonModifier ReferenceSkeletonModifier(GetRefSkeleton(), nullptr);
 
 		FMeshBoneInfo MeshBoneInfo;
 		constexpr const TCHAR* RootName = TEXT("Root");
@@ -827,7 +830,7 @@ void UChaosClothAsset::SetReferenceSkeleton(const FReferenceSkeleton* ReferenceS
 	// Rebind the meshes
 	if (bRebindMeshes)
 	{
-		for (TSharedRef<FManagedArrayCollection>& ClothCollection : ClothCollections)
+		for (TSharedRef<FManagedArrayCollection>& ClothCollection : GetClothCollections())
 		{
 			FClothGeometryTools::BindMeshToRootBone(ClothCollection, true, true);
 		}
@@ -844,8 +847,8 @@ void UChaosClothAsset::UpdateSkeletonFromCollection(bool bRebuildModels)
 {
 	using namespace UE::Chaos::ClothAsset;
 
-	check(ClothCollections.Num());
-	FCollectionClothConstFacade ClothFacade(ClothCollections[0]);
+	check(GetClothCollections().Num());
+	FCollectionClothConstFacade ClothFacade(GetClothCollections()[0]);
 	check(ClothFacade.IsValid());
 
 	const FString& SkeletalMeshPathName = ClothFacade.GetSkeletalMeshPathName();
@@ -861,7 +864,7 @@ void UChaosClothAsset::UpdateSkeletonFromCollection(bool bRebuildModels)
 void UChaosClothAsset::CopySimMeshToRenderMesh(UMaterialInterface* Material)
 {
 	using namespace UE::Chaos::ClothAsset;
-	check(ClothCollections.Num());
+	check(GetClothCollections().Num());
 
 	// Add a default material if none is specified
 	const FString RenderMaterialPathName = Material ?
@@ -869,7 +872,7 @@ void UChaosClothAsset::CopySimMeshToRenderMesh(UMaterialInterface* Material)
 		FString(TEXT("/Engine/EditorMaterials/Cloth/CameraLitDoubleSided.CameraLitDoubleSided"));
 
 	bool bAnyLodHasRenderMesh = false;
-	for (TSharedRef<FManagedArrayCollection>& ClothCollection : ClothCollections)
+	for (TSharedRef<FManagedArrayCollection>& ClothCollection : GetClothCollections())
 	{
 		constexpr bool bSingleRenderPattern = true;
 		FClothGeometryTools::CopySimMeshToRenderMesh(ClothCollection, RenderMaterialPathName, bSingleRenderPattern);
@@ -905,9 +908,9 @@ void UChaosClothAsset::ReregisterComponents()
 void UChaosClothAsset::BindSimMeshToRootBone()
 {
 	using namespace UE::Chaos::ClothAsset;
-	check(ClothCollections.Num());
+	check(GetClothCollections().Num());
 
-	for (TSharedRef<FManagedArrayCollection>& ClothCollection : ClothCollections)
+	for (TSharedRef<FManagedArrayCollection>& ClothCollection : GetClothCollections())
 	{
 		FClothGeometryTools::BindMeshToRootBone(ClothCollection, true, false);
 	}
diff --git a/Engine/Plugins/ChaosClothAsset/Source/ChaosClothAssetEngine/Private/ChaosClothAsset/ClothAssetBuilder.cpp b/Engine/Plugins/ChaosClothAsset/Source/ChaosClothAssetEngine/Private/ChaosClothAsset/ClothAssetBuilder.cpp
--- a/Engine/Plugins/ChaosClothAsset/Source/ChaosClothAssetEngine/Private/ChaosClothAsset/ClothAssetBuilder.cpp
+++ b/Engine/Plugins/ChaosClothAsset/Source/ChaosClothAssetEngine/Private/ChaosClothAsset/ClothAssetBuilder.cpp
@@ -137,7 +137,7 @@ void UChaosClothAsset::FBuilder::BuildLod(FSkeletalMeshLODModel& LODModel, const
 
 	// Keep track of the active bone indices for this LOD model
 	TSet<FBoneIndexType> ActiveBoneIndices;
-	ActiveBoneIndices.Reserve(ClothAsset.RefSkeleton.GetNum());
+	ActiveBoneIndices.Reserve(ClothAsset.GetRefSkeleton().GetNum());
 
 	// Load the mesh utilities module used to optimized the index buffer
 	IMeshUtilities& MeshUtilities = FModuleManager::Get().LoadModuleChecked<IMeshUtilities>("MeshUtilities");
@@ -478,10 +478,10 @@ void UChaosClothAsset::FBuilder::BuildLod(FSkeletalMeshLODModel& LODModel, const
 	LODModel.ActiveBoneIndices = ActiveBoneIndices.Array();
 
 	// Ensure parent exists with incoming active bone indices, and the result should be sorted
-	ClothAsset.RefSkeleton.EnsureParentsExistAndSort(LODModel.ActiveBoneIndices);
+	ClothAsset.GetRefSkeleton().EnsureParentsExistAndSort(LODModel.ActiveBoneIndices);
 
 	// Compute the required bones for this model.
-	USkeletalMesh::CalculateRequiredBones(LODModel, ClothAsset.RefSkeleton, nullptr);
+	USkeletalMesh::CalculateRequiredBones(LODModel, ClothAsset.GetRefSkeleton(), nullptr);
 }
 
 #undef LOCTEXT_NAMESPACE
diff --git a/Engine/Plugins/ChaosClothAsset/Source/ChaosClothAssetEngine/Public/ChaosClothAsset/ClothAsset.h b/Engine/Plugins/ChaosClothAsset/Source/ChaosClothAssetEngine/Public/ChaosClothAsset/ClothAsset.h
--- a/Engine/Plugins/ChaosClothAsset/Source/ChaosClothAssetEngine/Public/ChaosClothAsset/ClothAsset.h
+++ b/Engine/Plugins/ChaosClothAsset/Source/ChaosClothAssetEngine/Public/ChaosClothAsset/ClothAsset.h
@@ -24,6 +24,9 @@ enum class EClothAssetAsyncProperties : uint64
 	None = 0,
 	RenderData = 1 << 0,
 	ThumbnailInfo = 1 << 1,
+	ImportedModel = 1 << 2,
+	ClothCollection = 1 << 3,
+	RefSkeleton = 1 << 4,
 	All = MAX_uint64
 };
 ENUM_CLASS_FLAGS(EClothAssetAsyncProperties);
@@ -51,8 +54,20 @@ public:
 	//~ End UObject interface
 
 	//~ Begin USkinnedAsset interface
-	virtual FReferenceSkeleton& GetRefSkeleton()								{ return RefSkeleton; }
-	virtual const FReferenceSkeleton& GetRefSkeleton() const					{ return RefSkeleton; }
+	virtual FReferenceSkeleton& GetRefSkeleton()								
+	{
+		WaitUntilAsyncPropertyReleased(EClothAssetAsyncProperties::RefSkeleton);
+		PRAGMA_DISABLE_DEPRECATION_WARNINGS
+		return RefSkeleton;
+		PRAGMA_ENABLE_DEPRECATION_WARNINGS
+	}
+	virtual const FReferenceSkeleton& GetRefSkeleton() const					
+	{
+		WaitUntilAsyncPropertyReleased(EClothAssetAsyncProperties::RefSkeleton, ESkinnedAssetAsyncPropertyLockType::ReadOnly);
+		PRAGMA_DISABLE_DEPRECATION_WARNINGS
+		return RefSkeleton;
+		PRAGMA_ENABLE_DEPRECATION_WARNINGS
+	}
 	virtual FSkeletalMeshLODInfo* GetLODInfo(int32 Index);
 	virtual const FSkeletalMeshLODInfo* GetLODInfo(int32 Index) const;
 	UFUNCTION(BlueprintGetter)
@@ -100,15 +115,33 @@ public:
 	virtual int32 GetMaxNumOptionalLODs(const ITargetPlatform* TargetPlatform) const override { return 0; }
 #endif
 #if WITH_EDITORONLY_DATA
-	virtual FSkeletalMeshModel* GetImportedModel() const override				{ return MeshModel.Get(); }
+	virtual FSkeletalMeshModel* GetImportedModel() const override
+	{
+		WaitUntilAsyncPropertyReleased(EClothAssetAsyncProperties::ImportedModel);
+		PRAGMA_DISABLE_DEPRECATION_WARNINGS
+		return MeshModel.Get();
+		PRAGMA_ENABLE_DEPRECATION_WARNINGS
+	}
 #endif
 	//~ End USkinnedAsset interface
 
 	/** Return the enclosed Cloth Collection object. */
-	TArray<TSharedRef<FManagedArrayCollection>>& GetClothCollections() { return ClothCollections; }
+	TArray<TSharedRef<FManagedArrayCollection>>& GetClothCollections() 
+	{
+		WaitUntilAsyncPropertyReleased(EClothAssetAsyncProperties::ClothCollection);
+		PRAGMA_DISABLE_DEPRECATION_WARNINGS
+		return ClothCollections;
+		PRAGMA_ENABLE_DEPRECATION_WARNINGS
+	}
 
 	/** Return the enclosed Cloth Collection object, const version. */
-	const TArray<TSharedRef<const FManagedArrayCollection>>& GetClothCollections() const { return reinterpret_cast<const TArray<TSharedRef<const FManagedArrayCollection>>&>(ClothCollections); }
+	const TArray<TSharedRef<const FManagedArrayCollection>>& GetClothCollections() const 
+	{
+		WaitUntilAsyncPropertyReleased(EClothAssetAsyncProperties::ClothCollection, ESkinnedAssetAsyncPropertyLockType::ReadOnly);
+		PRAGMA_DISABLE_DEPRECATION_WARNINGS
+		return reinterpret_cast<const TArray<TSharedRef<const FManagedArrayCollection>>&>(ClothCollections);
+		PRAGMA_ENABLE_DEPRECATION_WARNINGS
+	}
 
 	/** Return the cloth simulation ready LOD model data. */
 	TSharedPtr<const FChaosClothSimulationModel> GetClothSimulationModel() const { return ClothSimulationModel; }
@@ -143,7 +176,7 @@ public:
 	/** Set the bone hierachy to use for this cloth. */
 	UE_DEPRECATED(5.3, "Use SetReferenceSkeleton(const FReferenceSkeleton*, bool, bool) instead")
 	PRAGMA_DISABLE_DEPRECATION_WARNINGS
-	void SetReferenceSkeleton(const FReferenceSkeleton& InReferenceSkeleton, bool bRebuildClothSimulationModel = true) { RefSkeleton = InReferenceSkeleton; UpdateSkeleton(bRebuildClothSimulationModel); }
+	void SetReferenceSkeleton(const FReferenceSkeleton& InReferenceSkeleton, bool bRebuildClothSimulationModel = true) { GetRefSkeleton() = InReferenceSkeleton; UpdateSkeleton(bRebuildClothSimulationModel); }
 	PRAGMA_ENABLE_DEPRECATION_WARNINGS
 
 	/** Set the skinning weights for all of the sim vertices to be bound to the root node of the reference skeleton. */
@@ -294,9 +327,11 @@ private:
 	TArray<FMatrix> CachedComposedRefPoseMatrices;
 
 	/** Cloth Collection containing this asset data. One per LOD. */
+	UE_DEPRECATED(5.4, "This must be protected for async build, always use the accessors even internally.")
 	TArray<TSharedRef<FManagedArrayCollection>> ClothCollections;
 
 	/** Reference skeleton created from the provided skeleton asset. */
+	UE_DEPRECATED(5.4, "This must be protected for async build, always use the accessors even internally.")
 	FReferenceSkeleton RefSkeleton;
 
 	/** Reference skeleton precomputed bases. */
@@ -310,6 +345,7 @@ private:
 
 #if WITH_EDITORONLY_DATA
 	/** Source mesh geometry information (not used at runtime). */
+	UE_DEPRECATED(5.4, "This must be protected for async build, always use the accessors even internally.")
 	TSharedPtr<FSkeletalMeshModel> MeshModel;
 #endif
 
