Optional schemas for cloth fabrics

diff --git a/Engine/Plugins/ChaosClothAsset/Source/ChaosClothAsset/Private/ChaosClothAsset/ClothCollection.cpp b/Engine/Plugins/ChaosClothAsset/Source/ChaosClothAsset/Private/ChaosClothAsset/ClothCollection.cpp
--- a/Engine/Plugins/ChaosClothAsset/Source/ChaosClothAsset/Private/ChaosClothAsset/ClothCollection.cpp
+++ b/Engine/Plugins/ChaosClothAsset/Source/ChaosClothAsset/Private/ChaosClothAsset/ClothCollection.cpp
@@ -279,16 +279,17 @@ namespace UE::Chaos::ClothAsset
 			SimIndices3D &&
 
 			// Fabrics Group
-			FabricBendingStiffness &&
-			FabricBucklingStiffness &&
-			FabricStretchStiffness &&
-			FabricBucklingRatio &&
-			FabricDensity &&
-			FabricFriction &&
-			FabricDamping &&
-			FabricPressure &&
-			FabricLayer &&
-			FabricCollisionThickness &&
+			(!EnumHasAnyFlags(OptionalSchemas, EClothCollectionOptionalSchemas::Fabrics) ||
+				(FabricBendingStiffness &&
+				FabricBucklingStiffness &&
+				FabricStretchStiffness &&
+				FabricBucklingRatio &&
+				FabricDensity &&
+				FabricFriction &&
+				FabricDamping &&
+				FabricPressure &&
+				FabricLayer &&
+				FabricCollisionThickness )) &&
 
 			// Sim Vertices 2D Group
 			SimPosition2D &&
@@ -376,16 +377,19 @@ namespace UE::Chaos::ClothAsset
 		RenderMaterialPathName = &ManagedArrayCollection->AddAttribute<FString>(ClothCollectionAttribute::RenderMaterialPathName, ClothCollectionGroup::RenderPatterns);
 
 		//~ Fabric Group
-		FabricBendingStiffness = &ManagedArrayCollection->AddAttribute<FVector3f>(ClothCollectionAttribute::FabricBendingStiffness, ClothCollectionGroup::Fabrics);
-		FabricBucklingStiffness = &ManagedArrayCollection->AddAttribute<FVector3f>(ClothCollectionAttribute::FabricBucklingStiffness, ClothCollectionGroup::Fabrics);
-		FabricStretchStiffness = &ManagedArrayCollection->AddAttribute<FVector3f>(ClothCollectionAttribute::FabricStretchStiffness, ClothCollectionGroup::Fabrics);
-		FabricBucklingRatio = &ManagedArrayCollection->AddAttribute<float>(ClothCollectionAttribute::FabricBucklingRatio, ClothCollectionGroup::Fabrics);
-		FabricDensity = &ManagedArrayCollection->AddAttribute<float>(ClothCollectionAttribute::FabricDensity, ClothCollectionGroup::Fabrics);
-		FabricFriction = &ManagedArrayCollection->AddAttribute<float>(ClothCollectionAttribute::FabricFriction, ClothCollectionGroup::Fabrics);
-		FabricDamping = &ManagedArrayCollection->AddAttribute<float>(ClothCollectionAttribute::FabricDamping, ClothCollectionGroup::Fabrics);
-		FabricPressure = &ManagedArrayCollection->AddAttribute<float>(ClothCollectionAttribute::FabricPressure, ClothCollectionGroup::Fabrics);
-		FabricLayer = &ManagedArrayCollection->AddAttribute<int32>(ClothCollectionAttribute::FabricLayer, ClothCollectionGroup::Fabrics);
-		FabricCollisionThickness = &ManagedArrayCollection->AddAttribute<float>(ClothCollectionAttribute::FabricCollisionThickness, ClothCollectionGroup::Fabrics);
+		if (EnumHasAnyFlags(OptionalSchemas, EClothCollectionOptionalSchemas::Fabrics))
+		{
+			FabricBendingStiffness = &ManagedArrayCollection->AddAttribute<FVector3f>(ClothCollectionAttribute::FabricBendingStiffness, ClothCollectionGroup::Fabrics);
+			FabricBucklingStiffness = &ManagedArrayCollection->AddAttribute<FVector3f>(ClothCollectionAttribute::FabricBucklingStiffness, ClothCollectionGroup::Fabrics);
+			FabricStretchStiffness = &ManagedArrayCollection->AddAttribute<FVector3f>(ClothCollectionAttribute::FabricStretchStiffness, ClothCollectionGroup::Fabrics);
+			FabricBucklingRatio = &ManagedArrayCollection->AddAttribute<float>(ClothCollectionAttribute::FabricBucklingRatio, ClothCollectionGroup::Fabrics);
+			FabricDensity = &ManagedArrayCollection->AddAttribute<float>(ClothCollectionAttribute::FabricDensity, ClothCollectionGroup::Fabrics);
+			FabricFriction = &ManagedArrayCollection->AddAttribute<float>(ClothCollectionAttribute::FabricFriction, ClothCollectionGroup::Fabrics);
+			FabricDamping = &ManagedArrayCollection->AddAttribute<float>(ClothCollectionAttribute::FabricDamping, ClothCollectionGroup::Fabrics);
+			FabricPressure = &ManagedArrayCollection->AddAttribute<float>(ClothCollectionAttribute::FabricPressure, ClothCollectionGroup::Fabrics);
+			FabricLayer = &ManagedArrayCollection->AddAttribute<int32>(ClothCollectionAttribute::FabricLayer, ClothCollectionGroup::Fabrics);
+			FabricCollisionThickness = &ManagedArrayCollection->AddAttribute<float>(ClothCollectionAttribute::FabricCollisionThickness, ClothCollectionGroup::Fabrics);
+		}
 		
 		// Sim Faces Group
 		SimIndices2D = &ManagedArrayCollection->AddAttribute<FIntVector3>(ClothCollectionAttribute::SimIndices2D, ClothCollectionGroup::SimFaces, SimVertices2DDependency);
diff --git a/Engine/Plugins/ChaosClothAsset/Source/ChaosClothAsset/Private/ChaosClothAsset/ClothCollection.h b/Engine/Plugins/ChaosClothAsset/Source/ChaosClothAsset/Private/ChaosClothAsset/ClothCollection.h
--- a/Engine/Plugins/ChaosClothAsset/Source/ChaosClothAsset/Private/ChaosClothAsset/ClothCollection.h
+++ b/Engine/Plugins/ChaosClothAsset/Source/ChaosClothAsset/Private/ChaosClothAsset/ClothCollection.h
@@ -239,6 +239,8 @@ namespace UE::Chaos::ClothAsset
 		//~ LODs Group (There should be only one LOD per ClothCollection)
 		TManagedArray<FString>* GetPhysicsAssetPathName(){ return PhysicsAssetPathName; }
 		TManagedArray<FString>* GetSkeletalMeshPathName() { return SkeletalMeshPathName; }
+		
+		//~ Solver Group
 		TManagedArray<FVector3f>* GetSolverGravity() { return SolverGravity; }
 		TManagedArray<float>* GetSolverAirDamping() { return SolverAirDamping; }
 		TManagedArray<int32>* GetSolverSubSteps() { return SolverSubSteps; }
diff --git a/Engine/Plugins/ChaosClothAsset/Source/ChaosClothAsset/Private/ChaosClothAsset/CollectionClothFabricFacade.cpp b/Engine/Plugins/ChaosClothAsset/Source/ChaosClothAsset/Private/ChaosClothAsset/CollectionClothFabricFacade.cpp
--- a/Engine/Plugins/ChaosClothAsset/Source/ChaosClothAsset/Private/ChaosClothAsset/CollectionClothFabricFacade.cpp
+++ b/Engine/Plugins/ChaosClothAsset/Source/ChaosClothAsset/Private/ChaosClothAsset/CollectionClothFabricFacade.cpp
@@ -26,55 +26,56 @@ namespace UE::Chaos::ClothAsset
 	
 	FCollectionClothFabricConstFacade::FAnisotropicData FCollectionClothFabricConstFacade::GetBendingStiffness() const
 	{
-		return FCollectionClothFabricConstFacade::FAnisotropicData(
-			ClothCollection->GetElements(ClothCollection->GetFabricBendingStiffness())[GetElementIndex()]);
+		return ClothCollection->GetFabricBendingStiffness() ? FCollectionClothFabricConstFacade::FAnisotropicData(
+			ClothCollection->GetElements(ClothCollection->GetFabricBendingStiffness())[GetElementIndex()]) : FAnisotropicData(FDefaultFabric::BendingStiffness);
 	}
 	
 	float FCollectionClothFabricConstFacade::GetBucklingRatio() const
 	{
-		return ClothCollection->GetElements(ClothCollection->GetFabricBucklingRatio())[GetElementIndex()];
+		return ClothCollection->GetFabricBucklingRatio() ? ClothCollection->GetElements(ClothCollection->GetFabricBucklingRatio())[GetElementIndex()] :
+				FDefaultFabric::BucklingRatio;
 	}
 	
 	FCollectionClothFabricConstFacade::FAnisotropicData FCollectionClothFabricConstFacade::GetBucklingStiffness() const
 	{
-		return FCollectionClothFabricConstFacade::FAnisotropicData(
-			ClothCollection->GetElements(ClothCollection->GetFabricBucklingStiffness())[GetElementIndex()]);
+		return ClothCollection->GetFabricBucklingStiffness() ? FCollectionClothFabricConstFacade::FAnisotropicData(
+			ClothCollection->GetElements(ClothCollection->GetFabricBucklingStiffness())[GetElementIndex()]) : FAnisotropicData(FDefaultFabric::BucklingStiffness);
 	}
 	
 	FCollectionClothFabricConstFacade::FAnisotropicData FCollectionClothFabricConstFacade::GetStretchStiffness() const
 	{
-		return FCollectionClothFabricConstFacade::FAnisotropicData(
-			ClothCollection->GetElements(ClothCollection->GetFabricStretchStiffness())[GetElementIndex()]);
+		return ClothCollection->GetFabricStretchStiffness() ? FCollectionClothFabricConstFacade::FAnisotropicData(
+			ClothCollection->GetElements(ClothCollection->GetFabricStretchStiffness())[GetElementIndex()]) : FAnisotropicData(FDefaultFabric::StretchStiffness);
 	}
 	
 	float FCollectionClothFabricConstFacade::GetDensity() const
 	{
-		return ClothCollection->GetElements(ClothCollection->GetFabricDensity())[GetElementIndex()];
+		return ClothCollection->GetFabricDensity() ? ClothCollection->GetElements(ClothCollection->GetFabricDensity())[GetElementIndex()] : FDefaultFabric::Density;
 	}
 	
 	float FCollectionClothFabricConstFacade::GetDamping() const
 	{
-		return ClothCollection->GetElements(ClothCollection->GetFabricDamping())[GetElementIndex()];
+		return ClothCollection->GetFabricDamping() ? ClothCollection->GetElements(ClothCollection->GetFabricDamping())[GetElementIndex()] : FDefaultFabric::Damping;
 	}
 	
 	float FCollectionClothFabricConstFacade::GetFriction() const
 	{
-		return ClothCollection->GetElements(ClothCollection->GetFabricFriction())[GetElementIndex()];
+		return ClothCollection->GetFabricFriction() ? ClothCollection->GetElements(ClothCollection->GetFabricFriction())[GetElementIndex()] : FDefaultFabric::Friction;
 	}
 	
 	float FCollectionClothFabricConstFacade::GetPressure() const
 	{
-		return ClothCollection->GetElements(ClothCollection->GetFabricPressure())[GetElementIndex()];
+		return ClothCollection->GetFabricPressure() ? ClothCollection->GetElements(ClothCollection->GetFabricPressure())[GetElementIndex()] : FDefaultFabric::Pressure;
 	}
 	
 	int32 FCollectionClothFabricConstFacade::GetLayer() const
 	{
-		return ClothCollection->GetElements(ClothCollection->GetFabricLayer())[GetElementIndex()];
+		return ClothCollection->GetFabricLayer() ? ClothCollection->GetElements(ClothCollection->GetFabricLayer())[GetElementIndex()] : FDefaultFabric::Layer;
 	}
 
 	float FCollectionClothFabricConstFacade::GetCollisionThickness() const
 	{
-		return ClothCollection->GetElements(ClothCollection->GetFabricCollisionThickness())[GetElementIndex()];
+		return ClothCollection->GetFabricCollisionThickness() ? ClothCollection->GetElements(ClothCollection->GetFabricCollisionThickness())[GetElementIndex()] : FDefaultFabric::CollisionThickness;
 	}
 	
 	FCollectionClothFabricConstFacade::FCollectionClothFabricConstFacade(const TSharedRef<const FClothCollection>& ClothCollection, int32 InFabricIndex)
@@ -93,6 +94,10 @@ namespace UE::Chaos::ClothAsset
 		check((ElementIndex >= 0) && (ElementIndex < ClothCollection->GetNumElements(ClothCollectionGroup::Fabrics)));
 		
 		const TSharedRef<class FClothCollection> ClothCol = GetClothCollection();
+		if (!ClothCol->IsValid(EClothCollectionOptionalSchemas::Fabrics))
+		{
+			ClothCol->DefineSchema(EClothCollectionOptionalSchemas::Fabrics);
+		}
 		
 		ClothCol->GetElements(ClothCol->GetFabricBendingStiffness())[ElementIndex] = BendingStiffness.GetVectorDatas();
 		ClothCol->GetElements(ClothCol->GetFabricBucklingRatio())[ElementIndex] = BucklingRatio;
diff --git a/Engine/Plugins/ChaosClothAsset/Source/ChaosClothAsset/Private/ChaosClothAsset/CollectionClothFacade.cpp b/Engine/Plugins/ChaosClothAsset/Source/ChaosClothAsset/Private/ChaosClothAsset/CollectionClothFacade.cpp
--- a/Engine/Plugins/ChaosClothAsset/Source/ChaosClothAsset/Private/ChaosClothAsset/CollectionClothFacade.cpp
+++ b/Engine/Plugins/ChaosClothAsset/Source/ChaosClothAsset/Private/ChaosClothAsset/CollectionClothFacade.cpp
@@ -534,11 +534,7 @@ namespace UE::Chaos::ClothAsset
 		SetNumRenderPatterns(0);
 		SetNumSeams(0); // Do this after removing SimVertices3D and SimPatterns. Otherwise, Seams will do a bunch of unnecessary work to unseam stuff.
 		SetNumFabrics(0); 
-
-		if (IsValid(EClothCollectionOptionalSchemas::Solvers))
-		{
-			GetClothCollection()->SetNumElements(0, ClothCollectionGroup::Solvers);
-		}
+		GetClothCollection()->SetNumElements(0, ClothCollectionGroup::Solvers);
 	}
 
 	void FCollectionClothFacade::Initialize(const FCollectionClothConstFacade& Other)
diff --git a/Engine/Plugins/ChaosClothAsset/Source/ChaosClothAsset/Public/ChaosClothAsset/ClothCollectionOptionalSchemas.h b/Engine/Plugins/ChaosClothAsset/Source/ChaosClothAsset/Public/ChaosClothAsset/ClothCollectionOptionalSchemas.h
--- a/Engine/Plugins/ChaosClothAsset/Source/ChaosClothAsset/Public/ChaosClothAsset/ClothCollectionOptionalSchemas.h
+++ b/Engine/Plugins/ChaosClothAsset/Source/ChaosClothAsset/Public/ChaosClothAsset/ClothCollectionOptionalSchemas.h
@@ -12,7 +12,8 @@ namespace UE::Chaos::ClothAsset
 	{
 		None = 0,
 		RenderDeformer = 1 << 1,
-		Solvers = 1 << 2
+		Solvers = 1 << 2,
+		Fabrics = 1 << 3
 	};
 	ENUM_CLASS_FLAGS(EClothCollectionOptionalSchemas)
 }  // End namespace UE::Chaos::ClothAsset
diff --git a/Engine/Plugins/ChaosClothAssetEditor/Source/ChaosClothAssetDataflowNodes/Private/ChaosClothAsset/SimulationBaseConfigNode.cpp b/Engine/Plugins/ChaosClothAssetEditor/Source/ChaosClothAssetDataflowNodes/Private/ChaosClothAsset/SimulationBaseConfigNode.cpp
--- a/Engine/Plugins/ChaosClothAssetEditor/Source/ChaosClothAssetDataflowNodes/Private/ChaosClothAsset/SimulationBaseConfigNode.cpp
+++ b/Engine/Plugins/ChaosClothAssetEditor/Source/ChaosClothAssetDataflowNodes/Private/ChaosClothAsset/SimulationBaseConfigNode.cpp
@@ -400,7 +400,7 @@ void FChaosClothAssetSimulationBaseConfigNode::FPropertyHelper::SetSolverPropert
 	const TArray<FName>& SimilarPropertyNames, ECollectionPropertyFlags PropertyFlags)
 {
 	UE::Chaos::ClothAsset::FCollectionClothFacade ClothFacade(GetClothCollection());
-	if(PropertyValue.bUseImportedValue && ClothFacade.IsValid() && ClothFacade.HasSolverElement())
+	if(PropertyValue.bUseImportedValue && ClothFacade.IsValid(UE::Chaos::ClothAsset::EClothCollectionOptionalSchemas::Solvers) && ClothFacade.HasSolverElement())
 	{
 		PropertyValue.ImportedValue = SolverValueFunction(ClothFacade);
 	}
@@ -413,7 +413,7 @@ void FChaosClothAssetSimulationBaseConfigNode::FPropertyHelper::SetFabricPropert
 	const TArray<FName>& SimilarPropertyNames, ECollectionPropertyFlags PropertyFlags)
 {
 	UE::Chaos::ClothAsset::FCollectionClothFacade ClothFacade(GetClothCollection());
-	if(PropertyValue.bUseImportedValue && ClothFacade.IsValid() && ClothFacade.GetNumFabrics() > 0)
+	if(PropertyValue.bUseImportedValue && ClothFacade.IsValid(UE::Chaos::ClothAsset::EClothCollectionOptionalSchemas::Fabrics) && ClothFacade.GetNumFabrics() > 0)
 	{
 		const int32 NumFabrics = ClothFacade.GetNumFabrics();
 		
@@ -436,7 +436,7 @@ void FChaosClothAssetSimulationBaseConfigNode::FPropertyHelper::SetFabricPropert
 	const TArray<FName>& SimilarPropertyNames, ECollectionPropertyFlags PropertyFlags)
 {
 	UE::Chaos::ClothAsset::FCollectionClothFacade ClothFacade(GetClothCollection());
-	if(PropertyValue.bCouldUseFabrics && ClothFacade.IsValid() && (ClothFacade.GetNumFabrics() > 0) &&
+	if(PropertyValue.bCouldUseFabrics && ClothFacade.IsValid(UE::Chaos::ClothAsset::EClothCollectionOptionalSchemas::Fabrics) && (ClothFacade.GetNumFabrics() > 0) &&
 			(PropertyValue.bImportFabricBounds || PropertyValue.bBuildFabricMaps))
 	{
 		UE::Chaos::ClothAsset::FWeightedValueBounds WeightedValueBounds;
@@ -468,7 +468,7 @@ void FChaosClothAssetSimulationBaseConfigNode::FPropertyHelper::SetFabricPropert
 	const TArray<FName>& SimilarPropertyNames, ECollectionPropertyFlags PropertyFlags, const FName& GroupName)
 {
 	UE::Chaos::ClothAsset::FCollectionClothFacade ClothFacade(GetClothCollection());
-	if(PropertyValue.bCouldUseFabrics && ClothFacade.IsValid() && (ClothFacade.GetNumFabrics() > 0) && PropertyValue.bBuildFabricMaps)
+	if(PropertyValue.bCouldUseFabrics && ClothFacade.IsValid(UE::Chaos::ClothAsset::EClothCollectionOptionalSchemas::Fabrics) && (ClothFacade.GetNumFabrics() > 0) && PropertyValue.bBuildFabricMaps)
 	{
 		const int32 NumPatterns = ClothFacade.GetNumSimPatterns();
 
@@ -522,7 +522,7 @@ void FChaosClothAssetSimulationBaseConfigNode::FPropertyHelper::SetSolverPropert
 	const TArray<FName>& SimilarPropertyNames, ECollectionPropertyFlags PropertyFlags)
 {
 	UE::Chaos::ClothAsset::FCollectionClothFacade ClothFacade(GetClothCollection());
-	if(PropertyValue.bCouldUseFabrics && (ClothFacade.IsValid() && ClothFacade.HasSolverElement()) && PropertyValue.bImportFabricBounds)
+	if(PropertyValue.bCouldUseFabrics && (ClothFacade.IsValid(UE::Chaos::ClothAsset::EClothCollectionOptionalSchemas::Solvers) && ClothFacade.HasSolverElement()) && PropertyValue.bImportFabricBounds)
 	{
 		const float SolverValue = SolverValueFunction(ClothFacade);
 	
