Chaos Cloth Asset - Disable Cloth Asset properties set by Dataflow to avoid user confusion and lost settings.

diff --git a/Engine/Plugins/ChaosClothAsset/Source/ChaosClothAssetEngine/Private/ChaosClothAsset/ClothAssetBuilder.cpp b/Engine/Plugins/ChaosClothAsset/Source/ChaosClothAssetEngine/Private/ChaosClothAsset/ClothAssetBuilder.cpp
--- a/Engine/Plugins/ChaosClothAsset/Source/ChaosClothAssetEngine/Private/ChaosClothAsset/ClothAssetBuilder.cpp
+++ b/Engine/Plugins/ChaosClothAsset/Source/ChaosClothAssetEngine/Private/ChaosClothAsset/ClothAssetBuilder.cpp
@@ -417,15 +417,16 @@ void UChaosClothAsset::FBuilder::BuildLod(FSkeletalMeshLODModel& LODModel, const
 					SectionRenderNormals,
 					SectionRenderTangents,
 					SectionRenderIndices);
-
+PRAGMA_DISABLE_DEPRECATION_WARNINGS
 				ClothingMeshUtils::GenerateMeshToMeshVertData(
 					Section.ClothMappingDataLODs[0],
 					TargetMesh,
 					SourceMesh,
 					&MaxDistances,
-					ClothAsset.bSmoothTransition,
-					ClothAsset.bUseMultipleInfluences,
-					ClothAsset.SkinningKernelRadius);
+					ClothAsset.bSmoothTransition_DEPRECATED,
+					ClothAsset.bUseMultipleInfluences_DEPRECATED,
+					ClothAsset.SkinningKernelRadius_DEPRECATED);
+PRAGMA_ENABLE_DEPRECATION_WARNINGS
 			}
 		}
 
diff --git a/Engine/Plugins/ChaosClothAsset/Source/ChaosClothAssetEngine/Public/ChaosClothAsset/ClothAsset.h b/Engine/Plugins/ChaosClothAsset/Source/ChaosClothAssetEngine/Public/ChaosClothAsset/ClothAsset.h
--- a/Engine/Plugins/ChaosClothAsset/Source/ChaosClothAssetEngine/Public/ChaosClothAsset/ClothAsset.h
+++ b/Engine/Plugins/ChaosClothAsset/Source/ChaosClothAssetEngine/Public/ChaosClothAsset/ClothAsset.h
@@ -267,26 +267,30 @@ private:
 	/** Reregister all components using this asset to reset the simulation in case anything has changed. */
 	void ReregisterComponents();
 
-	/** List of materials for this cloth asset. */
-	UPROPERTY(EditAnywhere, Category = Materials)
+	/** List of materials for this cloth asset. Set by the Dataflow evaluation. */
+	UPROPERTY(EditAnywhere, Category = Materials, Meta = (EditCondition = "DataflowAsset == nullptr"))
 	TArray<FSkeletalMaterial> Materials;
 
-	/** Skeleton asset used at creation time. This is of limited use since this USkeleton's reference skeleton might not necessarily match the one created for this asset. */
-	UPROPERTY(EditAnywhere, Setter = SetSkeleton, Category = Skeleton)
+	/**
+	 * Skeleton asset used at creation time.
+	 * This is of limited use since this USkeleton's reference skeleton might not necessarily match the one created for this asset.
+	 * Set by the Dataflow evaluation.
+	 */
+	UPROPERTY(EditAnywhere, Setter = SetSkeleton, Category = Skeleton, Meta = (EditCondition = "DataflowAsset == nullptr"))
 	TObjectPtr<USkeleton> Skeleton;
 
-	/** Physics asset used for collision. */
-	UPROPERTY(EditAnywhere, Category = Collision)
+	/** Physics asset used for collision. Set by the Dataflow evaluation. */
+	UPROPERTY(EditAnywhere, Category = Collision, Meta = (EditCondition = "DataflowAsset == nullptr"))
 	TObjectPtr<UPhysicsAsset> PhysicsAsset;
 
-	/** Struct containing information for each LOD level, such as materials to use, and when use the LOD. */
-	UPROPERTY(EditAnywhere, EditFixedSize, Category = LevelOfDetails)
+	/** Struct containing information for each LOD level, such as materials to use, and when use the LOD. Not currently editable or customizable through the Dataflow. */
+	UPROPERTY(VisibleAnywhere, EditFixedSize, Category = LevelOfDetails)
 	TArray<FSkeletalMeshLODInfo> LODInfo;
 
 	UPROPERTY(EditAnywhere, Category = LODSettings)
 	FPerPlatformBool DisableBelowMinLodStripping;
 
-	UPROPERTY(EditAnywhere, Category = LODSettings, meta = (DisplayName = "Minimum LOD"))
+	UPROPERTY(EditAnywhere, Category = LODSettings, Meta = (DisplayName = "Minimum LOD"))
 	FPerPlatformInt MinLod;
 
 	/** Enable raytracing for this asset. */
@@ -297,17 +301,20 @@ private:
 	UPROPERTY(EditAnywhere, Category = RayTracing)
 	int32 RayTracingMinLOD;
 
-	/** Whether to blend positions between the skinned/simulated transitions of the cloth render mesh. */
-	UPROPERTY(EditAnywhere, Category = ClothDeformer)
-	bool bSmoothTransition = true;
+	/** Whether to blend positions between the skinned/simulated transitions of the cloth render mesh. Only used when not overriden in Dataflow by the ProxyDefomer node. */
+	UE_DEPRECATED(5.4, "Superseded by the ProxyDefomer node.")
+	UPROPERTY()
+	bool bSmoothTransition_DEPRECATED = true;
 
-	/** Whether to use multiple triangle influences on the proxy wrap deformer to help smoothe deformations. */
-	UPROPERTY(EditAnywhere, Category = ClothDeformer)
-	bool bUseMultipleInfluences = false;
+	/** Whether to use multiple triangle influences on the proxy wrap deformer to help smoothe deformations. Only used when not overriden in Dataflow by the ProxyDefomer node. */
+	UE_DEPRECATED(5.4, "Superseded by the ProxyDefomer node.")
+	UPROPERTY()
+	bool bUseMultipleInfluences_DEPRECATED = false;
 
-	/** The radius from which to get the multiple triangle influences from the simulated proxy mesh. */
-	UPROPERTY(EditAnywhere, Category = ClothDeformer)
-	float SkinningKernelRadius = 30.f;
+	/** The radius from which to get the multiple triangle influences from the simulated proxy mesh. Only used when not overriden in Dataflow with the ProxyDefomer node. */
+	UE_DEPRECATED(5.4, "Superseded by the ProxyDefomer node.")
+	UPROPERTY()
+	float SkinningKernelRadius_DEPRECATED = 30.f;
 
 	/**
 	 * Physics asset whose shapes will be used for shadowing when components have bCastCharacterCapsuleDirectShadow or bCastCharacterCapsuleIndirectShadow enabled.
