Chaos Cloth Asset - Removed unused bones from the ClothAsset RequiredBones lists.

diff --git a/Engine/Plugins/ChaosClothAsset/Source/ChaosClothAssetEngine/Private/ChaosClothAsset/ClothAssetBuilder.cpp b/Engine/Plugins/ChaosClothAsset/Source/ChaosClothAssetEngine/Private/ChaosClothAsset/ClothAssetBuilder.cpp
--- a/Engine/Plugins/ChaosClothAsset/Source/ChaosClothAssetEngine/Private/ChaosClothAsset/ClothAssetBuilder.cpp
+++ b/Engine/Plugins/ChaosClothAsset/Source/ChaosClothAssetEngine/Private/ChaosClothAsset/ClothAssetBuilder.cpp
@@ -480,9 +480,21 @@ PRAGMA_ENABLE_DEPRECATION_WARNINGS
 
 	// Ensure parent exists with incoming active bone indices, and the result should be sorted
 	ClothAsset.GetRefSkeleton().EnsureParentsExistAndSort(LODModel.ActiveBoneIndices);
+	LODModel.ActiveBoneIndices.Shrink();
 
-	// Compute the required bones for this model.
-	USkeletalMesh::CalculateRequiredBones(LODModel, ClothAsset.GetRefSkeleton(), nullptr);
+	// Add the extra simulation bones to the list of required bones if any
+	LODModel.RequiredBones = LODModel.ActiveBoneIndices;
+
+	if (const TSharedPtr<const FChaosClothSimulationModel> ClothSimulationModel = ClothAsset.GetClothSimulationModel())
+	{
+		if (ensure(ClothSimulationModel->IsValidLodIndex(LodIndex)))
+		{
+			LODModel.RequiredBones.Append(ClothSimulationModel->ClothSimulationLodModels[LodIndex].RequiredExtraBoneIndices);
+
+			ClothAsset.GetRefSkeleton().EnsureParentsExistAndSort(LODModel.RequiredBones);
+			LODModel.RequiredBones.Shrink();
+		}
+	}
 }
 
 #undef LOCTEXT_NAMESPACE
diff --git a/Engine/Plugins/ChaosClothAsset/Source/ChaosClothAssetEngine/Private/ChaosClothAsset/ClothComponent.cpp b/Engine/Plugins/ChaosClothAsset/Source/ChaosClothAssetEngine/Private/ChaosClothAsset/ClothComponent.cpp
--- a/Engine/Plugins/ChaosClothAsset/Source/ChaosClothAssetEngine/Private/ChaosClothAsset/ClothComponent.cpp
+++ b/Engine/Plugins/ChaosClothAsset/Source/ChaosClothAssetEngine/Private/ChaosClothAsset/ClothComponent.cpp
@@ -400,6 +400,9 @@ void UChaosClothComponent::SetSkinnedAssetAndUpdate(USkinnedAsset* InSkinnedAsse
 
 void UChaosClothComponent::GetAdditionalRequiredBonesForLeader(int32 LeaderLODIndex, TArray<FBoneIndexType>& InOutRequiredBones) const
 {
+	TArray<FBoneIndexType> RequiredBones;
+
+	// Add the follower's bones (including sim and render mesh bones, both stored in the LODRenderData RequiredBones array)
 	if (const FSkeletalMeshRenderData* const SkeletalMeshRenderData = GetSkeletalMeshRenderData())
 	{
 		const int32 MinLODIndex = ComputeMinLOD();
@@ -409,8 +412,6 @@ void UChaosClothComponent::GetAdditionalRequiredBonesForLeader(int32 LeaderLODIn
 
 		if (SkeletalMeshRenderData->LODRenderData.IsValidIndex(LODIndex))
 		{
-			// Gather the follower's bones
-			TArray<FBoneIndexType> RequiredBones;
 			RequiredBones.Reserve(SkeletalMeshRenderData->LODRenderData[LODIndex].RequiredBones.Num());
 
 			for (const FBoneIndexType RequiredBone : SkeletalMeshRenderData->LODRenderData[LODIndex].RequiredBones)
@@ -427,11 +428,23 @@ void UChaosClothComponent::GetAdditionalRequiredBonesForLeader(int32 LeaderLODIn
 
 			// Then sort array of required bones in hierarchy order
 			RequiredBones.Sort();
+		}
+	}
 
-			// Make sure all of these are in RequiredBones.
-			MergeInBoneIndexArrays(InOutRequiredBones, RequiredBones);
+	// Merge the physics asset bones (the leader's physics asset can be different to this component's cloth asset)
+	if (const UPhysicsAsset* const PhysicsAsset = GetClothAsset() ? GetClothAsset()->GetPhysicsAsset() : nullptr)
+	{
+		if (const USkinnedAsset* const LeaderSkinnedAsset = ensure(LeaderPoseComponent.IsValid()) ? LeaderPoseComponent->GetSkinnedAsset() : nullptr)  // Needs the leader SkinnedAsset for the correct RefSkeleton
+		{
+			USkinnedMeshComponent::GetPhysicsRequiredBones(LeaderSkinnedAsset, PhysicsAsset, RequiredBones);
 		}
 	}
+
+	if (RequiredBones.Num())
+	{
+		// Make sure all of these are in RequiredBones.
+		MergeInBoneIndexArrays(InOutRequiredBones, RequiredBones);
+	}
 }
 
 TSharedPtr<UE::Chaos::ClothAsset::FClothSimulationProxy> UChaosClothComponent::CreateClothSimulationProxy()
