Chaos Cloth Asset: add blueprint support.

diff --git a/Engine/Plugins/ChaosClothAsset/Source/ChaosClothAssetEngine/Private/ChaosClothAsset/ClothAssetInteractor.cpp b/Engine/Plugins/ChaosClothAsset/Source/ChaosClothAssetEngine/Private/ChaosClothAsset/ClothAssetInteractor.cpp
new file mode 100644
--- /dev/null
+++ b/Engine/Plugins/ChaosClothAsset/Source/ChaosClothAssetEngine/Private/ChaosClothAsset/ClothAssetInteractor.cpp
@@ -0,0 +1,296 @@
+// Copyright Epic Games, Inc. All Rights Reserved.
+
+#include "ChaosClothAsset/ClothAssetInteractor.h"
+#include "Chaos/CollectionPropertyFacade.h"
+
+#include UE_INLINE_GENERATED_CPP_BY_NAME(ClothAssetInteractor)
+
+void UChaosClothAssetInteractor::SetProperties(const TArray<TSharedPtr<::Chaos::Softs::FCollectionPropertyFacade>>& InCollectionPropertyFacades)
+{
+	CollectionPropertyFacades.Reset(InCollectionPropertyFacades.Num());
+	for (const TSharedPtr<::Chaos::Softs::FCollectionPropertyFacade>& InFacade : InCollectionPropertyFacades)
+	{
+		CollectionPropertyFacades.Emplace(InFacade);
+	}
+}
+
+void UChaosClothAssetInteractor::ResetProperties()
+{
+	CollectionPropertyFacades.Reset();
+}
+
+TArray<FString> UChaosClothAssetInteractor::GetAllProperties(int32 LODIndex) const
+{
+	TArray<FString> Keys;
+	if (LODIndex == INDEX_NONE)
+	{
+		for (int32 Index = 0; Index < CollectionPropertyFacades.Num(); ++Index)
+		{
+			if (TSharedPtr<::Chaos::Softs::FCollectionPropertyFacade> PropertyFacade = CollectionPropertyFacades[Index].Pin())
+			{
+				if (Keys.IsEmpty())
+				{
+					// This is the first non-empty LOD. We can add all keys without worrying about uniqueness.
+					Keys.Reserve(PropertyFacade->Num());
+					for (int32 KeyIndex = 0; KeyIndex < PropertyFacade->Num(); ++KeyIndex)
+					{
+						Keys.Add(PropertyFacade->GetKey(KeyIndex));
+					}
+				}
+				else
+				{
+					for (int32 KeyIndex = 0; KeyIndex < PropertyFacade->Num(); ++KeyIndex)
+					{
+						Keys.AddUnique(PropertyFacade->GetKey(KeyIndex));
+					}
+				}
+			}
+		}
+	}
+	else
+	{
+		if (CollectionPropertyFacades.IsValidIndex(LODIndex))
+		{
+			if (TSharedPtr<::Chaos::Softs::FCollectionPropertyFacade> PropertyFacade = CollectionPropertyFacades[LODIndex].Pin())
+			{
+				Keys.Reserve(PropertyFacade->Num());
+				for (int32 KeyIndex = 0; KeyIndex < PropertyFacade->Num(); ++KeyIndex)
+				{
+					Keys.Add(PropertyFacade->GetKey(KeyIndex));
+				}
+			}
+		}
+	}
+	return Keys;
+}
+
+float UChaosClothAssetInteractor::GetFloatValue(const FString& PropertyName, int32 LODIndex, float DefaultValue) const
+{
+	if (CollectionPropertyFacades.IsValidIndex(LODIndex))
+	{
+		if (TSharedPtr<::Chaos::Softs::FCollectionPropertyFacade> PropertyFacade = CollectionPropertyFacades[LODIndex].Pin())
+		{
+			return PropertyFacade->GetValue<float>(PropertyName, DefaultValue);
+		}
+	}
+	return DefaultValue;
+}
+
+float UChaosClothAssetInteractor::GetLowFloatValue(const FString& PropertyName, int32 LODIndex, float DefaultValue) const
+{
+	if (CollectionPropertyFacades.IsValidIndex(LODIndex))
+	{
+		if (TSharedPtr<::Chaos::Softs::FCollectionPropertyFacade> PropertyFacade = CollectionPropertyFacades[LODIndex].Pin())
+		{
+			return PropertyFacade->GetLowValue<float>(PropertyName, DefaultValue);
+		}
+	}
+	return DefaultValue;
+}
+
+float UChaosClothAssetInteractor::GetHighFloatValue(const FString& PropertyName, int32 LODIndex, float DefaultValue) const
+{
+	if (CollectionPropertyFacades.IsValidIndex(LODIndex))
+	{
+		if (TSharedPtr<::Chaos::Softs::FCollectionPropertyFacade> PropertyFacade = CollectionPropertyFacades[LODIndex].Pin())
+		{
+			return PropertyFacade->GetHighValue<float>(PropertyName, DefaultValue);
+		}
+	}
+	return DefaultValue;
+}
+
+FVector2D UChaosClothAssetInteractor::GetWeightedFloatValue(const FString& PropertyName, int32 LODIndex, FVector2D DefaultValue) const
+{
+	if (CollectionPropertyFacades.IsValidIndex(LODIndex))
+	{
+		if (TSharedPtr<::Chaos::Softs::FCollectionPropertyFacade> PropertyFacade = CollectionPropertyFacades[LODIndex].Pin())
+		{
+			return FVector2D(PropertyFacade->GetWeightedFloatValue(PropertyName, FVector2f(DefaultValue)));
+		}
+	}
+	return DefaultValue;
+}
+
+int32 UChaosClothAssetInteractor::GetIntValue(const FString& PropertyName, int32 LODIndex, int32 DefaultValue) const
+{
+	if (CollectionPropertyFacades.IsValidIndex(LODIndex))
+	{
+		if (TSharedPtr<::Chaos::Softs::FCollectionPropertyFacade> PropertyFacade = CollectionPropertyFacades[LODIndex].Pin())
+		{
+			return PropertyFacade->GetValue<int32>(PropertyName, DefaultValue);
+		}
+	}
+	return DefaultValue;
+}
+
+FVector UChaosClothAssetInteractor::GetVectorValue(const FString& PropertyName, int32 LODIndex, FVector DefaultValue) const
+{
+	if (CollectionPropertyFacades.IsValidIndex(LODIndex))
+	{
+		if (TSharedPtr<::Chaos::Softs::FCollectionPropertyFacade> PropertyFacade = CollectionPropertyFacades[LODIndex].Pin())
+		{
+			return FVector(PropertyFacade->GetValue<FVector3f>(PropertyName, FVector3f(DefaultValue)));
+		}
+	}
+	return DefaultValue;
+}
+
+FString UChaosClothAssetInteractor::GetStringValue(const FString& PropertyName, int32 LODIndex, const FString& DefaultValue) const
+{
+	if (CollectionPropertyFacades.IsValidIndex(LODIndex))
+	{
+		if (TSharedPtr<::Chaos::Softs::FCollectionPropertyFacade> PropertyFacade = CollectionPropertyFacades[LODIndex].Pin())
+		{
+			return PropertyFacade->GetStringValue(PropertyName, DefaultValue);
+		}
+	}
+	return DefaultValue;
+}
+
+void UChaosClothAssetInteractor::SetFloatValue(const FString& PropertyName, int32 LODIndex, float Value)
+{
+	if (LODIndex == INDEX_NONE)
+	{
+		for (int32 Index = 0; Index < CollectionPropertyFacades.Num(); ++Index)
+		{
+			if (TSharedPtr<::Chaos::Softs::FCollectionPropertyFacade> PropertyFacade = CollectionPropertyFacades[Index].Pin())
+			{
+				PropertyFacade->SetValue(PropertyName, Value);
+			}
+		}
+	}
+	else if (CollectionPropertyFacades.IsValidIndex(LODIndex))
+	{
+		if (TSharedPtr<::Chaos::Softs::FCollectionPropertyFacade> PropertyFacade = CollectionPropertyFacades[LODIndex].Pin())
+		{
+			PropertyFacade->SetValue(PropertyName, Value);
+		}
+	}
+}
+
+void UChaosClothAssetInteractor::SetLowFloatValue(const FString& PropertyName, int32 LODIndex, float Value)
+{
+	if (LODIndex == INDEX_NONE)
+	{
+		for (int32 Index = 0; Index < CollectionPropertyFacades.Num(); ++Index)
+		{
+			if (TSharedPtr<::Chaos::Softs::FCollectionPropertyFacade> PropertyFacade = CollectionPropertyFacades[Index].Pin())
+			{
+				PropertyFacade->SetLowValue(PropertyName, Value);
+			}
+		}
+	}
+	else if (CollectionPropertyFacades.IsValidIndex(LODIndex))
+	{
+		if (TSharedPtr<::Chaos::Softs::FCollectionPropertyFacade> PropertyFacade = CollectionPropertyFacades[LODIndex].Pin())
+		{
+			PropertyFacade->SetLowValue(PropertyName, Value);
+		}
+	}
+}
+
+void UChaosClothAssetInteractor::SetHighFloatValue(const FString& PropertyName, int32 LODIndex, float Value)
+{
+	if (LODIndex == INDEX_NONE)
+	{
+		for (int32 Index = 0; Index < CollectionPropertyFacades.Num(); ++Index)
+		{
+			if (TSharedPtr<::Chaos::Softs::FCollectionPropertyFacade> PropertyFacade = CollectionPropertyFacades[Index].Pin())
+			{
+				PropertyFacade->SetHighValue(PropertyName, Value);
+			}
+		}
+	}
+	else if (CollectionPropertyFacades.IsValidIndex(LODIndex))
+	{
+		if (TSharedPtr<::Chaos::Softs::FCollectionPropertyFacade> PropertyFacade = CollectionPropertyFacades[LODIndex].Pin())
+		{
+			PropertyFacade->SetHighValue(PropertyName, Value);
+		}
+	}
+}
+
+void UChaosClothAssetInteractor::SetWeightedFloatValue(const FString& PropertyName, int32 LODIndex, FVector2D Value)
+{
+	if (LODIndex == INDEX_NONE)
+	{
+		for (int32 Index = 0; Index < CollectionPropertyFacades.Num(); ++Index)
+		{
+			if (TSharedPtr<::Chaos::Softs::FCollectionPropertyFacade> PropertyFacade = CollectionPropertyFacades[Index].Pin())
+			{
+				PropertyFacade->SetWeightedFloatValue(PropertyName, FVector2f(Value));
+			}
+		}
+	}
+	else if (CollectionPropertyFacades.IsValidIndex(LODIndex))
+	{
+		if (TSharedPtr<::Chaos::Softs::FCollectionPropertyFacade> PropertyFacade = CollectionPropertyFacades[LODIndex].Pin())
+		{
+			PropertyFacade->SetWeightedFloatValue(PropertyName, FVector2f(Value));
+		}
+	}
+}
+
+void UChaosClothAssetInteractor::SetIntValue(const FString& PropertyName, int32 LODIndex, int32 Value)
+{
+	if (LODIndex == INDEX_NONE)
+	{
+		for (int32 Index = 0; Index < CollectionPropertyFacades.Num(); ++Index)
+		{
+			if (TSharedPtr<::Chaos::Softs::FCollectionPropertyFacade> PropertyFacade = CollectionPropertyFacades[Index].Pin())
+			{
+				PropertyFacade->SetValue(PropertyName, Value);
+			}
+		}
+	}
+	else if (CollectionPropertyFacades.IsValidIndex(LODIndex))
+	{
+		if (TSharedPtr<::Chaos::Softs::FCollectionPropertyFacade> PropertyFacade = CollectionPropertyFacades[LODIndex].Pin())
+		{
+			PropertyFacade->SetValue(PropertyName, Value);
+		}
+	}
+}
+
+void UChaosClothAssetInteractor::SetVectorValue(const FString& PropertyName, int32 LODIndex, FVector Value)
+{
+	if (LODIndex == INDEX_NONE)
+	{
+		for (int32 Index = 0; Index < CollectionPropertyFacades.Num(); ++Index)
+		{
+			if (TSharedPtr<::Chaos::Softs::FCollectionPropertyFacade> PropertyFacade = CollectionPropertyFacades[Index].Pin())
+			{
+				PropertyFacade->SetValue(PropertyName, FVector3f(Value));
+			}
+		}
+	}
+	else if (CollectionPropertyFacades.IsValidIndex(LODIndex))
+	{
+		if (TSharedPtr<::Chaos::Softs::FCollectionPropertyFacade> PropertyFacade = CollectionPropertyFacades[LODIndex].Pin())
+		{
+			PropertyFacade->SetValue(PropertyName, FVector3f(Value));
+		}
+	}
+}
+
+void UChaosClothAssetInteractor::SetStringValue(const FString& PropertyName, int32 LODIndex, const FString& Value)
+{
+	if (LODIndex == INDEX_NONE)
+	{
+		for (int32 Index = 0; Index < CollectionPropertyFacades.Num(); ++Index)
+		{
+			if (TSharedPtr<::Chaos::Softs::FCollectionPropertyFacade> PropertyFacade = CollectionPropertyFacades[Index].Pin())
+			{
+				PropertyFacade->SetStringValue(PropertyName, Value);
+			}
+		}
+	}
+	else if (CollectionPropertyFacades.IsValidIndex(LODIndex))
+	{
+		if (TSharedPtr<::Chaos::Softs::FCollectionPropertyFacade> PropertyFacade = CollectionPropertyFacades[LODIndex].Pin())
+		{
+			PropertyFacade->SetStringValue(PropertyName, Value);
+		}
+	}
+}
diff --git a/Engine/Plugins/ChaosClothAsset/Source/ChaosClothAssetEngine/Private/ChaosClothAsset/ClothComponent.cpp b/Engine/Plugins/ChaosClothAsset/Source/ChaosClothAssetEngine/Private/ChaosClothAsset/ClothComponent.cpp
--- a/Engine/Plugins/ChaosClothAsset/Source/ChaosClothAssetEngine/Private/ChaosClothAsset/ClothComponent.cpp
+++ b/Engine/Plugins/ChaosClothAsset/Source/ChaosClothAssetEngine/Private/ChaosClothAsset/ClothComponent.cpp
@@ -2,6 +2,7 @@
 
 #include "ChaosClothAsset/ClothComponent.h"
 #include "ChaosClothAsset/ClothAsset.h"
+#include "ChaosClothAsset/ClothAssetInteractor.h"
 #include "ChaosClothAsset/ClothAssetPrivate.h"
 #include "ChaosClothAsset/ClothSimulationModel.h"
 #include "ChaosClothAsset/ClothSimulationProxy.h"
@@ -85,12 +86,15 @@ void UChaosClothComponent::ResetConfigProperties()
 				::Chaos::Softs::FCollectionPropertyMutableFacade CollectionPropertyMutableFacade(PropertyCollection);
 				CollectionPropertyMutableFacade.Copy(*ClothCollection);
 
-				CollectionPropertyFacades.Add(MakeUnique<::Chaos::Softs::FCollectionPropertyFacade>(PropertyCollection));
-
+				CollectionPropertyFacades.Add(MakeShared<::Chaos::Softs::FCollectionPropertyFacade>(PropertyCollection));
 			}
+			check(ClothOutfitInteractor);
+			ClothOutfitInteractor->SetProperties(CollectionPropertyFacades);
 		}
 		else
 		{
+			check(ClothOutfitInteractor);
+			ClothOutfitInteractor->ResetProperties();
 			PropertyCollections.Reset();
 			CollectionPropertyFacades.Reset();
 		}
@@ -195,6 +199,7 @@ void UChaosClothComponent::OnRegister()
 	UpdateComponentSpaceTransforms();
 
 	// Fill up the property collection with the original cloth asset properties
+	ClothOutfitInteractor = NewObject<UChaosClothAssetInteractor>();
 	ResetConfigProperties();
 
 	// Create the proxy to start the simulation
@@ -212,6 +217,7 @@ void UChaosClothComponent::OnUnregister()
 	ClothSimulationProxy.Reset();
 
 	// Release the runtime simulation collection and facade
+	ClothOutfitInteractor->ResetProperties();
 	CollectionPropertyFacades.Empty();
 	PropertyCollections.Empty();
 }
@@ -519,3 +525,10 @@ void UChaosClothComponent::UpdateVisibility()
 		SetVisibility(false);
 	}
 }
+
+UChaosClothAssetInteractor* UChaosClothComponent::GetClothOutfitInteractor()
+{
+	check(IsInGameThread());
+	return ClothOutfitInteractor;
+}
+
diff --git a/Engine/Plugins/ChaosClothAsset/Source/ChaosClothAssetEngine/Public/ChaosClothAsset/ClothAssetInteractor.h b/Engine/Plugins/ChaosClothAsset/Source/ChaosClothAssetEngine/Public/ChaosClothAsset/ClothAssetInteractor.h
new file mode 100644
--- /dev/null
+++ b/Engine/Plugins/ChaosClothAsset/Source/ChaosClothAssetEngine/Public/ChaosClothAsset/ClothAssetInteractor.h
@@ -0,0 +1,87 @@
+// Copyright Epic Games, Inc. All Rights Reserved.
+
+#pragma once
+#include "UObject/Object.h"
+#include "CoreTypes.h"
+#include "Math/MathFwd.h"
+#include "ClothAssetInteractor.generated.h"
+
+namespace Chaos::Softs
+{
+class FCollectionPropertyFacade;
+}
+
+UCLASS(BlueprintType, MinimalAPI)
+class UChaosClothAssetInteractor : public UObject
+{
+	GENERATED_BODY()
+public:
+
+	/** Set properties this interactor references.*/
+	CHAOSCLOTHASSETENGINE_API void SetProperties(const TArray<TSharedPtr<::Chaos::Softs::FCollectionPropertyFacade>>& InCollectionPropertyFacades);
+
+	/** Empty references to all properties.*/
+	CHAOSCLOTHASSETENGINE_API void ResetProperties();
+
+	/** Generate a list of all properties held by this interactor. Properties for all LODs will be returned if LODIndex = -1.*/
+	UFUNCTION(BlueprintCallable, Category = ClothProperty)
+	TArray<FString> GetAllProperties(int32 LODIndex = -1) const;
+
+	/** Get the value for a property cast to float. DefaultValue will be returned if the property is not found.*/
+	UFUNCTION(BlueprintCallable, Category = ClothProperty)
+	float GetFloatValue(const FString& PropertyName, int32 LODIndex = 0, float DefaultValue = 0.f) const;
+
+	/** Get the low value for a weighted property value (same as GetFloatValue). DefaultValue will be returned if the property is not found.*/
+	UFUNCTION(BlueprintCallable, Category = ClothProperty)
+	float GetLowFloatValue(const FString& PropertyName, int32 LODIndex = 0, float DefaultValue = 0.f) const;
+
+	/** Get the high value for a weighted property value. DefaultValue will be returned if the property is not found.*/
+	UFUNCTION(BlueprintCallable, Category = ClothProperty)
+	float GetHighFloatValue(const FString& PropertyName, int32 LODIndex = 0, float DefaultValue = 0.f) const;
+
+	/** Get the low and high values for a weighted property value. DefaultValue will be returned if the property is not found.*/
+	UFUNCTION(BlueprintCallable, Category = ClothProperty)
+	FVector2D GetWeightedFloatValue(const FString& PropertyName, int32 LODIndex = 0, FVector2D DefaultValue = FVector2D(0., 0.)) const;
+
+	/** Get the value for a property cast to int. DefaultValue will be returned if the property is not found.*/
+	UFUNCTION(BlueprintCallable, Category = ClothProperty)
+	int32 GetIntValue(const FString& PropertyName, int32 LODIndex = 0, int32 DefaultValue = 0) const;
+
+	/** Get the value for a property cast to vector. DefaultValue will be returned if the property is not found.*/
+	UFUNCTION(BlueprintCallable, Category = ClothProperty)
+	FVector GetVectorValue(const FString& PropertyName, int32 LODIndex = 0, FVector DefaultValue = FVector(0.)) const;
+
+	/** Get the string value for a property (typically the associated map name for weighted values). DefaultValue will be returned if the property is not found.*/
+	UFUNCTION(BlueprintCallable, Category = ClothProperty)
+	FString GetStringValue(const FString& PropertyName, int32 LODIndex = 0, const FString& DefaultValue = "") const;
+
+	/**Set the value for a property (if it exists). This sets the Low and High values for weighted values. All LODs will be set when LODIndex = -1.*/
+	UFUNCTION(BlueprintCallable, Category = ClothProperty)
+	void SetFloatValue(const FString& PropertyName, int32 LODIndex = -1, float Value = 0.f);
+
+	/**Set the low value for a weighted property (if it exists). All LODs will be set when LODIndex = -1.*/
+	UFUNCTION(BlueprintCallable, Category = ClothProperty)
+	void SetLowFloatValue(const FString& PropertyName, int32 LODIndex = -1, float Value = 0.f);
+
+	/**Set the high value for a weighted property (if it exists). All LODs will be set when LODIndex = -1.*/
+	UFUNCTION(BlueprintCallable, Category = ClothProperty)
+	void SetHighFloatValue(const FString& PropertyName, int32 LODIndex = -1, float Value = 0.f);
+
+	/**Set the low and high values for a weighted property (if it exists). All LODs will be set when LODIndex = -1.*/
+	UFUNCTION(BlueprintCallable, Category = ClothProperty)
+	void SetWeightedFloatValue(const FString& PropertyName, int32 LODIndex = -1, FVector2D Value = FVector2D(0., 0.));
+
+	/**Set the value for a property (if it exists). All LODs will be set when LODIndex = -1.*/
+	UFUNCTION(BlueprintCallable, Category = ClothProperty)
+	void SetIntValue(const FString& PropertyName, int32 LODIndex = -1, int32 Value = 0);
+
+	/**Set the value for a property (if it exists). All LODs will be set when LODIndex = -1.*/
+	UFUNCTION(BlueprintCallable, Category = ClothProperty)
+	void SetVectorValue(const FString& PropertyName, int32 LODIndex = -1, FVector Value = FVector(0.));
+
+	/**Set the string value for a property (if it exists). This is typically the map name associated with a property. All LODs will be set when LODIndex = -1.*/
+	void SetStringValue(const FString& PropertyName, int32 LODIndex = -1, const FString& Value = "");
+
+private:
+	TArray<TWeakPtr<::Chaos::Softs::FCollectionPropertyFacade>> CollectionPropertyFacades;
+};
\ No newline at end of file
diff --git a/Engine/Plugins/ChaosClothAsset/Source/ChaosClothAssetEngine/Public/ChaosClothAsset/ClothComponent.h b/Engine/Plugins/ChaosClothAsset/Source/ChaosClothAssetEngine/Public/ChaosClothAsset/ClothComponent.h
--- a/Engine/Plugins/ChaosClothAsset/Source/ChaosClothAssetEngine/Public/ChaosClothAsset/ClothComponent.h
+++ b/Engine/Plugins/ChaosClothAsset/Source/ChaosClothAssetEngine/Public/ChaosClothAsset/ClothComponent.h
@@ -7,6 +7,7 @@
 
 class UChaosClothAsset;
 class UChaosClothComponent;
+class UChaosClothAssetInteractor;
 struct FManagedArrayCollection;
 
 namespace Chaos::Softs
@@ -82,7 +83,7 @@ public:
 	UFUNCTION(BlueprintCallable, Category = "ClothComponent", Meta = (Keywords = "Chaos Cloth Simulation Enable"))
 	bool IsSimulationEnabled() const;
 
-	/** Reset all cloth simulation config properties to the values stored in the original cloth asset. */
+	/** Reset all cloth simulation config properties to the values stored in the original cloth asset.*/
 	UFUNCTION(BlueprintCallable, Category = "ClothComponent", Meta = (Keywords = "Chaos Cloth Config Property"))
 	void ResetConfigProperties();
 
@@ -90,6 +91,11 @@ public:
 	UFUNCTION(CallInEditor, BlueprintCallable, Category = "ClothComponent", Meta = (DisplayName = "Hard Reset Simulation", Keywords = "Chaos Cloth Recreate Simulation Proxy"))
 	void RecreateClothSimulationProxy();
 
+	/** Get the current interactor for the cloth outfit associated with this cloth component. 
+	 * Interact with solver-level properties as well as all cloth assets within the cloth outfit (once multi-asset outfits exist).*/
+	UFUNCTION(BlueprintCallable, Category = "ClothComponent")
+	UChaosClothAssetInteractor* GetClothOutfitInteractor();
+
 	/**
 	 * Return the property collections holding the runtime properties for this cloth component (one per LOD).
 	 * This might be different from the cloth asset's since the component's properties can be modified in code or in blueprints.
@@ -203,8 +209,11 @@ private:
 	bool bTickOnceInEditor = false;
 #endif
 
+	UPROPERTY(Transient)
+	TObjectPtr<UChaosClothAssetInteractor> ClothOutfitInteractor;
+
 	TArray<TSharedPtr<FManagedArrayCollection>> PropertyCollections;
-	TArray<TUniquePtr<::Chaos::Softs::FCollectionPropertyFacade>> CollectionPropertyFacades;
+	TArray<TSharedPtr<::Chaos::Softs::FCollectionPropertyFacade>> CollectionPropertyFacades;
 
 	TSharedPtr<UE::Chaos::ClothAsset::FClothSimulationProxy> ClothSimulationProxy;
 };
