Chaos Cloth Asset - Fixed Cloth Asset size calculation by adding a GetResourceSizeEx method to the Cloth Asset class.

diff --git a/Engine/Plugins/ChaosClothAsset/Source/ChaosClothAssetEngine/Private/ChaosClothAsset/ClothAsset.cpp b/Engine/Plugins/ChaosClothAsset/Source/ChaosClothAssetEngine/Private/ChaosClothAsset/ClothAsset.cpp
--- a/Engine/Plugins/ChaosClothAsset/Source/ChaosClothAssetEngine/Private/ChaosClothAsset/ClothAsset.cpp
+++ b/Engine/Plugins/ChaosClothAsset/Source/ChaosClothAssetEngine/Private/ChaosClothAsset/ClothAsset.cpp
@@ -193,7 +193,6 @@ FMatrix UChaosClothAsset::GetComposedRefPoseMatrix(FName InBoneName) const
 	return LocalPose;
 }
 
-
 void UChaosClothAsset::Serialize(FArchive& Ar)
 {
 	Super::Serialize(Ar);
@@ -201,7 +200,6 @@ void UChaosClothAsset::Serialize(FArchive& Ar)
 	bool bCooked = Ar.IsCooking();
 	Ar << bCooked;
 
-
 	if (bCooked && Ar.IsSaving())
 	{
 		TArray<TSharedRef<FManagedArrayCollection>> OutputCollections = 
@@ -217,7 +215,7 @@ void UChaosClothAsset::Serialize(FArchive& Ar)
 
 	Ar << GetRefSkeleton();
 
-	if (bCooked && !IsTemplate() && !Ar.IsCountingMemory())
+	if (bCooked && !IsTemplate() && !Ar.IsCountingMemory())  // Counting of these resources are done in GetResourceSizeEx, so skip these when counting memory
 	{
 		if (Ar.IsLoading())
 		{
@@ -250,6 +248,65 @@ void UChaosClothAsset::PostEditChangeProperty(FPropertyChangedEvent& PropertyCha
 }
 #endif // #if WITH_EDITOR
 
+void UChaosClothAsset::GetResourceSizeEx(FResourceSizeEx& CumulativeResourceSize)
+{
+	Super::GetResourceSizeEx(CumulativeResourceSize);
+
+	GetResourceForRendering()->GetResourceSizeEx(CumulativeResourceSize);
+	ClothSimulationModel->GetResourceSizeEx(CumulativeResourceSize);
+
+#if !UE_BUILD_SHIPPING
+	FString MemoryReport;
+	MemoryReport.Appendf(TEXT("---- Memory report for Cloth Asset [%s] ----"), *this->GetName());
+
+	FResourceSizeEx RenderDataResourceSize;
+	if (GetResourceForRendering())
+	{
+		for (int32 LodIndex = 0; LodIndex < GetResourceForRendering()->LODRenderData.Num(); ++LodIndex)
+		{
+			FResourceSizeEx LODRenderDataResourceSize;
+			GetResourceForRendering()->LODRenderData[LodIndex].GetResourceSizeEx(LODRenderDataResourceSize);
+			MemoryReport.Appendf(TEXT("\n LODRenderData LOD%d size: %ld bytes"), LodIndex, LODRenderDataResourceSize.GetTotalMemoryBytes());
+		}
+
+		GetResourceForRendering()->GetResourceSizeEx(RenderDataResourceSize);
+	}
+	MemoryReport.Appendf(TEXT("\n Total RenderData size: %ld bytes"), RenderDataResourceSize.GetTotalMemoryBytes());
+
+	FResourceSizeEx ClothSimulationModelResourceSize;
+	if (ClothSimulationModel)
+	{
+		for (int32 LodIndex = 0; LodIndex < ClothSimulationModel->GetNumLods(); ++LodIndex)
+		{
+			FResourceSizeEx ClothSimulationLodModelResourceSize;
+			ClothSimulationModel->ClothSimulationLodModels[LodIndex].GetResourceSizeEx(ClothSimulationLodModelResourceSize);
+
+			const FName Tag(*FString::Printf(TEXT("ClothSimulationLodModel[%d]"), LodIndex));
+			MemoryReport.Appendf(TEXT("\n ClothSimulationLodModel LOD%d size: %ld bytes"), LodIndex, ClothSimulationLodModelResourceSize.GetTotalMemoryBytes());
+		}
+
+		ClothSimulationModel->GetResourceSizeEx(ClothSimulationModelResourceSize);
+	}
+	MemoryReport.Appendf(TEXT("\n Total ClothSimulationModel size: %ld bytes"), ClothSimulationModelResourceSize.GetTotalMemoryBytes());
+
+	const int64 TotalResourceSize = RenderDataResourceSize.GetTotalMemoryBytes() + ClothSimulationModelResourceSize.GetTotalMemoryBytes();
+	MemoryReport.Appendf(
+		TEXT("\n Total resource size for Cloth Asset [%s]: %ld bytes (%.3f MB)"),
+		*this->GetName(),
+		TotalResourceSize,
+		(float)TotalResourceSize / (1024.f * 1024.f));
+
+	const int64 TotalSize = CumulativeResourceSize.GetTotalMemoryBytes();
+	MemoryReport.Appendf(
+		TEXT("\n Total size for Cloth Asset [%s]: %ld bytes (%.3f MB)"),
+		*this->GetName(),
+		TotalSize,
+		(float)TotalSize / (1024.f * 1024.f));
+
+	UE_LOG(LogChaosClothAsset, Display, TEXT("\n%s"), *MemoryReport);
+#endif
+}
+
 void UChaosClothAsset::BeginPostLoadInternal(FSkinnedAssetPostLoadContext& Context)
 {
 #if WITH_EDITOR
diff --git a/Engine/Plugins/ChaosClothAsset/Source/ChaosClothAssetEngine/Private/ChaosClothAsset/ClothSimulationModel.cpp b/Engine/Plugins/ChaosClothAsset/Source/ChaosClothAssetEngine/Private/ChaosClothAsset/ClothSimulationModel.cpp
--- a/Engine/Plugins/ChaosClothAsset/Source/ChaosClothAssetEngine/Private/ChaosClothAsset/ClothSimulationModel.cpp
+++ b/Engine/Plugins/ChaosClothAsset/Source/ChaosClothAssetEngine/Private/ChaosClothAsset/ClothSimulationModel.cpp
@@ -13,6 +13,25 @@
 
 namespace UE::Chaos::ClothAsset::Private
 {
+	// Memory counter archive for GetResourceSizeEx
+	class FMemoryCounterArchive final : public FArchive
+	{
+	public:
+		FMemoryCounterArchive()
+			: Size(0)
+		{
+			SetIsSaving(true);
+			SetIsPersistent(true);
+			ArIsCountingMemory = true;
+		}
+
+		virtual void Serialize(void*, int64 Length) override { Size += Length; }
+		virtual int64 TotalSize() override { return Size; }
+
+	private:
+		int64 Size;
+	};
+
 	// Find the root bone for this cloth asset (common bone for all used bones)
 	static int32 CalculateReferenceBoneIndex(const TArray<FChaosClothSimulationLodModel> ClothSimulationLodModels, const FReferenceSkeleton& ReferenceSkeleton, const TArray<int32>& UsedBoneIndices)
 	{
@@ -137,7 +156,7 @@ bool FChaosClothSimulationLodModel::Serialize(FArchive& Ar)
 	using namespace UE::Chaos::ClothAsset;
 
 	// Serialize normal tagged property data
-	if (Ar.IsLoading() || Ar.IsSaving())
+	if (Ar.IsLoading() || Ar.IsSaving() || Ar.IsCountingMemory())
 	{
 		UScriptStruct* const Struct = FChaosClothSimulationLodModel::StaticStruct();
 		Struct->SerializeTaggedProperties(Ar, (uint8*)this, Struct, nullptr);
@@ -217,6 +236,13 @@ bool FChaosClothSimulationLodModel::Serialize(FArchive& Ar)
 	return true;
 }
 
+void FChaosClothSimulationLodModel::GetResourceSizeEx(FResourceSizeEx& CumulativeResourceSize)
+{
+	UE::Chaos::ClothAsset::Private::FMemoryCounterArchive Ar;
+	Serialize(Ar);
+	CumulativeResourceSize.AddDedicatedSystemMemoryBytes(Ar.TotalSize());
+}
+
 FChaosClothSimulationModel::FChaosClothSimulationModel(const TArray<TSharedRef<const FManagedArrayCollection>>& ClothCollections, 
 	const FReferenceSkeleton& ReferenceSkeleton,
 	TArray<FChaosClothAssetLodTransitionDataCache>* InOutTransitionCache)
@@ -452,4 +478,14 @@ void FChaosClothSimulationModel::CalculateLODTransitionUpDownData(TArray<FChaosC
 			}
 		}
 	}
-}
\ No newline at end of file
+}
+
+void FChaosClothSimulationModel::GetResourceSizeEx(FResourceSizeEx& CumulativeResourceSize)
+{
+	UE::Chaos::ClothAsset::Private::FMemoryCounterArchive Ar;
+
+	UScriptStruct* const Struct = FChaosClothSimulationModel::StaticStruct();
+	Struct->SerializeTaggedProperties(Ar, (uint8*)this, Struct, nullptr);
+
+	CumulativeResourceSize.AddDedicatedSystemMemoryBytes(Ar.TotalSize());
+}
diff --git a/Engine/Plugins/ChaosClothAsset/Source/ChaosClothAssetEngine/Private/ChaosClothAsset/ClothSimulationModel.h b/Engine/Plugins/ChaosClothAsset/Source/ChaosClothAssetEngine/Private/ChaosClothAsset/ClothSimulationModel.h
--- a/Engine/Plugins/ChaosClothAsset/Source/ChaosClothAssetEngine/Private/ChaosClothAsset/ClothSimulationModel.h
+++ b/Engine/Plugins/ChaosClothAsset/Source/ChaosClothAssetEngine/Private/ChaosClothAsset/ClothSimulationModel.h
@@ -77,6 +77,8 @@ struct FChaosClothSimulationLodModel
 
 	// Custom serialize for weight maps
 	bool Serialize(FArchive& Ar);
+
+	void GetResourceSizeEx(FResourceSizeEx& CumulativeResourceSize);
 };
 
 template<>
@@ -131,6 +133,7 @@ struct FChaosClothSimulationModel
 	TConstArrayView<FClothVertBoneData> GetBoneData(int32 LodIndex) const { return IsValidLodIndex(LodIndex) ? ClothSimulationLodModels[LodIndex].BoneData : TConstArrayView<FClothVertBoneData>(); }
 	TArray<TConstArrayView<TTuple<int32, int32, float>>> GetTethers(int32 LodIndex) const;
 
-
 	void CalculateLODTransitionUpDownData(TArray<FChaosClothAssetLodTransitionDataCache>* InOutTransitionCache = nullptr);
+
+	void GetResourceSizeEx(FResourceSizeEx& CumulativeResourceSize);
 };
diff --git a/Engine/Plugins/ChaosClothAsset/Source/ChaosClothAssetEngine/Public/ChaosClothAsset/ClothAsset.h b/Engine/Plugins/ChaosClothAsset/Source/ChaosClothAssetEngine/Public/ChaosClothAsset/ClothAsset.h
--- a/Engine/Plugins/ChaosClothAsset/Source/ChaosClothAssetEngine/Public/ChaosClothAsset/ClothAsset.h
+++ b/Engine/Plugins/ChaosClothAsset/Source/ChaosClothAssetEngine/Public/ChaosClothAsset/ClothAsset.h
@@ -51,6 +51,7 @@ public:
 #if WITH_EDITOR
 	virtual void PostEditChangeProperty(FPropertyChangedEvent& PropertyChangedEvent) override;
 #endif
+	virtual void GetResourceSizeEx(FResourceSizeEx& CumulativeResourceSize) override;
 	//~ End UObject interface
 
 	//~ Begin USkinnedAsset interface
