CIS - Fixed Static Analysis warning "The 'GetResourceForRendering()' pointer was utilized before it was verified against nullptr."

diff --git a/Engine/Plugins/ChaosClothAsset/Source/ChaosClothAssetEngine/Private/ChaosClothAsset/ClothAsset.cpp b/Engine/Plugins/ChaosClothAsset/Source/ChaosClothAssetEngine/Private/ChaosClothAsset/ClothAsset.cpp
--- a/Engine/Plugins/ChaosClothAsset/Source/ChaosClothAssetEngine/Private/ChaosClothAsset/ClothAsset.cpp
+++ b/Engine/Plugins/ChaosClothAsset/Source/ChaosClothAssetEngine/Private/ChaosClothAsset/ClothAsset.cpp
@@ -252,8 +252,15 @@ void UChaosClothAsset::GetResourceSizeEx(FResourceSizeEx& CumulativeResourceSize
 {
 	Super::GetResourceSizeEx(CumulativeResourceSize);
 
-	GetResourceForRendering()->GetResourceSizeEx(CumulativeResourceSize);
-	ClothSimulationModel->GetResourceSizeEx(CumulativeResourceSize);
+	if (GetResourceForRendering())
+	{
+		GetResourceForRendering()->GetResourceSizeEx(CumulativeResourceSize);
+	}
+
+	if (ClothSimulationModel)
+	{
+		ClothSimulationModel->GetResourceSizeEx(CumulativeResourceSize);
+	}
 
 #if !UE_BUILD_SHIPPING
 	FString MemoryReport;
