Chaos Cloth: actually disable simulation when EnableSimulation = false on the Cloth Component.

Also, don't generate the CurrentSimulationData since this is also unused. Do continue to update the solver on LOD switches because this data is used to calculate the BoundingBox.

diff --git a/Engine/Plugins/ChaosClothAsset/Source/ChaosClothAssetEngine/Private/ChaosClothAsset/ClothComponent.cpp b/Engine/Plugins/ChaosClothAsset/Source/ChaosClothAssetEngine/Private/ChaosClothAsset/ClothComponent.cpp
--- a/Engine/Plugins/ChaosClothAsset/Source/ChaosClothAssetEngine/Private/ChaosClothAsset/ClothComponent.cpp
+++ b/Engine/Plugins/ChaosClothAsset/Source/ChaosClothAssetEngine/Private/ChaosClothAsset/ClothComponent.cpp
@@ -56,9 +56,7 @@ UChaosClothAsset* UChaosClothComponent::GetClothAsset() const
 
 bool UChaosClothComponent::IsSimulationSuspended() const
 {
-	static IConsoleVariable* const CVarClothPhysics = IConsoleManager::Get().FindConsoleVariable(TEXT("p.ClothPhysics"));
-
-	return bSuspendSimulation || !ClothSimulationProxy.IsValid() || (CVarClothPhysics && !CVarClothPhysics->GetBool());
+	return bSuspendSimulation || !IsSimulationEnabled();
 }
 
 bool UChaosClothComponent::IsSimulationEnabled() const
@@ -361,7 +359,7 @@ void UChaosClothComponent::GetUpdateClothSimulationData_AnyThread(TMap<int32, FC
 		OutBlendWeight = BlendWeight;
 		OutClothSimulData = LeaderPoseClothComponent->ClothSimulationProxy->GetCurrentSimulationData_AnyThread();
 	}
-	else if (bEnableSimulation && !bBindToLeaderComponent && ClothSimulationProxy)
+	else if (IsSimulationEnabled() && !bBindToLeaderComponent && ClothSimulationProxy)
 	{
 		OutBlendWeight = BlendWeight;
 		OutClothSimulData = ClothSimulationProxy->GetCurrentSimulationData_AnyThread();
diff --git a/Engine/Plugins/ChaosClothAsset/Source/ChaosClothAssetEngine/Private/ChaosClothAsset/ClothSimulationProxy.cpp b/Engine/Plugins/ChaosClothAsset/Source/ChaosClothAssetEngine/Private/ChaosClothAsset/ClothSimulationProxy.cpp
--- a/Engine/Plugins/ChaosClothAsset/Source/ChaosClothAssetEngine/Private/ChaosClothAsset/ClothSimulationProxy.cpp
+++ b/Engine/Plugins/ChaosClothAsset/Source/ChaosClothAssetEngine/Private/ChaosClothAsset/ClothSimulationProxy.cpp
@@ -196,7 +196,15 @@ PRAGMA_ENABLE_DEPRECATION_WARNINGS
 		}
 
 		// It still needs to write back to the GT cache as the context has changed
-		WriteSimulationData();
+		UpdateClothLODs();
+		if (ClothComponent.IsSimulationEnabled())
+		{
+			WriteSimulationData();
+		}
+		else
+		{
+			CurrentSimulationData.Reset();
+		}
 
 		return false;  // Not simulating
 	}
@@ -354,6 +362,39 @@ PRAGMA_ENABLE_DEPRECATION_WARNINGS
 		}
 	}
 
+	void FClothSimulationProxy::UpdateClothLODs()
+	{
+		using namespace ::Chaos;
+
+		bool bAnyLODsChanged = false;
+		for (const TUniquePtr<FClothingSimulationCloth>& Cloth : Cloths)
+		{
+			const int32 AssetIndex = Cloth->GetGroupId();
+
+			if (!Cloth->GetMesh())
+			{
+				continue;  // Invalid or empty cloth
+			}
+
+			// If the LOD has changed while the simulation is suspended, the cloth still needs to be updated with the correct LOD data
+			PRAGMA_DISABLE_DEPRECATION_WARNINGS  // TODO: CHAOS_IS_CLOTHINGSIMULATIONMESH_ABSTRACT
+			const int32 LODIndex = Cloth->GetMesh()->GetLODIndex();
+			PRAGMA_ENABLE_DEPRECATION_WARNINGS
+			if (LODIndex != Cloth->GetLODIndex(Solver.Get()))
+			{
+				if (!ClothComponent.IsSimulationEnabled())
+				{
+					// Mark the cloth as needing to be reset so it doesn't both proxy-deforming lod transitions.
+					Cloth->Reset();
+				}
+				bAnyLODsChanged = true;
+			}
+		}
+		if (bAnyLODsChanged)
+		{
+			Solver->Update(Softs::FSolverReal(0.));  // Update for LOD switching, but do not simulate
+		}
+	}
 
 	void FClothSimulationProxy::WriteSimulationData()
 	{
@@ -363,6 +404,7 @@ PRAGMA_ENABLE_DEPRECATION_WARNINGS
 		TRACE_CPUPROFILER_EVENT_SCOPE(FClothSimulationProxy_WriteSimulationData);
 		SCOPE_CYCLE_COUNTER(STAT_ClothSimulationProxy_WriteSimulationData);
 
+
 		USkinnedMeshComponent* LeaderPoseComponent = nullptr;
 		if (ClothComponent.LeaderPoseComponent.IsValid())
 		{
@@ -406,13 +448,11 @@ PRAGMA_ENABLE_DEPRECATION_WARNINGS
 			}
 
 			// If the LOD has changed while the simulation is suspended, the cloth still needs to be updated with the correct LOD data
-PRAGMA_DISABLE_DEPRECATION_WARNINGS  // TODO: CHAOS_IS_CLOTHINGSIMULATIONMESH_ABSTRACT
+			// This should be handled by calling UpdateClothLODs when not ticking/simulating.
+			PRAGMA_DISABLE_DEPRECATION_WARNINGS  // TODO: CHAOS_IS_CLOTHINGSIMULATIONMESH_ABSTRACT
 			const int32 LODIndex = Cloth->GetMesh()->GetLODIndex();
-PRAGMA_ENABLE_DEPRECATION_WARNINGS
-			if (LODIndex != Cloth->GetLODIndex(Solver.Get()))
-			{
-				Solver->Update(Softs::FSolverReal(0.));  // Update for LOD switching, but do not simulate
-			}
+			PRAGMA_ENABLE_DEPRECATION_WARNINGS
+			ensure(LODIndex == Cloth->GetLODIndex(Solver.Get()));
 
 			if (Cloth->GetParticleRangeId(Solver.Get()) == INDEX_NONE || Cloth->GetLODIndex(Solver.Get()) == INDEX_NONE)
 			{
diff --git a/Engine/Plugins/ChaosClothAsset/Source/ChaosClothAssetEngine/Public/ChaosClothAsset/ClothSimulationProxy.h b/Engine/Plugins/ChaosClothAsset/Source/ChaosClothAssetEngine/Public/ChaosClothAsset/ClothSimulationProxy.h
--- a/Engine/Plugins/ChaosClothAsset/Source/ChaosClothAssetEngine/Public/ChaosClothAsset/ClothSimulationProxy.h
+++ b/Engine/Plugins/ChaosClothAsset/Source/ChaosClothAssetEngine/Public/ChaosClothAsset/ClothSimulationProxy.h
@@ -79,6 +79,7 @@ namespace UE::Chaos::ClothAsset
 
 	private:
 		bool ShouldEnableSolver(bool bSolverCurrentlyEnabled) const;
+		void UpdateClothLODs();
 
 		// Internal physics thread object
 		friend class FClothSimulationProxyParallelTask;
