Chaos Cloth Asset - Fixed Bounds calculations when the cloth asset is a render mesh only.

#tests Tried different cloth asset configurations (with/without sim mesh, with/without component root) and checked its bounds in the level view,

diff --git a/Engine/Plugins/ChaosClothAsset/Source/ChaosClothAssetEngine/Private/ChaosClothAsset/ClothComponent.cpp b/Engine/Plugins/ChaosClothAsset/Source/ChaosClothAssetEngine/Private/ChaosClothAsset/ClothComponent.cpp
--- a/Engine/Plugins/ChaosClothAsset/Source/ChaosClothAssetEngine/Private/ChaosClothAsset/ClothComponent.cpp
+++ b/Engine/Plugins/ChaosClothAsset/Source/ChaosClothAssetEngine/Private/ChaosClothAsset/ClothComponent.cpp
@@ -284,6 +284,8 @@ void UChaosClothComponent::OnPreEndOfFrameSync()
 
 FBoxSphereBounds UChaosClothComponent::CalcBounds(const FTransform& LocalToWorld) const
 {
+	QUICK_SCOPE_CYCLE_COUNTER(STAT_CalcClothComponentBounds);
+
 	FBoxSphereBounds NewBounds(ForceInitToZero);
 
 	// Use cached local bounds if possible
@@ -292,17 +294,58 @@ FBoxSphereBounds UChaosClothComponent::CalcBounds(const FTransform& LocalToWorld
 		NewBounds = bCachedLocalBoundsUpToDate ?
 			CachedWorldOrLocalSpaceBounds.TransformBy(LocalToWorld) :
 			CachedWorldOrLocalSpaceBounds.TransformBy(CachedWorldToLocalTransform * LocalToWorld.ToMatrixWithScale());
+	
+		if (bIncludeComponentLocationIntoBounds)
+		{
+			NewBounds = NewBounds + FBoxSphereBounds(GetComponentLocation(), FVector(1.0f), 1.0f);
+		}
 	}
 	else  // Calculate new bounds
 	{
-		const IConsoleVariable* const CVarCacheLocalSpaceBounds = IConsoleManager::Get().FindConsoleVariable(TEXT("a.CacheLocalSpaceBounds"));
+		FVector RootBoneOffset(ForceInitToZero);
+
+		// If attached to a skeletal mesh component that uses fixed bounds, add the root bone translation
+		if (const USkeletalMeshComponent* const SkeletalMeshComponent = Cast<USkeletalMeshComponent>(LeaderPoseComponent.Get()))
+		{
+			if (SkeletalMeshComponent->GetSkinnedAsset() && SkeletalMeshComponent->bComponentUseFixedSkelBounds)
+			{
+				RootBoneOffset = SkeletalMeshComponent->RootBoneTranslation; // Adjust bounds by root bone translation
+			}
+		}
+
+		static IConsoleVariable* const CVarCacheLocalSpaceBounds = IConsoleManager::Get().FindConsoleVariable(TEXT("a.CacheLocalSpaceBounds"));
 		const bool bCacheLocalSpaceBounds = CVarCacheLocalSpaceBounds ? (CVarCacheLocalSpaceBounds->GetInt() != 0) : true;
 
 		const FTransform CachedBoundsTransform = bCacheLocalSpaceBounds ? FTransform::Identity : LocalToWorld;
 
-		if (ClothSimulationProxy)
+		// Add render mesh bounds
+		constexpr bool bHasValidBodies = false;
+		NewBounds = CalcMeshBound((FVector3f)RootBoneOffset, bHasValidBodies, CachedBoundsTransform);
+
+		if (bIncludeComponentLocationIntoBounds)
+		{
+			const FVector ComponentLocation = GetComponentLocation();
+			const FBoxSphereBounds ComponentLocationBounds(ComponentLocation, FVector(1.), 1.);
+			if (bCacheLocalSpaceBounds)
+			{
+				NewBounds = NewBounds.TransformBy(LocalToWorld);
+				NewBounds = NewBounds + ComponentLocationBounds;
+				NewBounds = NewBounds.TransformBy(LocalToWorld.ToInverseMatrixWithScale());
+			}
+			else
+			{
+				NewBounds = NewBounds + ComponentLocationBounds;
+			}
+		}
+
+		// Add sim mesh bounds
+		if (ClothSimulationProxy.IsValid())
 		{
-			NewBounds = ClothSimulationProxy->CalculateBounds_AnyThread().TransformBy(CachedBoundsTransform);
+			const FBoxSphereBounds SimulationBounds = ClothSimulationProxy->CalculateBounds_AnyThread();
+			if (SimulationBounds.SphereRadius > UE_SMALL_NUMBER)  // Don't add the simulation bounds if there are empty, otherwise it could unwillingly add the component's location
+			{
+				NewBounds = NewBounds + SimulationBounds.TransformBy(CachedBoundsTransform);
+			}
 		}
 
 		CachedWorldOrLocalSpaceBounds = NewBounds;
diff --git a/Engine/Source/Runtime/Engine/Classes/Components/SkeletalMeshComponent.h b/Engine/Source/Runtime/Engine/Classes/Components/SkeletalMeshComponent.h
--- a/Engine/Source/Runtime/Engine/Classes/Components/SkeletalMeshComponent.h
+++ b/Engine/Source/Runtime/Engine/Classes/Components/SkeletalMeshComponent.h
@@ -759,11 +759,6 @@ public:
 	UPROPERTY(transient)
 	uint8 bAnimTreeInitialised:1;
 
-	/** If true, the Location of this Component will be included into its bounds calculation
-	* (this can be useful when using SMU_OnlyTickPoseWhenRendered on a character that moves away from the root and no bones are left near the origin of the component) */
-	UPROPERTY(EditAnywhere, AdvancedDisplay, BlueprintReadOnly, Category = SkeletalMesh)
-	uint8 bIncludeComponentLocationIntoBounds : 1;
-
 	/** If true, line checks will test against the bounding box of this skeletal mesh component and return a hit if there is a collision. */
 	UPROPERTY()
 	uint8 bEnableLineCheckWithBounds:1;
diff --git a/Engine/Source/Runtime/Engine/Classes/Components/SkinnedMeshComponent.h b/Engine/Source/Runtime/Engine/Classes/Components/SkinnedMeshComponent.h
--- a/Engine/Source/Runtime/Engine/Classes/Components/SkinnedMeshComponent.h
+++ b/Engine/Source/Runtime/Engine/Classes/Components/SkinnedMeshComponent.h
@@ -705,6 +705,11 @@ public:
 	uint8 bUseBoundsFromMasterPoseComponent : 1;
 #endif // WITH_EDITORONLY_DATA
 
+	/** If true, the Location of this Component will be included into its bounds calculation
+	* (this can be useful when using SMU_OnlyTickPoseWhenRendered on a character that moves away from the root and no bones are left near the origin of the component) */
+	UPROPERTY(EditAnywhere, AdvancedDisplay, BlueprintReadOnly, Category = SkeletalMesh)
+	uint8 bIncludeComponentLocationIntoBounds : 1;
+
 	/** Forces the mesh to draw in wireframe mode. */
 	UPROPERTY()
 	uint8 bForceWireframe:1;
diff --git a/Engine/Source/Runtime/Engine/Private/Components/SkinnedMeshComponent.cpp b/Engine/Source/Runtime/Engine/Private/Components/SkinnedMeshComponent.cpp
--- a/Engine/Source/Runtime/Engine/Private/Components/SkinnedMeshComponent.cpp
+++ b/Engine/Source/Runtime/Engine/Private/Components/SkinnedMeshComponent.cpp
@@ -1755,7 +1755,9 @@ FBoxSphereBounds USkinnedMeshComponent::CalcBounds(const FTransform& LocalToWorl
 {
 	SCOPE_CYCLE_COUNTER(STAT_CalcSkelMeshBounds);
 
-	return CalcMeshBound( FVector3f::ZeroVector, false, LocalToWorld );
+	return bIncludeComponentLocationIntoBounds ?
+		CalcMeshBound(FVector3f::ZeroVector, false, LocalToWorld) + FBoxSphereBounds(GetComponentLocation(), FVector(1.0f), 1.0f) :
+		CalcMeshBound(FVector3f::ZeroVector, false, LocalToWorld);
 }
 
 void USkinnedMeshComponent::UpdateBounds()
