Chaos Cloth Asset - Added the missing Physics Fields to the cloth asset simulation.

diff --git a/Engine/Plugins/ChaosClothAsset/Source/ChaosClothAssetEngine/Private/ChaosClothAsset/ClothComponent.cpp b/Engine/Plugins/ChaosClothAsset/Source/ChaosClothAssetEngine/Private/ChaosClothAsset/ClothComponent.cpp
--- a/Engine/Plugins/ChaosClothAsset/Source/ChaosClothAssetEngine/Private/ChaosClothAsset/ClothComponent.cpp
+++ b/Engine/Plugins/ChaosClothAsset/Source/ChaosClothAssetEngine/Private/ChaosClothAsset/ClothComponent.cpp
@@ -233,12 +233,6 @@ void UChaosClothComponent::TickComponent(float DeltaTime, enum ELevelTick TickTy
 	// Tick USkinnedMeshComponent first so it will update the predicted lod
 	Super::TickComponent(DeltaTime, TickType, ThisTickFunction);
 
-	// TODO: Fields
-	//if (ClothingSimulation)
-	//{
-	//	ClothingSimulation->UpdateWorldForces(this);
-	//}
-
 	// Make sure that the previous frame simulation has completed
 	HandleExistingParallelSimulation();
 
diff --git a/Engine/Plugins/ChaosClothAsset/Source/ChaosClothAssetEngine/Private/ChaosClothAsset/ClothSimulationProxy.cpp b/Engine/Plugins/ChaosClothAsset/Source/ChaosClothAssetEngine/Private/ChaosClothAsset/ClothSimulationProxy.cpp
--- a/Engine/Plugins/ChaosClothAsset/Source/ChaosClothAssetEngine/Private/ChaosClothAsset/ClothSimulationProxy.cpp
+++ b/Engine/Plugins/ChaosClothAsset/Source/ChaosClothAssetEngine/Private/ChaosClothAsset/ClothSimulationProxy.cpp
@@ -13,6 +13,7 @@
 #include "ChaosCloth/ChaosClothingSimulationSolver.h"
 #include "ChaosCloth/ChaosClothVisualization.h"
 #include "PhysicsEngine/PhysicsSettings.h"
+#include "PhysicsField/PhysicsFieldComponent.h"
 #include "Rendering/SkeletalMeshRenderData.h"
 #include "ClothingSimulation.h"
 
@@ -194,7 +195,26 @@ PRAGMA_ENABLE_DEPRECATION_WARNINGS
 			return false; // Not Simulating
 		}
 
-		Solver->SetEnableSolver(ShouldEnableSolver(Solver->GetEnableSolver()));
+		// Check whether the solver should be enabled for caching purpose
+		const bool bShouldEnableSolver = ShouldEnableSolver(Solver->GetEnableSolver());
+
+		// UpdateWorldForces
+		if (bShouldEnableSolver)
+		{
+			if (UWorld* const World = ClothComponent.GetWorld())
+			{
+				if (UPhysicsFieldComponent* const PhysicsField = World->PhysicsField)
+				{
+					const FBox BoundingBox = CalculateBounds_AnyThread().GetBox().TransformBy(ClothComponent.GetComponentTransform());
+
+					PhysicsField->FillTransientCommands(false, BoundingBox, Solver->GetTime(), Solver->GetPerSolverField().GetTransientCommands());
+					PhysicsField->FillPersistentCommands(false, BoundingBox, Solver->GetTime(), Solver->GetPerSolverField().GetPersistentCommands());
+				}
+			}
+		}
+
+		// Prepare the solver task
+		Solver->SetEnableSolver(bShouldEnableSolver);
 
 		const bool bUseCache = ClothSimulationContext->CacheData.HasData();
 		const bool bCreateParallelTask = (DeltaTime > 0.f && !ClothComponent.IsSimulationSuspended()) || bUseCache;
