Chaos Cloth Asset:
- Enable painting of some predefined cloth collection attributes.
- Fixed the missing RenderDeformerSkinningBlend output SetValue.

#tests Painted the RenderDeformerSkinningBlend attribute and checked that the deformer wasn't active in these areas.
[RN] Chaos Cloth Asset - Enabled the painting of the ProxyDeformer's RenderDeformerSkinningBlend render attribute to allow for tweaking or disabling the proxy deformer in some areas of the render mesh.

diff --git a/Engine/Plugins/ChaosClothAsset/Source/ChaosClothAsset/Private/ChaosClothAsset/ClothCollection.cpp b/Engine/Plugins/ChaosClothAsset/Source/ChaosClothAsset/Private/ChaosClothAsset/ClothCollection.cpp
--- a/Engine/Plugins/ChaosClothAsset/Source/ChaosClothAsset/Private/ChaosClothAsset/ClothCollection.cpp
+++ b/Engine/Plugins/ChaosClothAsset/Source/ChaosClothAsset/Private/ChaosClothAsset/ClothCollection.cpp
@@ -143,6 +143,12 @@ namespace UE::Chaos::ClothAsset::Private
 		{ ClothCollectionGroup::Fabrics, SimFabricGroupAttributes },
 		{ ClothCollectionGroup::Solvers, SolverGroupAttributes },
 	};
+
+	// All paintable attributes for this collection
+	static const TMap<FName, TArray<FName>> UserAccessibleAttributeNamesMap =
+	{
+		{ ClothCollectionGroup::RenderVertices, { ClothCollectionAttribute::RenderDeformerSkinningBlend } },  // Can be painted if required
+	};
 }  // End namespace UE::Chaos::ClothAsset::Private
 
 namespace UE::Chaos::ClothAsset
@@ -708,6 +714,14 @@ namespace UE::Chaos::ClothAsset
 				return ManagedArrayCollection->FindOrAddAttributeTyped<T>(Name, GroupName);
 			}
 		}
+		if (const TArray<FName>* const UserAccessibleAttributeNames = UserAccessibleAttributeNamesMap.Find(GroupName))
+		{
+			// Paintable attribute can only be found, not created (they are part of the schema)
+			if (UserAccessibleAttributeNames->Contains(Name))
+			{
+				return ManagedArrayCollection->FindAttributeTyped<T>(Name, GroupName);
+			}
+		}
 		return nullptr;
 	}
 	template CHAOSCLOTHASSET_API TManagedArray<bool>* FClothCollection::FindOrAddUserDefinedAttribute<bool>(const FName& Name, const FName& GroupName);
@@ -740,6 +754,14 @@ namespace UE::Chaos::ClothAsset
 				return ManagedArrayCollection->FindAttributeTyped<T>(Name, GroupName) != nullptr;
 			}
 		}
+		if (const TArray<FName>* const UserAccessibleAttributeNames = UserAccessibleAttributeNamesMap.Find(GroupName))
+		{
+			// Paintable attribute can only be found, not created (they are part of the schema)
+			if (UserAccessibleAttributeNames->Contains(Name))
+			{
+				return ManagedArrayCollection->FindAttributeTyped<T>(Name, GroupName) != nullptr;
+			}
+		}
 		return false;
 	}
 	template CHAOSCLOTHASSET_API bool FClothCollection::HasUserDefinedAttribute<bool>(const FName& Name, const FName& GroupName) const;
@@ -759,6 +781,14 @@ namespace UE::Chaos::ClothAsset
 				return ManagedArrayCollection->FindAttributeTyped<T>(Name, GroupName);
 			}
 		}
+		if (const TArray<FName>* const UserAccessibleAttributeNames = UserAccessibleAttributeNamesMap.Find(GroupName))
+		{
+			// Paintable attribute can only be found, not created (they are part of the schema)
+			if (UserAccessibleAttributeNames->Contains(Name))
+			{
+				return ManagedArrayCollection->FindAttributeTyped<T>(Name, GroupName);
+			}
+		}
 		return nullptr;
 	}
 	template CHAOSCLOTHASSET_API const TManagedArray<bool>* FClothCollection::GetUserDefinedAttribute<bool>(const FName& Name, const FName& GroupName) const;
@@ -778,6 +808,14 @@ namespace UE::Chaos::ClothAsset
 				return ManagedArrayCollection->FindAttributeTyped<T>(Name, GroupName);
 			}
 		}
+		if (const TArray<FName>* const UserAccessibleAttributeNames = UserAccessibleAttributeNamesMap.Find(GroupName))
+		{
+			// Paintable attribute can only be found, not created (they are part of the schema)
+			if (UserAccessibleAttributeNames->Contains(Name))
+			{
+				return ManagedArrayCollection->FindAttributeTyped<T>(Name, GroupName);
+			}
+		}
 		return nullptr;
 	}
 	template CHAOSCLOTHASSET_API TManagedArray<bool>* FClothCollection::GetUserDefinedAttribute<bool>(const FName& Name, const FName& GroupName);
@@ -802,6 +840,14 @@ namespace UE::Chaos::ClothAsset
 				return true;
 			}
 		}
+		if (const TArray<FName>* const UserAccessibleAttributeNames = UserAccessibleAttributeNamesMap.Find(GroupName))
+		{
+			// Paintable attribute can only be found, not created (they are part of the schema)
+			if (UserAccessibleAttributeNames->Contains(Name) && ManagedArrayCollection->HasAttribute(Name, GroupName))
+			{
+				return true;
+			}
+		}
 		return false;
 	}
 } // End namespace UE::Chaos::ClothAsset
diff --git a/Engine/Plugins/ChaosClothAsset/Source/ChaosClothAsset/Private/ChaosClothAsset/ClothCollection.h b/Engine/Plugins/ChaosClothAsset/Source/ChaosClothAsset/Private/ChaosClothAsset/ClothCollection.h
--- a/Engine/Plugins/ChaosClothAsset/Source/ChaosClothAsset/Private/ChaosClothAsset/ClothCollection.h
+++ b/Engine/Plugins/ChaosClothAsset/Source/ChaosClothAsset/Private/ChaosClothAsset/ClothCollection.h
@@ -152,7 +152,7 @@ namespace UE::Chaos::ClothAsset
 		TManagedArray<T>* GetUserDefinedAttribute(const FName& Name, const FName& GroupName);
 
 		static bool IsValidClothCollectionGroupName(const FName& GroupName);
-		static bool IsValidUserDefinedAttributeName(const FName& Name, const FName& GroupName);
+		bool IsValidUserDefinedAttributeName(const FName& Name, const FName& GroupName);
 
 		//~ LODs Group (There should be only one LOD per ClothCollection)
 		const TManagedArray<FString>* GetPhysicsAssetPathName() const { return PhysicsAssetPathName; }
diff --git a/Engine/Plugins/ChaosClothAssetEditor/Source/ChaosClothAssetDataflowNodes/Private/ChaosClothAsset/ProxyDeformerNode.cpp b/Engine/Plugins/ChaosClothAssetEditor/Source/ChaosClothAssetDataflowNodes/Private/ChaosClothAsset/ProxyDeformerNode.cpp
--- a/Engine/Plugins/ChaosClothAssetEditor/Source/ChaosClothAssetDataflowNodes/Private/ChaosClothAsset/ProxyDeformerNode.cpp
+++ b/Engine/Plugins/ChaosClothAssetEditor/Source/ChaosClothAssetDataflowNodes/Private/ChaosClothAsset/ProxyDeformerNode.cpp
@@ -363,6 +363,10 @@ void FChaosClothAssetProxyDeformerNode::Evaluate(Dataflow::FContext& Context, co
 
 		SetValue(Context, MoveTemp(*ClothCollection), &Collection);
 	}
+	else if (Out->IsA<FString>(&SkinningBlendName))
+	{
+		SetValue(Context, SkinningBlendName, &SkinningBlendName);
+	}
 }
 
 Dataflow::FPin FChaosClothAssetProxyDeformerNode::AddPin()
