Add ChaosClothAsset to Cloth Total cycle stat

diff --git a/Engine/Plugins/ChaosClothAsset/Source/ChaosClothAssetEngine/Private/ChaosClothAsset/ClothSimulationProxy.cpp b/Engine/Plugins/ChaosClothAsset/Source/ChaosClothAssetEngine/Private/ChaosClothAsset/ClothSimulationProxy.cpp
--- a/Engine/Plugins/ChaosClothAsset/Source/ChaosClothAssetEngine/Private/ChaosClothAsset/ClothSimulationProxy.cpp
+++ b/Engine/Plugins/ChaosClothAsset/Source/ChaosClothAssetEngine/Private/ChaosClothAsset/ClothSimulationProxy.cpp
@@ -12,6 +12,7 @@
 #include "ChaosCloth/ChaosClothingSimulationConfig.h"
 #include "ChaosCloth/ChaosClothingSimulationSolver.h"
 #include "ChaosCloth/ChaosClothVisualization.h"
+#include "Engine/SkinnedAssetCommon.h"
 #include "PhysicsEngine/PhysicsSettings.h"
 #include "PhysicsField/PhysicsFieldComponent.h"
 #include "Rendering/SkeletalMeshRenderData.h"
@@ -81,6 +82,7 @@ namespace UE::Chaos::ClothAsset
 		void DoTask(ENamedThreads::Type CurrentThread, const FGraphEventRef& MyCompletionGraphEvent)
 		{
 			FScopeCycleCounterUObject ContextScope(ClothSimulationProxy.ClothComponent.GetSkinnedAsset());
+			SCOPE_CYCLE_COUNTER(STAT_ClothTotalTime);
 			CSV_SCOPED_TIMING_STAT(Animation, Cloth);
 
 			ClothSimulationProxy.Tick();
@@ -436,7 +438,7 @@ PRAGMA_ENABLE_DEPRECATION_WARNINGS
 		CSV_SCOPED_TIMING_STAT(Animation, Cloth);
 		TRACE_CPUPROFILER_EVENT_SCOPE(FClothSimulationProxy_WriteSimulationData);
 		SCOPE_CYCLE_COUNTER(STAT_ClothSimulationProxy_WriteSimulationData);
-
+		SCOPE_CYCLE_COUNTER(STAT_ClothWriteback)
 
 		USkinnedMeshComponent* LeaderPoseComponent = nullptr;
 		if (ClothComponent.LeaderPoseComponent.IsValid())
diff --git a/Engine/Source/Runtime/Engine/Classes/Engine/SkinnedAssetCommon.h b/Engine/Source/Runtime/Engine/Classes/Engine/SkinnedAssetCommon.h
--- a/Engine/Source/Runtime/Engine/Classes/Engine/SkinnedAssetCommon.h
+++ b/Engine/Source/Runtime/Engine/Classes/Engine/SkinnedAssetCommon.h
@@ -24,6 +24,10 @@ struct FSkelMeshSection;
 struct FSkeletalMeshLODGroupSettings;
 
 
+//This is the total cloth time split up among multiple computation (updating gpu, updating sim, etc...)
+DECLARE_CYCLE_STAT(TEXT("Cloth Total"), STAT_ClothTotalTime, STATGROUP_Physics);
+DECLARE_CYCLE_STAT(TEXT("Cloth Writeback"), STAT_ClothWriteback, STATGROUP_Physics);
+
 UENUM()
 enum class ESkinCacheUsage : uint8
 {
diff --git a/Engine/Source/Runtime/Engine/Private/SkeletalMeshComponentPhysics.cpp b/Engine/Source/Runtime/Engine/Private/SkeletalMeshComponentPhysics.cpp
--- a/Engine/Source/Runtime/Engine/Private/SkeletalMeshComponentPhysics.cpp
+++ b/Engine/Source/Runtime/Engine/Private/SkeletalMeshComponentPhysics.cpp
@@ -63,10 +63,6 @@ static TAutoConsoleVariable<float> CVarClothTeleportRotationThreshold(TEXT("p.Cl
 
 static TAutoConsoleVariable<int32> CVarEnableKinematicDeferralPrePhysicsCondition(TEXT("p.EnableKinematicDeferralPrePhysicsCondition"), 1, TEXT("If is 1, and deferral would've been disallowed due to EUpdateTransformFlags, allow if in PrePhysics tick. If 0, condition is unchanged."));
 
-//This is the total cloth time split up among multiple computation (updating gpu, updating sim, etc...)
-DECLARE_CYCLE_STAT(TEXT("Cloth Total"), STAT_ClothTotalTime, STATGROUP_Physics);
-DECLARE_CYCLE_STAT(TEXT("Cloth Writeback"), STAT_ClothWriteback, STATGROUP_Physics);
-
 // Used as a default return value for invalid cloth data access
 static const TMap<int32, FClothSimulData> SEmptyClothSimulationData;
 
