Chaos Cloth Asset - Fixed issue where the LODModel RequiredBones list contains duplicate indices.

[RN] Chaos Cloth Asset - Fixed a crash in the leader component code where differences in sim and render mesh skinning would generate lists of duplicated bones.

diff --git a/Engine/Plugins/ChaosClothAsset/Source/ChaosClothAssetEngine/Private/ChaosClothAsset/ClothAssetBuilder.cpp b/Engine/Plugins/ChaosClothAsset/Source/ChaosClothAssetEngine/Private/ChaosClothAsset/ClothAssetBuilder.cpp
--- a/Engine/Plugins/ChaosClothAsset/Source/ChaosClothAssetEngine/Private/ChaosClothAsset/ClothAssetBuilder.cpp
+++ b/Engine/Plugins/ChaosClothAsset/Source/ChaosClothAssetEngine/Private/ChaosClothAsset/ClothAssetBuilder.cpp
@@ -489,7 +489,9 @@ PRAGMA_ENABLE_DEPRECATION_WARNINGS
 	{
 		if (ensure(ClothSimulationModel->IsValidLodIndex(LodIndex)))
 		{
-			LODModel.RequiredBones.Append(ClothSimulationModel->ClothSimulationLodModels[LodIndex].RequiredExtraBoneIndices);
+			ActiveBoneIndices.Append(ClothSimulationModel->ClothSimulationLodModels[LodIndex].RequiredExtraBoneIndices);
+
+			LODModel.RequiredBones = ActiveBoneIndices.Array();
 
 			ClothAsset.GetRefSkeleton().EnsureParentsExistAndSort(LODModel.RequiredBones);
 			LODModel.RequiredBones.Shrink();
