Chaos Cloth Asset - Added support for external collision sources.

#tests Tested with cape v3 prototype, the collisions are showing where they should be but the new solver needs fixing.
[RN] Chaos Cloth Asset - Added external collision support on the ChaosClothComponent with new Add/RemoveCollisionSource blueprint functions.

diff --git a/Engine/Plugins/ChaosCloth/Source/ChaosCloth/Private/ChaosCloth/ChaosClothingSimulationCollider.cpp b/Engine/Plugins/ChaosCloth/Source/ChaosCloth/Private/ChaosCloth/ChaosClothingSimulationCollider.cpp
--- a/Engine/Plugins/ChaosCloth/Source/ChaosCloth/Private/ChaosCloth/ChaosClothingSimulationCollider.cpp
+++ b/Engine/Plugins/ChaosCloth/Source/ChaosCloth/Private/ChaosCloth/ChaosClothingSimulationCollider.cpp
@@ -496,7 +496,14 @@ FClothCollisionData FClothingSimulationCollider::GetCollisionData(const FClothin
 	return ClothCollisionData;
 }
 
-void FClothingSimulationCollider::ExtractPhysicsAssetCollision(FClothCollisionData& ClothCollisionData, TArray<FLevelSetCollisionData>& LevelSetCollisions, TArray<FSkinnedLevelSetCollisionData>& SkinnedLevelSetCollisions, TArray<int32>& UsedBoneIndices)
+void FClothingSimulationCollider::ExtractPhysicsAssetCollision(
+	const UPhysicsAsset* PhysicsAsset,
+	const FReferenceSkeleton* ReferenceSkeleton,
+	FClothCollisionData& ClothCollisionData,
+	TArray<FLevelSetCollisionData>& LevelSetCollisions,
+	TArray<FSkinnedLevelSetCollisionData>& SkinnedLevelSetCollisions,
+	TArray<int32>& UsedBoneIndices,
+	bool bUseSphylsOnly)
 {
 	ClothCollisionData.Reset();
 	UsedBoneIndices.Reset();
@@ -564,18 +571,6 @@ void FClothingSimulationCollider::ExtractPhysicsAssetCollision(FClothCollisionDa
 				ClothCollisionData.Spheres.Add(Sphere);
 			}
 
-			// Add boxes
-			for (const FKBoxElem& BoxElem : AggGeom.BoxElems)
-			{
-				// Add extracted box collision data
-				FClothCollisionPrim_Box Box;
-				Box.LocalPosition = BoxElem.Center;
-				Box.LocalRotation = BoxElem.Rotation.Quaternion();
-				Box.HalfExtents = FVector(BoxElem.X, BoxElem.Y, BoxElem.Z) * 0.5f;
-				Box.BoneIndex = MappedBoneIndex;
-				ClothCollisionData.Boxes.Add(Box);
-			}
-
 			// Add tapered capsules
 			for (const FKTaperedCapsuleElem& TaperedCapsuleElem : AggGeom.TaperedCapsuleElems)
 			{
@@ -610,6 +605,23 @@ void FClothingSimulationCollider::ExtractPhysicsAssetCollision(FClothCollisionDa
 				}
 			}
 
+			if (bUseSphylsOnly)
+			{
+				continue;  // Exit now if only spheres and capsules are required
+			}
+
+			// Add boxes
+			for (const FKBoxElem& BoxElem : AggGeom.BoxElems)
+			{
+				// Add extracted box collision data
+				FClothCollisionPrim_Box Box;
+				Box.LocalPosition = BoxElem.Center;
+				Box.LocalRotation = BoxElem.Rotation.Quaternion();
+				Box.HalfExtents = FVector(BoxElem.X, BoxElem.Y, BoxElem.Z) * 0.5f;
+				Box.BoneIndex = MappedBoneIndex;
+				ClothCollisionData.Boxes.Add(Box);
+			}
+
 			// Add convexes
 			for (const FKConvexElem& ConvexElem : AggGeom.ConvexElems)
 			{
@@ -750,7 +762,7 @@ PRAGMA_ENABLE_DEPRECATION_WARNINGS
 	TArray<FLevelSetCollisionData> LevelSetCollisions;
 	TArray<FSkinnedLevelSetCollisionData> SkinnedLevelSetCollisions;
 	TArray<int32> UsedBoneIndices;
-	ExtractPhysicsAssetCollision(PhysicsAssetCollisionData, LevelSetCollisions, SkinnedLevelSetCollisions, UsedBoneIndices);
+	ExtractPhysicsAssetCollision(PhysicsAsset, ReferenceSkeleton, PhysicsAssetCollisionData, LevelSetCollisions, SkinnedLevelSetCollisions, UsedBoneIndices);
 
 	LODData[(int32)ECollisionDataType::LODless]->Add(Solver, Cloth, PhysicsAssetCollisionData, LevelSetCollisions, SkinnedLevelSetCollisions, Scale, UsedBoneIndices);
 }
diff --git a/Engine/Plugins/ChaosCloth/Source/ChaosCloth/Public/ChaosCloth/ChaosClothingSimulationCollider.h b/Engine/Plugins/ChaosCloth/Source/ChaosCloth/Public/ChaosCloth/ChaosClothingSimulationCollider.h
--- a/Engine/Plugins/ChaosCloth/Source/ChaosCloth/Public/ChaosCloth/ChaosClothingSimulationCollider.h
+++ b/Engine/Plugins/ChaosCloth/Source/ChaosCloth/Public/ChaosCloth/ChaosClothingSimulationCollider.h
@@ -10,6 +10,8 @@ class USkeletalMeshComponent;
 class UClothingAssetCommon;
 class UPhysicsAsset;
 class FClothingSimulationContextCommon;
+class FLevelSetCollisionData;
+class FSkinnedLevelSetCollisionData;
 struct FReferenceSkeleton;
 
 namespace Chaos
@@ -29,6 +31,20 @@ namespace Chaos
 			LODs,  // LODIndex based start slot for LODs collisions
 		};
 
+		struct FLevelSetCollisionData
+		{
+			const TSharedPtr<Chaos::FLevelSet, ESPMode::ThreadSafe> LevelSet;
+			FTransform Transform;
+			int32 BoneIndex;
+		};
+
+		struct FSkinnedLevelSetCollisionData
+		{
+			const TRefCountPtr<Chaos::TWeightedLatticeImplicitObject<Chaos::FLevelSet>> WeightedLevelSet;
+			int32 BoneIndex;
+			TArray<int32> MappedSkinnedBones;
+		};
+
 		CHAOSCLOTH_API FClothingSimulationCollider(const UPhysicsAsset* InPhysicsAsset, const FReferenceSkeleton* InReferenceSkeleton);
 
 		UE_DEPRECATED(5.2, "Use FClothingSimulationCollider(const UPhysicsAsset* InPhysicsAsset, const FReferenceSkeleton* InReferenceSkeleton) instead.")
@@ -84,23 +100,17 @@ namespace Chaos
 		CHAOSCLOTH_API TConstArrayView<bool> GetCollisionStatus(const FClothingSimulationSolver* Solver, const FClothingSimulationCloth* Cloth, ECollisionDataType CollisionDataType) const;
 		// ---- End of the debugging and visualization functions ----
 
-	private:
-		struct FLevelSetCollisionData
-		{
-			const TSharedPtr<Chaos::FLevelSet, ESPMode::ThreadSafe> LevelSet;
-			FTransform Transform;
-			int32 BoneIndex;
-		};
-
-		struct FSkinnedLevelSetCollisionData
-		{
-			const TRefCountPtr<Chaos::TWeightedLatticeImplicitObject<Chaos::FLevelSet>> WeightedLevelSet;
-			int32 BoneIndex;
-			TArray<int32> MappedSkinnedBones;
-		};
-
-		CHAOSCLOTH_API void ExtractPhysicsAssetCollision(FClothCollisionData& ClothCollisionData, TArray<FLevelSetCollisionData>& LevelSetCollisions, TArray<FSkinnedLevelSetCollisionData>& SkinnedLevelSetCollisions, TArray<int32>& UsedBoneIndices);
+		// Extract all collisions from the physics asset into a more cloth friendly format
+		CHAOSCLOTH_API static void ExtractPhysicsAssetCollision(
+			const UPhysicsAsset* PhysicsAsset,
+			const FReferenceSkeleton* ReferenceSkeleton,
+			FClothCollisionData& ClothCollisionData,
+			TArray<FLevelSetCollisionData>& LevelSetCollisions,
+			TArray<FSkinnedLevelSetCollisionData>& SkinnedLevelSetCollisions,
+			TArray<int32>& UsedBoneIndices,
+			bool bUseSphylOnly = false);
 
+	private:
 		CHAOSCLOTH_API int32 GetNumGeometries(int32 InSlotIndex) const;
 
 		// Return the collision particle range ID for the specified slot being LODLess, external, or any of the LODs collision.
diff --git a/Engine/Plugins/ChaosClothAsset/Source/ChaosClothAssetEngine/Private/ChaosClothAsset/ClothComponent.cpp b/Engine/Plugins/ChaosClothAsset/Source/ChaosClothAssetEngine/Private/ChaosClothAsset/ClothComponent.cpp
--- a/Engine/Plugins/ChaosClothAsset/Source/ChaosClothAssetEngine/Private/ChaosClothAsset/ClothComponent.cpp
+++ b/Engine/Plugins/ChaosClothAsset/Source/ChaosClothAssetEngine/Private/ChaosClothAsset/ClothComponent.cpp
@@ -4,6 +4,7 @@
 #include "ChaosClothAsset/ClothAsset.h"
 #include "ChaosClothAsset/ClothAssetInteractor.h"
 #include "ChaosClothAsset/ClothAssetPrivate.h"
+#include "ChaosClothAsset/CollisionSources.h"
 #include "ChaosClothAsset/ClothSimulationModel.h"
 #include "ChaosClothAsset/ClothSimulationProxy.h"
 #include "Chaos/CollectionPropertyFacade.h"
@@ -27,6 +28,7 @@ UChaosClothComponent::UChaosClothComponent(const FObjectInitializer& ObjectIniti
 	, bEnableSimulation(1)
 	, bSuspendSimulation(0)
 	, bBindToLeaderComponent(0)
+	, CollisionSources(MakeUnique<UE::Chaos::ClothAsset::FCollisionSources>(this))
 {
 	PrimaryComponentTick.EndTickGroup = TG_PostPhysics;
 }
@@ -483,7 +485,7 @@ void UChaosClothComponent::GetAdditionalRequiredBonesForLeader(int32 LeaderLODIn
 
 	if (RequiredBones.Num())
 	{
-		// Make sure all of these are in RequiredBones.
+		// Make sure all of these are in RequiredBones, note MergeInBoneIndexArrays requires the arrays to be sorted and bone must be unique
 		MergeInBoneIndexArrays(InOutRequiredBones, RequiredBones);
 	}
 }
@@ -511,6 +513,26 @@ TSharedPtr<UE::Chaos::ClothAsset::FClothSimulationProxy> UChaosClothComponent::C
 	return MakeShared<FClothSimulationProxy>(*this);
 }
 
+void UChaosClothComponent::AddCollisionSource(USkinnedMeshComponent* SourceComponent, const UPhysicsAsset* SourcePhysicsAsset, bool bUseSphylsOnly)
+{
+	CollisionSources->Add(SourceComponent, SourcePhysicsAsset, bUseSphylsOnly);
+}
+
+void UChaosClothComponent::RemoveCollisionSources(const USkinnedMeshComponent* SourceComponent)
+{
+	CollisionSources->Remove(SourceComponent);
+}
+
+void UChaosClothComponent::RemoveCollisionSource(const USkinnedMeshComponent* SourceComponent, const UPhysicsAsset* SourcePhysicsAsset)
+{
+	CollisionSources->Remove(SourceComponent, SourcePhysicsAsset);
+}
+
+void UChaosClothComponent::ResetCollisionSources()
+{
+	CollisionSources->Reset();
+}
+
 void UChaosClothComponent::StartNewParallelSimulation(float DeltaTime)
 {
 	if (ClothSimulationProxy.IsValid())
diff --git a/Engine/Plugins/ChaosClothAsset/Source/ChaosClothAssetEngine/Private/ChaosClothAsset/ClothSimulationProxy.cpp b/Engine/Plugins/ChaosClothAsset/Source/ChaosClothAssetEngine/Private/ChaosClothAsset/ClothSimulationProxy.cpp
--- a/Engine/Plugins/ChaosClothAsset/Source/ChaosClothAssetEngine/Private/ChaosClothAsset/ClothSimulationProxy.cpp
+++ b/Engine/Plugins/ChaosClothAsset/Source/ChaosClothAssetEngine/Private/ChaosClothAsset/ClothSimulationProxy.cpp
@@ -7,6 +7,7 @@
 #include "ChaosClothAsset/ClothSimulationContext.h"
 #include "ChaosClothAsset/ClothSimulationMesh.h"
 #include "ChaosClothAsset/ClothSimulationModel.h"
+#include "ChaosClothAsset/CollisionSources.h"
 #include "ChaosCloth/ChaosClothingSimulationCloth.h"
 #include "ChaosCloth/ChaosClothingSimulationCollider.h"
 #include "ChaosCloth/ChaosClothingSimulationConfig.h"
@@ -95,6 +96,7 @@ namespace UE::Chaos::ClothAsset
 	FClothSimulationProxy::FClothSimulationProxy(const UChaosClothComponent& InClothComponent)
 		: ClothComponent(InClothComponent)
 		, ClothSimulationContext(MakeUnique<FClothSimulationContext>())
+		, CollisionSourcesProxy(MakeUnique<FCollisionSourcesProxy>(InClothComponent.GetCollisionSources()))
 		, Solver(nullptr)
 		, Visualization(nullptr)
 		, MaxDeltaTime(UPhysicsSettings::Get()->MaxPhysicsDeltaTime)
@@ -141,10 +143,9 @@ namespace UE::Chaos::ClothAsset
 		// Create collider simulation thread object
 PRAGMA_DISABLE_DEPRECATION_WARNINGS  // TODO: CHAOS_IS_CLOTHINGSIMULATIONMESH_ABSTRACT
 		const FReferenceSkeleton* const ReferenceSkeleton = &ClothAsset->GetRefSkeleton();
-		const int32 ColliderIndex = ClothComponent.GetPhysicsAsset() ? Colliders.Emplace(MakeUnique<FClothingSimulationCollider>(ClothComponent.GetPhysicsAsset(), ReferenceSkeleton)) : INDEX_NONE;
+		const int32 ColliderIndex = Colliders.Emplace(MakeUnique<FClothingSimulationCollider>(ClothComponent.GetPhysicsAsset(), ReferenceSkeleton));
 PRAGMA_ENABLE_DEPRECATION_WARNINGS
-
-		//Colliders[ColliderIndex]->SetCollisionData(&ExternalCollisionData);  // TODO: External collision data
+		Colliders[ColliderIndex]->SetCollisionData(&CollisionSourcesProxy->GetCollisionData());
 
 		// Create cloth config simulation thread object
 		const int32 ClothConfigIndex = Configs.Emplace(MakeUnique<FClothingSimulationConfig>(ClothComponent.GetPropertyCollections()));
@@ -199,30 +200,33 @@ PRAGMA_ENABLE_DEPRECATION_WARNINGS
 
 		// Check whether the solver should be enabled for caching purpose
 		const bool bShouldEnableSolver = ShouldEnableSolver(Solver->GetEnableSolver());
+		Solver->SetEnableSolver(bShouldEnableSolver);
 
-		// UpdateWorldForces
-		if (bShouldEnableSolver)
+		// Prepare the solver task
+		const bool bUseCache = ClothSimulationContext->CacheData.HasData();
+		const bool bCreateParallelTask = (DeltaTime > 0.f && !ClothComponent.IsSimulationSuspended()) || bUseCache;
+		if (bCreateParallelTask)
 		{
-			if (UWorld* const World = ClothComponent.GetWorld())
+			if (bShouldEnableSolver)
 			{
-				if (UPhysicsFieldComponent* const PhysicsField = World->PhysicsField)
+				// Update the config properties
+				InitializeConfigs();
+
+				// Update world forces
+				if (UWorld* const World = ClothComponent.GetWorld())
 				{
-					const FBox BoundingBox = CalculateBounds_AnyThread().GetBox().TransformBy(ClothComponent.GetComponentTransform());
+					if (UPhysicsFieldComponent* const PhysicsField = World->PhysicsField)
+					{
+						const FBox BoundingBox = CalculateBounds_AnyThread().GetBox().TransformBy(ClothComponent.GetComponentTransform());
 
-					PhysicsField->FillTransientCommands(false, BoundingBox, Solver->GetTime(), Solver->GetPerSolverField().GetTransientCommands());
-					PhysicsField->FillPersistentCommands(false, BoundingBox, Solver->GetTime(), Solver->GetPerSolverField().GetPersistentCommands());
+						PhysicsField->FillTransientCommands(false, BoundingBox, Solver->GetTime(), Solver->GetPerSolverField().GetTransientCommands());
+						PhysicsField->FillPersistentCommands(false, BoundingBox, Solver->GetTime(), Solver->GetPerSolverField().GetPersistentCommands());
+					}
 				}
-			}
-		}
 
-		// Prepare the solver task
-		Solver->SetEnableSolver(bShouldEnableSolver);
-
-		const bool bUseCache = ClothSimulationContext->CacheData.HasData();
-		const bool bCreateParallelTask = (DeltaTime > 0.f && !ClothComponent.IsSimulationSuspended()) || bUseCache;
-		if (bCreateParallelTask)
-		{
-			InitializeConfigs();
+				// Update external collision sources
+				CollisionSourcesProxy->ExtractCollisionData();
+			}
 
 			// Start the the cloth simulation thread
 			ParallelTask = TGraphTask<FClothSimulationProxyParallelTask>::CreateTask(nullptr, ENamedThreads::GameThread).ConstructAndDispatchWhenReady(*this);
diff --git a/Engine/Plugins/ChaosClothAsset/Source/ChaosClothAssetEngine/Private/ChaosClothAsset/CollisionSources.cpp b/Engine/Plugins/ChaosClothAsset/Source/ChaosClothAssetEngine/Private/ChaosClothAsset/CollisionSources.cpp
new file mode 100644
--- /dev/null
+++ b/Engine/Plugins/ChaosClothAsset/Source/ChaosClothAssetEngine/Private/ChaosClothAsset/CollisionSources.cpp
@@ -0,0 +1,196 @@
+// Copyright Epic Games, Inc. All Rights Reserved.
+#include "ChaosClothAsset/CollisionSources.h"
+
+#include "Chaos/Levelset.h"
+#include "Chaos/WeightedLatticeImplicitObject.h"
+#include "ChaosCloth/ChaosClothingSimulationCollider.h"
+#include "Components/SkinnedMeshComponent.h"
+#include "Engine/SkinnedAsset.h"
+#include "PhysicsEngine/PhysicsAsset.h"
+
+namespace UE::Chaos::ClothAsset
+{
+	void FCollisionSources::Add(USkinnedMeshComponent* SourceComponent, const UPhysicsAsset* SourcePhysicsAsset, bool bUseSphylsOnly)
+	{
+		using namespace UE::Chaos::ClothAsset;
+
+		if (OwnerComponent.IsValid() && SourceComponent && SourcePhysicsAsset)
+		{
+			FCollisionSource* const FoundCollisionSource = CollisionSources.FindByPredicate(
+				[SourceComponent, SourcePhysicsAsset](const FCollisionSource& CollisionSource)
+				{
+					return CollisionSource.SourceComponent == SourceComponent && CollisionSource.SourcePhysicsAsset == SourcePhysicsAsset;
+				}
+			);
+
+			if (!FoundCollisionSource)
+			{
+				// Delegate called after the transform buffer flip, so that the cloths' transform gets updated when the component owning the cloth isn't moving, but the collision source is
+				const FSimpleDelegate OnBoneTransformsFinalizedDelegate = FSimpleDelegate::CreateLambda(
+					[this]()
+					{
+						++Version;
+					});
+
+				// Add the new collision source
+				CollisionSources.Emplace(SourceComponent, SourcePhysicsAsset, OnBoneTransformsFinalizedDelegate, bUseSphylsOnly);
+
+				// Add prerequisite so we don't end up with a frame delay
+				OwnerComponent->PrimaryComponentTick.AddPrerequisite(SourceComponent, SourceComponent->PrimaryComponentTick);
+
+				// Mark the collision sources as changed
+				++Version;
+			}
+		}
+	}
+
+	void FCollisionSources::Remove(const USkinnedMeshComponent* SourceComponent)
+	{
+		using namespace UE::Chaos::ClothAsset;
+
+		if (SourceComponent)
+		{
+			// Note: Stale prerequises are removed when in the QueueTickFunction once the source object has been destroyed.
+			const int32 NumRemoved = CollisionSources.RemoveAll([SourceComponent](const FCollisionSource& CollisionSource)
+				{
+					return !CollisionSource.SourceComponent.IsValid() || CollisionSource.SourceComponent == SourceComponent;
+				});
+		
+			// Mark the collision sources as changed
+			if (NumRemoved > 0)
+			{
+				++Version;
+			}
+		}
+	}
+
+	void FCollisionSources::Remove(const USkinnedMeshComponent* SourceComponent, const UPhysicsAsset* SourcePhysicsAsset)
+	{
+		using namespace UE::Chaos::ClothAsset;
+
+		if (SourceComponent)
+		{
+			// Note: Stale prerequises are removed when in the QueueTickFunction once the source object has been destroyed.
+			const int32 NumRemoved = CollisionSources.RemoveAll([SourceComponent, SourcePhysicsAsset](const FCollisionSource& CollisionSource)
+				{
+					return !CollisionSource.SourceComponent.IsValid() ||
+						(CollisionSource.SourceComponent == SourceComponent && CollisionSource.SourcePhysicsAsset == SourcePhysicsAsset);
+				});
+
+			// Mark the collision sources as changed
+			if (NumRemoved > 0)
+			{
+				++Version;
+			}
+		}
+	}
+
+	void FCollisionSources::Reset()
+	{
+		CollisionSources.Reset();
+		++Version;
+	}
+
+	void FCollisionSources::ExtractCollisionData(FClothCollisionData& CollisionData)
+	{
+		CollisionData.Reset();
+		if (OwnerComponent.IsValid())
+		{
+			for (FCollisionSource& CollisionSource : CollisionSources)
+			{
+				CollisionSource.ExtractCollisionData(*OwnerComponent, CollisionData);
+			}
+		}
+	}
+
+	void FCollisionSourcesProxy::ExtractCollisionData()
+	{
+		if (Version != CollisionSources.Version)
+		{
+			CollisionSources.ExtractCollisionData(CollisionData);
+			Version = CollisionSources.Version;
+		}
+	}
+
+	FCollisionSources::FCollisionSource::FCollisionSource(
+		USkinnedMeshComponent* InSourceComponent,
+		const UPhysicsAsset* InSourcePhysicsAsset,
+		const FSimpleDelegate& InOnBoneTransformsFinalizedDelegate,
+		bool bInUseSphylsOnly)
+		: SourceComponent(InSourceComponent)
+		, SourcePhysicsAsset(InSourcePhysicsAsset)
+		, bUseSphylsOnly(bInUseSphylsOnly)
+	{
+		if (InSourceComponent)
+		{
+			OnBoneTransformsFinalizedHandle = InSourceComponent->RegisterOnBoneTransformsFinalizedDelegate(InOnBoneTransformsFinalizedDelegate);
+		}
+	}
+
+	FCollisionSources::FCollisionSource::~FCollisionSource()
+	{
+		if (USkinnedMeshComponent* const SourceComponentPtr = SourceComponent.Get())
+		{
+			check(OnBoneTransformsFinalizedHandle.IsValid());
+			SourceComponentPtr->UnregisterOnBoneTransformsFinalizedDelegate(OnBoneTransformsFinalizedHandle);
+		}
+	}
+
+	void FCollisionSources::FCollisionSource::ExtractCollisionData(const USkinnedMeshComponent& InOwnerComponent, FClothCollisionData& CollisionData)
+	{
+		if (USkinnedMeshComponent* const SourceComponentPtr = SourceComponent.Get())
+		{
+			const USkinnedAsset* SkinnedAsset = SourceComponentPtr->GetSkinnedAsset();
+
+			// Extract the collision data if not already cached
+			if (CachedSkinnedAsset != SkinnedAsset)
+			{
+				CachedSkinnedAsset = SkinnedAsset;
+				CachedCollisionData.Reset();
+				CachedUsedBoneIndices.Reset();
+
+				const UPhysicsAsset* const PhysicsAsset = SourcePhysicsAsset.Get();
+
+				if (SkinnedAsset && PhysicsAsset)
+				{
+					// Extract collisions
+					TArray<::Chaos::FClothingSimulationCollider::FLevelSetCollisionData> LevelSetCollisions;
+					TArray<::Chaos::FClothingSimulationCollider::FSkinnedLevelSetCollisionData> SkinnedLevelSetCollisions;
+
+					::Chaos::FClothingSimulationCollider::ExtractPhysicsAssetCollision(
+						PhysicsAsset,
+						&SkinnedAsset->GetRefSkeleton(),
+						CachedCollisionData,
+						LevelSetCollisions,
+						SkinnedLevelSetCollisions,
+						CachedUsedBoneIndices,
+						bUseSphylsOnly);
+				}
+			}
+
+			// Transform and add the cached collisions
+			if (CachedUsedBoneIndices.Num())
+			{
+				// Calculate the component to component transform
+				FTransform ComponentToComponentTransform;
+				if (SourceComponent != &InOwnerComponent)
+				{
+					FTransform DestClothComponentTransform = InOwnerComponent.GetComponentTransform();
+					DestClothComponentTransform.RemoveScaling();  // The collision source doesn't need the scale of the cloth skeletal mesh applied to it (but it does need the source scale from the component transform)
+					ComponentToComponentTransform = SourceComponent->GetComponentTransform() * DestClothComponentTransform.Inverse();
+				}
+
+				// Retrieve the bone transforms
+				TArray<FTransform> BoneTransforms;
+				BoneTransforms.Reserve(CachedUsedBoneIndices.Num());
+				for (const int32 UsedBoneIndex : CachedUsedBoneIndices)
+				{
+					BoneTransforms.Emplace(SourceComponent->GetBoneTransform(UsedBoneIndex, ComponentToComponentTransform));
+				}
+
+				// Append the transformed collision elements
+				CollisionData.AppendTransformed(CachedCollisionData, BoneTransforms);
+			}
+		}
+	}
+}  // End namespace UE::Chaos::ClothAsset
diff --git a/Engine/Plugins/ChaosClothAsset/Source/ChaosClothAssetEngine/Private/ChaosClothAsset/CollisionSources.h b/Engine/Plugins/ChaosClothAsset/Source/ChaosClothAssetEngine/Private/ChaosClothAsset/CollisionSources.h
new file mode 100644
--- /dev/null
+++ b/Engine/Plugins/ChaosClothAsset/Source/ChaosClothAssetEngine/Private/ChaosClothAsset/CollisionSources.h
@@ -0,0 +1,79 @@
+// Copyright Epic Games, Inc. All Rights Reserved.
+
+#pragma once
+
+#include "ClothCollisionData.h"
+#include "Delegates/Delegate.h"
+
+class USkinnedMeshComponent;
+class USkinnedAsset;
+class UPhysicsAsset;
+
+namespace UE::Chaos::ClothAsset
+{
+	/**
+	 * Cloth collision source container.
+	 */
+	class FCollisionSources final
+	{
+	public:
+		explicit FCollisionSources(USkinnedMeshComponent* InOwnerComponent): OwnerComponent(InOwnerComponent) {}
+
+		void Add(USkinnedMeshComponent* SourceComponent, const UPhysicsAsset* SourcePhysicsAsset, bool bUseSphylsOnly = false);
+
+		void Remove(const USkinnedMeshComponent* SourceComponent);
+
+		void Remove(const USkinnedMeshComponent* SourceComponent, const UPhysicsAsset* SourcePhysicsAsset);
+
+		void Reset();
+
+	private:
+		friend class FCollisionSourcesProxy;
+
+		struct FCollisionSource
+		{
+			const TWeakObjectPtr<USkinnedMeshComponent> SourceComponent;
+			const TWeakObjectPtr<const UPhysicsAsset> SourcePhysicsAsset;
+			TWeakObjectPtr<const USkinnedAsset> CachedSkinnedAsset;
+			FClothCollisionData CachedCollisionData;
+			TArray<int32> CachedUsedBoneIndices;
+			FDelegateHandle OnBoneTransformsFinalizedHandle;
+			bool bUseSphylsOnly = false;
+
+			FCollisionSource(
+				USkinnedMeshComponent* InSourceComponent,
+				const UPhysicsAsset* InSourcePhysicsAsset,
+				const FSimpleDelegate& InOnBoneTransformsFinalizedDelegate,
+				bool bInUseSphylsOnly);
+			~FCollisionSource();
+
+			void ExtractCollisionData(const USkinnedMeshComponent& InOwnerComponent, FClothCollisionData& CollisionData);
+		};
+
+		void ExtractCollisionData(FClothCollisionData& CollisionData);
+
+		const TWeakObjectPtr<USkinnedMeshComponent> OwnerComponent;
+		TArray<FCollisionSource> CollisionSources;
+		int32 Version = INDEX_NONE;
+	};
+
+	/**
+	 * Use a proxy object to extract collision data from the collision sources.
+	 * The proxy allows for a different ownership than of the CollisionSources' owning component,
+	 * permitting the collision data to remains with the simulation proxy even after the simulation proxy has been replaced.
+	 */
+	class FCollisionSourcesProxy final
+	{
+	public:
+		explicit FCollisionSourcesProxy(FCollisionSources& InCollisionSources) : CollisionSources(InCollisionSources) {}
+
+		const FClothCollisionData& GetCollisionData() const { return CollisionData; }
+
+		void ExtractCollisionData();
+
+	protected:
+		FCollisionSources& CollisionSources;
+		FClothCollisionData CollisionData;
+		int32 Version = INDEX_NONE;
+	};
+}
diff --git a/Engine/Plugins/ChaosClothAsset/Source/ChaosClothAssetEngine/Public/ChaosClothAsset/ClothComponent.h b/Engine/Plugins/ChaosClothAsset/Source/ChaosClothAssetEngine/Public/ChaosClothAsset/ClothComponent.h
--- a/Engine/Plugins/ChaosClothAsset/Source/ChaosClothAssetEngine/Public/ChaosClothAsset/ClothComponent.h
+++ b/Engine/Plugins/ChaosClothAsset/Source/ChaosClothAssetEngine/Public/ChaosClothAsset/ClothComponent.h
@@ -19,6 +19,7 @@ namespace UE::Chaos::ClothAsset
 {
 	class FClothSimulationProxy;
 	class FClothComponentCacheAdapter;
+	class FCollisionSources;
 }
 
 /**
@@ -96,6 +97,31 @@ public:
 	UFUNCTION(BlueprintCallable, Category = "ClothComponent")
 	UChaosClothAssetInteractor* GetClothOutfitInteractor();
 
+	/**
+	 * Add a collision source for the cloth on this component.
+	 * Each cloth tick, the collision defined by the physics asset, transformed by the bones in the source component, will be applied to the simulation.
+	 * @param SourceComponent The component to extract collision transforms from.
+	 * @param SourcePhysicsAsset The physics asset that defines the collision primitives (that will be transformed by the SourceComponent's bones).
+	 * @param bUseSphylsOnly Whether to only use spheres and capsules from the collision sources (which is faster and matches the legacy behavior).
+	 */
+	UFUNCTION(BlueprintCallable, Category = "ClothComponent", Meta = (Keywords = "Chaos Cloth Collision Source"))
+	void AddCollisionSource(USkinnedMeshComponent* SourceComponent, const UPhysicsAsset* SourcePhysicsAsset, bool bUseSphylsOnly = false);
+
+	/** Remove a cloth collision source matching the specified component and physics asset, */
+	UFUNCTION(BlueprintCallable, Category = "ClothComponent", Meta = (Keywords = "Chaos Cloth Collision Source"))
+	void RemoveCollisionSource(const USkinnedMeshComponent* SourceComponent, const UPhysicsAsset* SourcePhysicsAsset);
+
+	/** Remove all cloth collision sources matching the specified component. */
+	UFUNCTION(BlueprintCallable, Category = "ClothComponent", Meta = (Keywords = "Chaos Cloth Collision Source"))
+	void RemoveCollisionSources(const USkinnedMeshComponent* SourceComponent);
+
+	/** Remove all cloth collision sources. */
+	UFUNCTION(BlueprintCallable, Category = "ClothComponent", Meta = (Keywords = "Chaos Cloth Collision Source"))
+	void ResetCollisionSources();
+
+	/** Return all collision sources currently assigned to this component. */
+	UE::Chaos::ClothAsset::FCollisionSources& GetCollisionSources() const { return *CollisionSources; }
+
 	/**
 	 * Return the property collections holding the runtime properties for this cloth component (one per LOD).
 	 * This might be different from the cloth asset's since the component's properties can be modified in code or in blueprints.
@@ -220,6 +246,9 @@ private:
 
 	TSharedPtr<UE::Chaos::ClothAsset::FClothSimulationProxy> ClothSimulationProxy;
 
-	/** Multicaster fired when this component bone transforms are finalized */
+	// Multicaster fired when this component bone transforms are finalized
 	FOnBoneTransformsFinalizedMultiCast OnBoneTransformsFinalizedMC;
+
+	// External sources for collision
+	TUniquePtr<UE::Chaos::ClothAsset::FCollisionSources> CollisionSources;
 };
diff --git a/Engine/Plugins/ChaosClothAsset/Source/ChaosClothAssetEngine/Public/ChaosClothAsset/ClothSimulationProxy.h b/Engine/Plugins/ChaosClothAsset/Source/ChaosClothAssetEngine/Public/ChaosClothAsset/ClothSimulationProxy.h
--- a/Engine/Plugins/ChaosClothAsset/Source/ChaosClothAssetEngine/Public/ChaosClothAsset/ClothSimulationProxy.h
+++ b/Engine/Plugins/ChaosClothAsset/Source/ChaosClothAssetEngine/Public/ChaosClothAsset/ClothSimulationProxy.h
@@ -26,6 +26,7 @@ struct FClothingSimulationCacheData;
 
 namespace UE::Chaos::ClothAsset
 {
+	class FCollisionSourcesProxy;
 	struct FClothSimulationContext;
 
 	/**
@@ -96,6 +97,9 @@ namespace UE::Chaos::ClothAsset
 		// Simulation context used to store the required component data for the duration of the simulation
 		TUniquePtr<FClothSimulationContext> ClothSimulationContext;
 
+		// The collision data for the external collision sources
+		TUniquePtr<FCollisionSourcesProxy> CollisionSourcesProxy;
+
 		// The cloth simulation model used to create this simulation, ownership might get transferred to this proxy if it changes during the simulation
 		TSharedPtr<const FChaosClothSimulationModel> ClothSimulationModel;
 
diff --git a/Engine/Source/Runtime/ClothingSystemRuntimeInterface/Private/ClothCollisionData.cpp b/Engine/Source/Runtime/ClothingSystemRuntimeInterface/Private/ClothCollisionData.cpp
--- a/Engine/Source/Runtime/ClothingSystemRuntimeInterface/Private/ClothCollisionData.cpp
+++ b/Engine/Source/Runtime/ClothingSystemRuntimeInterface/Private/ClothCollisionData.cpp
@@ -37,3 +37,32 @@ void FClothCollisionData::Append(const FClothCollisionData& InOther)
 
 	Boxes.Append(InOther.Boxes);
 }
+
+void FClothCollisionData::AppendTransformed(const FClothCollisionData& InOther, const TArray<FTransform>& BoneTransforms)
+{
+	const int32 NumSpheresBefore = Spheres.Num();
+	const int32 NumConvexesBefore = Convexes.Num();
+	const int32 NumBoxesBefore = Boxes.Num();
+
+	Append(InOther);
+
+	for (int32 Index = NumSpheresBefore; Index < Spheres.Num(); ++Index)
+	{
+		const FTransform& Transform = BoneTransforms[Spheres[Index].BoneIndex];
+		Spheres[Index].LocalPosition = Transform.TransformPosition(Spheres[Index].LocalPosition);
+	}
+	for (int32 Index = NumConvexesBefore; Index < Convexes.Num(); ++Index)
+	{
+		const FTransform& Transform = BoneTransforms[Convexes[Index].BoneIndex];
+		for (FVector& SurfacePoint : Convexes[Index].SurfacePoints)
+		{
+			SurfacePoint = Transform.TransformPosition(SurfacePoint);
+		}
+	}
+	for (int32 Index = NumBoxesBefore; Index < Boxes.Num(); ++Index)
+	{
+		const FTransform& Transform = BoneTransforms[Boxes[Index].BoneIndex];
+		Boxes[Index].LocalPosition = Transform.TransformPosition(Boxes[Index].LocalPosition);
+		Boxes[Index].LocalRotation = Transform.TransformRotation(Boxes[Index].LocalRotation);
+	}
+}
diff --git a/Engine/Source/Runtime/ClothingSystemRuntimeInterface/Public/ClothCollisionData.h b/Engine/Source/Runtime/ClothingSystemRuntimeInterface/Public/ClothCollisionData.h
--- a/Engine/Source/Runtime/ClothingSystemRuntimeInterface/Public/ClothCollisionData.h
+++ b/Engine/Source/Runtime/ClothingSystemRuntimeInterface/Public/ClothCollisionData.h
@@ -16,6 +16,8 @@ struct FClothCollisionData
 
 	CLOTHINGSYSTEMRUNTIMEINTERFACE_API void Append(const FClothCollisionData& InOther);
 
+	CLOTHINGSYSTEMRUNTIMEINTERFACE_API void AppendTransformed(const FClothCollisionData& InOther, const TArray<FTransform>& BoneTransforms);
+
 	// Sphere data
 	UPROPERTY(EditAnywhere, Category = Collison)
 	TArray<FClothCollisionPrim_Sphere> Spheres;
