CIS - Fixed Issue 478390 - Compile errors in WeakObjectPtrTemplates.h.

diff --git a/Engine/Plugins/ChaosClothAsset/Source/ChaosClothAssetEngine/Private/ChaosClothAsset/CollisionSources.cpp b/Engine/Plugins/ChaosClothAsset/Source/ChaosClothAssetEngine/Private/ChaosClothAsset/CollisionSources.cpp
--- a/Engine/Plugins/ChaosClothAsset/Source/ChaosClothAssetEngine/Private/ChaosClothAsset/CollisionSources.cpp
+++ b/Engine/Plugins/ChaosClothAsset/Source/ChaosClothAssetEngine/Private/ChaosClothAsset/CollisionSources.cpp
@@ -10,6 +10,12 @@
 
 namespace UE::Chaos::ClothAsset
 {
+	FCollisionSources::FCollisionSources(USkinnedMeshComponent* InOwnerComponent)
+		: OwnerComponent(InOwnerComponent)
+	{}
+
+	FCollisionSources::~FCollisionSources() = default;
+
 	void FCollisionSources::Add(USkinnedMeshComponent* SourceComponent, const UPhysicsAsset* SourcePhysicsAsset, bool bUseSphylsOnly)
 	{
 		using namespace UE::Chaos::ClothAsset;
diff --git a/Engine/Plugins/ChaosClothAsset/Source/ChaosClothAssetEngine/Private/ChaosClothAsset/CollisionSources.h b/Engine/Plugins/ChaosClothAsset/Source/ChaosClothAssetEngine/Private/ChaosClothAsset/CollisionSources.h
--- a/Engine/Plugins/ChaosClothAsset/Source/ChaosClothAssetEngine/Private/ChaosClothAsset/CollisionSources.h
+++ b/Engine/Plugins/ChaosClothAsset/Source/ChaosClothAssetEngine/Private/ChaosClothAsset/CollisionSources.h
@@ -17,7 +17,8 @@ namespace UE::Chaos::ClothAsset
 	class FCollisionSources final
 	{
 	public:
-		explicit FCollisionSources(USkinnedMeshComponent* InOwnerComponent): OwnerComponent(InOwnerComponent) {}
+		explicit FCollisionSources(USkinnedMeshComponent* InOwnerComponent);
+		~FCollisionSources();
 
 		void Add(USkinnedMeshComponent* SourceComponent, const UPhysicsAsset* SourcePhysicsAsset, bool bUseSphylsOnly = false);
 
