Chaos Cloth Asset - Enabled the SetStringValue for the blueprint interator.

#rnx

diff --git a/Engine/Plugins/ChaosClothAsset/Source/ChaosClothAssetEngine/Public/ChaosClothAsset/ClothAssetInteractor.h b/Engine/Plugins/ChaosClothAsset/Source/ChaosClothAssetEngine/Public/ChaosClothAsset/ClothAssetInteractor.h
--- a/Engine/Plugins/ChaosClothAsset/Source/ChaosClothAssetEngine/Public/ChaosClothAsset/ClothAssetInteractor.h
+++ b/Engine/Plugins/ChaosClothAsset/Source/ChaosClothAssetEngine/Public/ChaosClothAsset/ClothAssetInteractor.h
@@ -80,6 +80,7 @@ public:
 	void SetVectorValue(const FString& PropertyName, int32 LODIndex = -1, FVector Value = FVector(0.));
 
 	/**Set the string value for a property (if it exists). This is typically the map name associated with a property. All LODs will be set when LODIndex = -1.*/
+	UFUNCTION(BlueprintCallable, Category = ClothProperty)
 	void SetStringValue(const FString& PropertyName, int32 LODIndex = -1, const FString& Value = "");
 
 private:
