Chaos Cloth Asset: fix sim mesh normals calculation on asset creation.

We weren't importing the Normals from Skeletal or Static Meshes and instead were calculating them from face normals. Furthermore, the normals we were calculating were flipped, including for the USDImport.

* When calculating normals in the PatternFacade, calculate the correct non-flipped normals.
* Create a v2 for Static and Skeletal Mesh Import Nodes. These import the normals from the mesh. v1 retains existing behavior including flipping normals back to be reversed.
* The USDImportNode will now generate non-flipped normals. Didn't create a v2 because the imported values are stored on the node, so no change in behavior as long as you don't reimport the mesh. If you do reimport and want to retain  the old behavior, you can add a ReverseNormals node to flip them back.

diff --git a/Engine/Plugins/ChaosClothAsset/Source/ChaosClothAsset/Private/ChaosClothAsset/ClothGeometryTools.cpp b/Engine/Plugins/ChaosClothAsset/Source/ChaosClothAsset/Private/ChaosClothAsset/ClothGeometryTools.cpp
--- a/Engine/Plugins/ChaosClothAsset/Source/ChaosClothAsset/Private/ChaosClothAsset/ClothGeometryTools.cpp
+++ b/Engine/Plugins/ChaosClothAsset/Source/ChaosClothAsset/Private/ChaosClothAsset/ClothGeometryTools.cpp
@@ -24,6 +24,7 @@ namespace UE::Chaos::ClothAsset
 			TArray<int32> Indices;  // 3x number of triangles
 			TArray<FVector2f> Positions2D;
 			TArray<FVector3f> Positions3D;  // Same size as Positions
+			TArray<FVector3f> Normals; // Empty or same size as Positions
 			TArray<int32> PositionToSourceIndex; // Same size as Positions. Index in the original welded position array
 		};
 
@@ -80,9 +81,11 @@ namespace UE::Chaos::ClothAsset
 			return Index0 < Index1 ? FIntVector2(Index0, Index1) : FIntVector2(Index1, Index0);
 		}
 
-		static void UnwrapDynamicMesh(const UE::Geometry::FDynamicMesh3& DynamicMesh, TArray<FIsland>& OutIslands)
+		static void UnwrapDynamicMesh(const UE::Geometry::FDynamicMesh3& DynamicMesh, bool bImportNormals, TArray<FIsland>& OutIslands)
 		{
 			using namespace UE::Geometry;
+			const UE::Geometry::FDynamicMeshAttributeSet* const AttributeSet = DynamicMesh.Attributes();
+			const UE::Geometry::FDynamicMeshNormalOverlay* const NormalOverlay = bImportNormals && AttributeSet ? AttributeSet->PrimaryNormals() : nullptr;
 
 			OutIslands.Reset();
 			constexpr float SquaredWeldingDistance = FMath::Square(0.01f);  // 0.1 mm
@@ -100,6 +103,7 @@ namespace UE::Chaos::ClothAsset
 					continue;
 				}
 				const FIndex3i TriangleIndices = DynamicMesh.GetTriangle(SeedTriangle);
+				const FIndex3i TriangleNormalElements = NormalOverlay ? NormalOverlay->GetTriangle(SeedTriangle) : FIndex3i::Invalid();
 
 				const int32 SeedIndex0 = TriangleIndices[0];
 				const int32 SeedIndex1 = TriangleIndices[1];
@@ -118,6 +122,11 @@ namespace UE::Chaos::ClothAsset
 
 				Island.Positions3D.Add(Position0);
 				Island.Positions3D.Add(Position1);
+				if (NormalOverlay)
+				{
+					Island.Normals.Add(NormalOverlay->GetElement(TriangleNormalElements[0]));
+					Island.Normals.Add(NormalOverlay->GetElement(TriangleNormalElements[1]));
+				}
 				Island.PositionToSourceIndex.Add(SeedIndex0);
 				Island.PositionToSourceIndex.Add(SeedIndex1);
 
@@ -129,12 +138,14 @@ namespace UE::Chaos::ClothAsset
 					int32 Triangle;
 					FIndex2i OldEdge;
 					FIndex2i NewEdge;
+					FIndex2i NormalIndices;
 					int32 CrossEdgePoint;  // Keep the opposite point to orientate degenerate cases
 				} Visitor =
 				{
 					SeedTriangle,
 					FIndex2i(SeedIndex0, SeedIndex1),
 					FIndex2i(SeedIndex2D0, SeedIndex2D1),
+					FIndex2i(TriangleNormalElements[0], TriangleNormalElements[1]),
 					INDEX_NONE
 				};
 
@@ -149,10 +160,13 @@ namespace UE::Chaos::ClothAsset
 					const int32 OldIndex1 = Visitor.OldEdge.B;
 					const int32 NewIndex0 = Visitor.NewEdge.A;
 					const int32 NewIndex1 = Visitor.NewEdge.B;
+					const int32 NormalIndex0 = Visitor.NormalIndices.A;
+					const int32 NormalIndex1 = Visitor.NormalIndices.B;
 
 					// Find opposite index from this triangle edge
 
 					const int32 OldIndex2 = IndexUtil::FindTriOtherVtxUnsafe(OldIndex0, OldIndex1, DynamicMesh.GetTriangle(Triangle));
+					const int32 NormalIndex2 = NormalOverlay ? IndexUtil::FindTriOtherVtxUnsafe(NormalIndex0, NormalIndex1, NormalOverlay->GetTriangle(Triangle)) : INDEX_NONE;
 
 					// Find the 2D intersection of the two connecting adjacent edges using the 3D reference length
 					const FVector3f P0(DynamicMesh.GetVertexRef(OldIndex0));
@@ -202,6 +216,10 @@ namespace UE::Chaos::ClothAsset
 					{
 						NewIndex2 = Island.Positions2D.Add(C2);
 						Island.Positions3D.Add(P2);
+						if (NormalOverlay)
+						{
+							Island.Normals.Add(NormalOverlay->GetElement(NormalIndex2));
+						}
 						Island.PositionToSourceIndex.Add(OldIndex2);
 					}
 
@@ -226,6 +244,12 @@ namespace UE::Chaos::ClothAsset
 						FIndex3i(NewIndex2, NewIndex1, NewIndex0),
 						FIndex3i(NewIndex0, NewIndex2, NewIndex1)
 					};
+					const FIndex2i NormalEdgeList[3] =
+					{
+						FIndex2i(NormalIndex1, NormalIndex0),
+						FIndex2i(NormalIndex2, NormalIndex1),
+						FIndex2i(NormalIndex0, NormalIndex2)
+					};
 					for (int32 Edge = 0; Edge < 3; ++Edge)
 					{
 						const int32 EdgeIndex0 = OldEdgeList[Edge].A;
@@ -246,6 +270,7 @@ namespace UE::Chaos::ClothAsset
 										NeighborTriangle,
 										OldEdgeList[Edge],
 										FIndex2i(NewEdgeList[Edge].A, NewEdgeList[Edge].B),
+										NormalEdgeList[Edge],
 										NewEdgeList[Edge].C,  // Pass the cross edge 2D opposite point to help define orientation of any degenerated triangles
 									});
 							}
@@ -256,12 +281,14 @@ namespace UE::Chaos::ClothAsset
 			}
 		}
 
-		static void BuildIslandsFromDynamicMeshUVs(const UE::Geometry::FDynamicMeshUVOverlay& UVOverlay, const FVector2f& UVScale, TArray<FIsland>& OutIslands)
+		static void BuildIslandsFromDynamicMeshUVs(const UE::Geometry::FDynamicMeshUVOverlay& UVOverlay, const FVector2f& UVScale, bool bImportNormals, TArray<FIsland>& OutIslands)
 		{
 			using namespace UE::Geometry;
 
-			const FDynamicMesh3* DynamicMesh = UVOverlay.GetParentMesh();
+			const FDynamicMesh3* const DynamicMesh = UVOverlay.GetParentMesh();
 			check(DynamicMesh);
+			const UE::Geometry::FDynamicMeshAttributeSet* const AttributeSet = DynamicMesh->Attributes();
+			const UE::Geometry::FDynamicMeshNormalOverlay* const NormalOverlay = bImportNormals && AttributeSet ? AttributeSet->PrimaryNormals() : nullptr;
 
 			OutIslands.Reset();
 
@@ -300,22 +327,27 @@ namespace UE::Chaos::ClothAsset
 					const int32 Triangle = Visitor.Triangle;
 					const FIndex3i TriangleIndices = DynamicMesh->GetTriangle(Triangle);
 					const FIndex3i TriangleUVElements = UVOverlay.GetTriangle(Triangle);
+					const FIndex3i TriangleNormalElements = NormalOverlay ? NormalOverlay->GetTriangle(Triangle) : FIndex3i::Invalid();
 
-					auto GetOrAddNewIndex = [&UVOverlay, &Island, &SourceElementIndexToNewIndex, &DynamicMesh, &UVScale](int32 ElementId, int32 VertexId)
+					auto GetOrAddNewIndex = [&UVOverlay, &Island, &SourceElementIndexToNewIndex, &DynamicMesh, &UVScale, NormalOverlay](int32 ElementId, int32 VertexId, int32 NormalId)
 					{
 						int32& NewIndex = SourceElementIndexToNewIndex[ElementId];
 						if (NewIndex == INDEX_NONE)
 						{
 							NewIndex = Island.Positions3D.Add(FVector3f(DynamicMesh->GetVertexRef(VertexId)));
 							Island.Positions2D.Add((FVector2f(1.f) - UVOverlay.GetElement(ElementId)) * UVScale);  // The static mesh import uses 1 - UV for some reason
+							if (NormalOverlay)
+							{
+								Island.Normals.Add(NormalOverlay->GetElement(NormalId));
+							}
 							Island.PositionToSourceIndex.Add(VertexId);
 						}
 						return NewIndex;
 					};
 
-					const int32 NewIndex0 = GetOrAddNewIndex(TriangleUVElements[0], TriangleIndices[0]);
-					const int32 NewIndex1 = GetOrAddNewIndex(TriangleUVElements[1], TriangleIndices[1]);
-					const int32 NewIndex2 = GetOrAddNewIndex(TriangleUVElements[2], TriangleIndices[2]);
+					const int32 NewIndex0 = GetOrAddNewIndex(TriangleUVElements[0], TriangleIndices[0], TriangleNormalElements[0]);
+					const int32 NewIndex1 = GetOrAddNewIndex(TriangleUVElements[1], TriangleIndices[1], TriangleNormalElements[1]);
+					const int32 NewIndex2 = GetOrAddNewIndex(TriangleUVElements[2], TriangleIndices[2], TriangleNormalElements[2]);
 					Island.Indices.Add(NewIndex0);
 					Island.Indices.Add(NewIndex1);
 					Island.Indices.Add(NewIndex2);
@@ -855,7 +887,7 @@ namespace UE::Chaos::ClothAsset
 
 	void FClothGeometryTools::BuildSimMeshFromDynamicMesh(
 		const TSharedRef<FManagedArrayCollection>& ClothCollection,
-		const UE::Geometry::FDynamicMesh3& DynamicMesh, int32 UVChannelIndex, const FVector2f& UVScale, bool bAppend)
+		const UE::Geometry::FDynamicMesh3& DynamicMesh, int32 UVChannelIndex, const FVector2f& UVScale, bool bAppend, bool bImportNormals)
 	{
 		using namespace Private::SimMeshBuilder;
 
@@ -866,16 +898,16 @@ namespace UE::Chaos::ClothAsset
 
 		const UE::Geometry::FDynamicMeshAttributeSet* const AttributeSet = DynamicMesh.Attributes();
 		const UE::Geometry::FDynamicMeshUVOverlay* const UVOverlay = AttributeSet ? AttributeSet->GetUVLayer(UVChannelIndex) : nullptr;
-		const UE::Geometry::FDynamicMeshVertexSkinWeightsAttribute* SkinWeights = AttributeSet ? AttributeSet->GetSkinWeightsAttribute(FName("Default")) : nullptr;
+		const UE::Geometry::FDynamicMeshVertexSkinWeightsAttribute* SkinWeights = AttributeSet ? AttributeSet->GetSkinWeightsAttribute(FName("Default")) : nullptr; 
 
 		TArray<FIsland> Islands;
 		if (UVOverlay)
 		{
-			BuildIslandsFromDynamicMeshUVs(*UVOverlay, UVScale, Islands);
+			BuildIslandsFromDynamicMeshUVs(*UVOverlay, UVScale, bImportNormals, Islands);
 		}
 		else
 		{
-			UnwrapDynamicMesh(DynamicMesh, Islands);
+			UnwrapDynamicMesh(DynamicMesh, bImportNormals, Islands);
 		}
 
 		FCollectionClothFacade Cloth(ClothCollection);
@@ -885,7 +917,7 @@ namespace UE::Chaos::ClothAsset
 			{
 				FCollectionClothSimPatternFacade Pattern = Cloth.AddGetSimPattern();
 				const int32 VertexOffset = Cloth.GetNumSimVertices3D();
-				Pattern.Initialize(Island.Positions2D, Island.Positions3D, Island.Indices);
+				Pattern.Initialize(Island.Positions2D, Island.Positions3D, Island.Indices, INDEX_NONE, Island.Normals);
 
 				// Copy skinning data
 				if (SkinWeights)
diff --git a/Engine/Plugins/ChaosClothAsset/Source/ChaosClothAsset/Private/ChaosClothAsset/CollectionClothSimPatternFacade.cpp b/Engine/Plugins/ChaosClothAsset/Source/ChaosClothAsset/Private/ChaosClothAsset/CollectionClothSimPatternFacade.cpp
--- a/Engine/Plugins/ChaosClothAsset/Source/ChaosClothAsset/Private/ChaosClothAsset/CollectionClothSimPatternFacade.cpp
+++ b/Engine/Plugins/ChaosClothAsset/Source/ChaosClothAsset/Private/ChaosClothAsset/CollectionClothSimPatternFacade.cpp
@@ -113,7 +113,7 @@ namespace UE::Chaos::ClothAsset
 	}
 
 	template<typename IndexType UE_REQUIRES_DEFINITION(TIsIndexType<IndexType>::Value)>
-	void FCollectionClothSimPatternFacade::Initialize(const TArray<FVector2f>& Positions2D, const TArray<FVector3f>& Positions3D, const TArray<IndexType>& Indices, const int32 FabricIndex)
+	void FCollectionClothSimPatternFacade::Initialize(const TArray<FVector2f>& Positions2D, const TArray<FVector3f>& Positions3D, const TArray<IndexType>& Indices, const int32 FabricIndex, const TArray<FVector3f>& Normals)
 	{
 		Reset();
 
@@ -134,7 +134,7 @@ namespace UE::Chaos::ClothAsset
 		{
 			SimPosition2D[SimVertexIndex] = Positions2D[SimVertexIndex];
 			SimPosition3D[SimVertexIndex + StartVertexIndex[1]] = Positions3D[SimVertexIndex];
-			SimNormal[SimVertexIndex + StartVertexIndex[1]] = FVector3f::ZeroVector;
+			SimNormal[SimVertexIndex + StartVertexIndex[1]] = Normals.IsValidIndex(SimVertexIndex) ? Normals[SimVertexIndex] : FVector3f::ZeroVector;
 		}
 
 		const int32 NumSimFaces = Private::GetNumFaces(Indices);
@@ -146,19 +146,19 @@ namespace UE::Chaos::ClothAsset
 		for (int32 SimFaceIndex = 0; SimFaceIndex < NumSimFaces; ++SimFaceIndex)
 		{
 			// Indices from the start of the pattern
-			const int32 PattternVertexIndex0 = Private::GetIndex<0>(Indices, SimFaceIndex);
-			const int32 PattternVertexIndex1 = Private::GetIndex<1>(Indices, SimFaceIndex);
-			const int32 PattternVertexIndex2 = Private::GetIndex<2>(Indices, SimFaceIndex);
+			const int32 PatternVertexIndex0 = Private::GetIndex<0>(Indices, SimFaceIndex);
+			const int32 PatternVertexIndex1 = Private::GetIndex<1>(Indices, SimFaceIndex);
+			const int32 PatternVertexIndex2 = Private::GetIndex<2>(Indices, SimFaceIndex);
 
 			// 2D Indices from the start of the cloth
-			const int32 Vertex2DIndex0 = PattternVertexIndex0 + BaseVertex2D;
-			const int32 Vertex2DIndex1 = PattternVertexIndex1 + BaseVertex2D;
-			const int32 Vertex2DIndex2 = PattternVertexIndex2 + BaseVertex2D;
+			const int32 Vertex2DIndex0 = PatternVertexIndex0 + BaseVertex2D;
+			const int32 Vertex2DIndex1 = PatternVertexIndex1 + BaseVertex2D;
+			const int32 Vertex2DIndex2 = PatternVertexIndex2 + BaseVertex2D;
 
 			// 3D Indices from the start of the cloth
-			const int32 Vertex3DIndex0 = PattternVertexIndex0 + StartVertexIndex[1];
-			const int32 Vertex3DIndex1 = PattternVertexIndex1 + StartVertexIndex[1];
-			const int32 Vertex3DIndex2 = PattternVertexIndex2 + StartVertexIndex[1];
+			const int32 Vertex3DIndex0 = PatternVertexIndex0 + StartVertexIndex[1];
+			const int32 Vertex3DIndex1 = PatternVertexIndex1 + StartVertexIndex[1];
+			const int32 Vertex3DIndex2 = PatternVertexIndex2 + StartVertexIndex[1];
 
 			// Set indices in cloth index space
 			SimIndices2D[SimFaceIndex] = FIntVector3(Vertex2DIndex0, Vertex2DIndex1, Vertex2DIndex2);
@@ -168,10 +168,19 @@ namespace UE::Chaos::ClothAsset
 			const FVector3f& Pos0 = SimPosition3D[Vertex3DIndex0];
 			const FVector3f& Pos1 = SimPosition3D[Vertex3DIndex1];
 			const FVector3f& Pos2 = SimPosition3D[Vertex3DIndex2];
-			const FVector3f Normal = (Pos1 - Pos0).Cross(Pos2 - Pos0).GetSafeNormal();
-			SimNormal[Vertex3DIndex0] += Normal;
-			SimNormal[Vertex3DIndex1] += Normal;
-			SimNormal[Vertex3DIndex2] += Normal;
+			const FVector3f Normal = (Pos2 - Pos0).Cross(Pos1 - Pos0).GetSafeNormal();
+			if (!Normals.IsValidIndex(PatternVertexIndex0))
+			{
+				SimNormal[Vertex3DIndex0] += Normal;
+			}
+			if (!Normals.IsValidIndex(PatternVertexIndex1))
+			{
+				SimNormal[Vertex3DIndex1] += Normal;
+			}
+			if (!Normals.IsValidIndex(PatternVertexIndex2))
+			{
+				SimNormal[Vertex3DIndex2] += Normal;
+			}
 		}
 
 		// Normalize normals
@@ -180,9 +189,9 @@ namespace UE::Chaos::ClothAsset
 			SimNormal[SimVertexIndex + StartVertexIndex[1]] = SimNormal[SimVertexIndex + StartVertexIndex[1]].GetSafeNormal(UE_SMALL_NUMBER, FVector3f::XAxisVector);
 		}
 	}
-	template CHAOSCLOTHASSET_API void FCollectionClothSimPatternFacade::Initialize(const TArray<FVector2f>& Positions2D, const TArray<FVector3f>& Positions3D, const TArray<int32>& Indices, const int32 FabricIndex);
-	template CHAOSCLOTHASSET_API void FCollectionClothSimPatternFacade::Initialize(const TArray<FVector2f>& Positions2D, const TArray<FVector3f>& Positions3D, const TArray<uint32>& Indices, const int32 FabricIndex);
-	template CHAOSCLOTHASSET_API void FCollectionClothSimPatternFacade::Initialize(const TArray<FVector2f>& Positions2D, const TArray<FVector3f>& Positions3D, const TArray<FIntVector3>& Indices, const int32 FabricIndex);
+	template CHAOSCLOTHASSET_API void FCollectionClothSimPatternFacade::Initialize(const TArray<FVector2f>& Positions2D, const TArray<FVector3f>& Positions3D, const TArray<int32>& Indices, const int32 FabricIndex, const TArray<FVector3f>& Normals);
+	template CHAOSCLOTHASSET_API void FCollectionClothSimPatternFacade::Initialize(const TArray<FVector2f>& Positions2D, const TArray<FVector3f>& Positions3D, const TArray<uint32>& Indices, const int32 FabricIndex, const TArray<FVector3f>& Normals);
+	template CHAOSCLOTHASSET_API void FCollectionClothSimPatternFacade::Initialize(const TArray<FVector2f>& Positions2D, const TArray<FVector3f>& Positions3D, const TArray<FIntVector3>& Indices, const int32 FabricIndex, const TArray<FVector3f>& Normals);
 
 	void FCollectionClothSimPatternFacade::Initialize(const FCollectionClothSimPatternConstFacade& Other, const int32 SimVertex3DOffset, const int32 FabricsOffset)
 	{
diff --git a/Engine/Plugins/ChaosClothAsset/Source/ChaosClothAsset/Public/ChaosClothAsset/ClothGeometryTools.h b/Engine/Plugins/ChaosClothAsset/Source/ChaosClothAsset/Public/ChaosClothAsset/ClothGeometryTools.h
--- a/Engine/Plugins/ChaosClothAsset/Source/ChaosClothAsset/Public/ChaosClothAsset/ClothGeometryTools.h
+++ b/Engine/Plugins/ChaosClothAsset/Source/ChaosClothAsset/Public/ChaosClothAsset/ClothGeometryTools.h
@@ -78,10 +78,11 @@ namespace UE::Chaos::ClothAsset
 
 		/**
 		* Unwrap and build SimMesh data from a DynamicMesh
+		* Normals are only imported if the DynamicMesh has both a UVOverlay and a NormalOverlay
 		*/
 		static void BuildSimMeshFromDynamicMesh(
 			const TSharedRef<FManagedArrayCollection>& ClothCollection,
-			const UE::Geometry::FDynamicMesh3& DynamicMesh, int32 UVChannelIndex, const FVector2f& UVScale, bool bAppend);
+			const UE::Geometry::FDynamicMesh3& DynamicMesh, int32 UVChannelIndex, const FVector2f& UVScale, bool bAppend, bool bImportNormals = false);
 
 		/**
 		* Remove (topologically) degenerate triangles. Remove any vertices that aren't in a triangle. Compact any lookup arrays that contain INDEX_NONEs.
diff --git a/Engine/Plugins/ChaosClothAsset/Source/ChaosClothAsset/Public/ChaosClothAsset/CollectionClothSimPatternFacade.h b/Engine/Plugins/ChaosClothAsset/Source/ChaosClothAsset/Public/ChaosClothAsset/CollectionClothSimPatternFacade.h
--- a/Engine/Plugins/ChaosClothAsset/Source/ChaosClothAsset/Public/ChaosClothAsset/CollectionClothSimPatternFacade.h
+++ b/Engine/Plugins/ChaosClothAsset/Source/ChaosClothAsset/Public/ChaosClothAsset/CollectionClothSimPatternFacade.h
@@ -89,9 +89,9 @@ namespace UE::Chaos::ClothAsset
 		/** Remove all geometry from this cloth pattern. */
 		void Reset();
 
-		/** Initialize the cloth pattern using the specified 3D and 2D positions, and topology. */
+		/** Initialize the cloth pattern using the specified 3D and 2D positions, and topology. If Normals are not supplied, they will be calculated from the face normals.*/
 		template<typename IndexType UE_REQUIRES(TIsIndexType<IndexType>::Value)>
-		void Initialize(const TArray<FVector2f>& Positions2D, const TArray<FVector3f>& Positions3D, const TArray<IndexType>& Indices, const int32 FabricIndex = INDEX_NONE);
+		void Initialize(const TArray<FVector2f>& Positions2D, const TArray<FVector3f>& Positions3D, const TArray<IndexType>& Indices, const int32 FabricIndex = INDEX_NONE, const TArray<FVector3f>& Normals = TArray<FVector3f>());
 
 		/** Initialize this pattern using another pattern collection. */
 		void Initialize(const FCollectionClothSimPatternConstFacade& Other, const int32 SimVertex3DOffset, const int32 FabricsOffset = 0);
diff --git a/Engine/Plugins/ChaosClothAssetEditor/Source/ChaosClothAssetDataflowNodes/Private/ChaosClothAsset/ChaosClothAssetDataflowNodesModule.cpp b/Engine/Plugins/ChaosClothAssetEditor/Source/ChaosClothAssetDataflowNodes/Private/ChaosClothAsset/ChaosClothAssetDataflowNodesModule.cpp
--- a/Engine/Plugins/ChaosClothAssetEditor/Source/ChaosClothAssetDataflowNodes/Private/ChaosClothAsset/ChaosClothAssetDataflowNodesModule.cpp
+++ b/Engine/Plugins/ChaosClothAssetEditor/Source/ChaosClothAssetDataflowNodes/Private/ChaosClothAsset/ChaosClothAssetDataflowNodesModule.cpp
@@ -127,7 +127,9 @@ namespace UE::Chaos::ClothAsset
 			DATAFLOW_NODE_REGISTER_CREATION_FACTORY(FChaosClothAssetSimulationXPBDBendingSpringConfigNode);
 			DATAFLOW_NODE_REGISTER_CREATION_FACTORY(FChaosClothAssetSimulationXPBDEdgeSpringConfigNode);
 			DATAFLOW_NODE_REGISTER_CREATION_FACTORY(FChaosClothAssetSkeletalMeshImportNode);
+			DATAFLOW_NODE_REGISTER_CREATION_FACTORY(FChaosClothAssetSkeletalMeshImportNode_v2);
 			DATAFLOW_NODE_REGISTER_CREATION_FACTORY(FChaosClothAssetStaticMeshImportNode);
+			DATAFLOW_NODE_REGISTER_CREATION_FACTORY(FChaosClothAssetStaticMeshImportNode_v2);
 			DATAFLOW_NODE_REGISTER_CREATION_FACTORY(FChaosClothAssetTerminalNode);
 			DATAFLOW_NODE_REGISTER_CREATION_FACTORY(FChaosClothAssetTerminalNode_v2);
 			DATAFLOW_NODE_REGISTER_CREATION_FACTORY(FChaosClothAssetTransferSkinWeightsNode);
diff --git a/Engine/Plugins/ChaosClothAssetEditor/Source/ChaosClothAssetDataflowNodes/Private/ChaosClothAsset/ClothDataflowTools.cpp b/Engine/Plugins/ChaosClothAssetEditor/Source/ChaosClothAssetDataflowNodes/Private/ChaosClothAsset/ClothDataflowTools.cpp
--- a/Engine/Plugins/ChaosClothAssetEditor/Source/ChaosClothAssetDataflowNodes/Private/ChaosClothAsset/ClothDataflowTools.cpp
+++ b/Engine/Plugins/ChaosClothAssetEditor/Source/ChaosClothAssetDataflowNodes/Private/ChaosClothAsset/ClothDataflowTools.cpp
@@ -26,7 +26,7 @@ namespace UE::Chaos::ClothAsset
 		// Wrapper for accessing a SkelMeshSection. Implements the interface expected by TToDynamicMesh<>.
 		// This will weld all vertices which are the same.
 		//
-		template<bool bHasNormals = false, bool bHasTangents = false, bool bHasBiTangents = false, bool bHasColors = false>
+		template<bool bHasTangents = false, bool bHasBiTangents = false, bool bHasColors = false>
 		struct FSkelMeshSectionWrapper
 		{
 			typedef int32 TriIDType;
@@ -36,8 +36,9 @@ namespace UE::Chaos::ClothAsset
 			typedef int32 NormalIDType;
 			typedef int32 ColorIDType;
 
-			FSkelMeshSectionWrapper(const FSkeletalMeshLODModel& SkeletalMeshModel, const int32 SectionIndex)
-				: SourceSection(SkeletalMeshModel.Sections[SectionIndex])
+			FSkelMeshSectionWrapper(const FSkeletalMeshLODModel& SkeletalMeshModel, const int32 SectionIndex, bool bInHasNormals)
+				: bHasNormals(bInHasNormals)
+				, SourceSection(SkeletalMeshModel.Sections[SectionIndex])
 				, IndexBuffer(SkeletalMeshModel.IndexBuffer.GetData() + SourceSection.BaseIndex, SourceSection.NumTriangles * 3)
 			{
 				const int32 NumVerts = SourceSection.SoftVertices.Num();
@@ -212,17 +213,28 @@ namespace UE::Chaos::ClothAsset
 
 			const TArray<NormalIDType>& GetNormalIDs() const
 			{
-				return EmptyArray;
+				if (bHasNormals)
+				{
+					return OriginalIndexes;
+				}
+				else
+				{
+					return EmptyArray;
+				}
 			}
 
 			FVector3f GetNormal(NormalIDType ID) const
 			{
-				check(false);
-				return FVector3f();
+				check(bHasNormals);
+				return SourceSection.SoftVertices[ID].TangentZ;
 			}
 
-			bool GetNormalTri(const TriIDType& TID, NormalIDType& NID0, NormalIDType& NID1, NormalIDType& NID2) const
+			bool GetNormalTri(const TriIDType& TriID, NormalIDType& NID0, NormalIDType& NID1, NormalIDType& NID2) const
 			{
+				if (bHasNormals)
+				{
+					return GetTri(TriID, NID0, NID1, NID2);
+				}
 				return false;
 			}
 
@@ -349,7 +361,7 @@ namespace UE::Chaos::ClothAsset
 				return FLinearColor::White;
 			}
 
-
+			const bool bHasNormals;
 			const FSkelMeshSection& SourceSection;
 			const TConstArrayView<uint32> IndexBuffer;
 			TArray<int32> OriginalIndexes; // UniqueIndex -> OrigIndex
@@ -421,13 +433,13 @@ namespace UE::Chaos::ClothAsset
 		ClothPatternFacade.SetRenderMaterialPathName(RenderMaterialPathName);
 	}
 	
-	void FClothDataflowTools::AddSimPatternsFromSkeletalMeshSection(const TSharedRef<FManagedArrayCollection>& ClothCollection, const FSkeletalMeshLODModel& SkeletalMeshModel, const int32 SectionIndex, const int32 UVChannelIndex, const FVector2f& UVScale)
+	void FClothDataflowTools::AddSimPatternsFromSkeletalMeshSection(const TSharedRef<FManagedArrayCollection>& ClothCollection, const FSkeletalMeshLODModel& SkeletalMeshModel, const int32 SectionIndex, const int32 UVChannelIndex, const FVector2f& UVScale, bool bImportNormals)
 	{
 		check(SectionIndex < SkeletalMeshModel.Sections.Num());
 
 		// Convert to DynamicMesh and then use that to create patterns.
 		UE::Geometry::TToDynamicMesh<Private::FSkelMeshSectionWrapper<>> SkelMeshSectionToDynamicMesh;
-		Private::FSkelMeshSectionWrapper<> SectionWrapper(SkeletalMeshModel, SectionIndex);
+		Private::FSkelMeshSectionWrapper<> SectionWrapper(SkeletalMeshModel, SectionIndex, bImportNormals);
 
 		UE::Geometry::FDynamicMesh3 DynamicMesh;
 		DynamicMesh.EnableAttributes();
diff --git a/Engine/Plugins/ChaosClothAssetEditor/Source/ChaosClothAssetDataflowNodes/Private/ChaosClothAsset/ClothDataflowTools.h b/Engine/Plugins/ChaosClothAssetEditor/Source/ChaosClothAssetDataflowNodes/Private/ChaosClothAsset/ClothDataflowTools.h
--- a/Engine/Plugins/ChaosClothAssetEditor/Source/ChaosClothAssetDataflowNodes/Private/ChaosClothAsset/ClothDataflowTools.h
+++ b/Engine/Plugins/ChaosClothAssetEditor/Source/ChaosClothAssetDataflowNodes/Private/ChaosClothAsset/ClothDataflowTools.h
@@ -21,7 +21,7 @@ namespace UE::Chaos::ClothAsset
 	{
 		static void AddRenderPatternFromSkeletalMeshSection(const TSharedRef<FManagedArrayCollection>& ClothCollection, const FSkeletalMeshLODModel& SkeletalMeshModel, const int32 SectionIndex, const FString& RenderMaterialPathName);
 
-		static void AddSimPatternsFromSkeletalMeshSection(const TSharedRef<FManagedArrayCollection>& ClothCollection, const FSkeletalMeshLODModel& SkeletalMeshModel, const int32 SectionIndex, const int32 UVChannelIndex, const FVector2f& UVScale);
+		static void AddSimPatternsFromSkeletalMeshSection(const TSharedRef<FManagedArrayCollection>& ClothCollection, const FSkeletalMeshLODModel& SkeletalMeshModel, const int32 SectionIndex, const int32 UVChannelIndex, const FVector2f& UVScale, bool bImportNormals = false);
 
 		static void LogAndToastWarning(const FDataflowNode& DataflowNode, const FText& Headline, const FText& Details);
 
diff --git a/Engine/Plugins/ChaosClothAssetEditor/Source/ChaosClothAssetDataflowNodes/Private/ChaosClothAsset/SkeletalMeshImportNode.cpp b/Engine/Plugins/ChaosClothAssetEditor/Source/ChaosClothAssetDataflowNodes/Private/ChaosClothAsset/SkeletalMeshImportNode.cpp
--- a/Engine/Plugins/ChaosClothAssetEditor/Source/ChaosClothAssetDataflowNodes/Private/ChaosClothAsset/SkeletalMeshImportNode.cpp
+++ b/Engine/Plugins/ChaosClothAssetEditor/Source/ChaosClothAssetDataflowNodes/Private/ChaosClothAsset/SkeletalMeshImportNode.cpp
@@ -17,6 +17,85 @@
 
 #define LOCTEXT_NAMESPACE "ChaosClothAssetSkeletalMeshImportNode"
 
+FChaosClothAssetSkeletalMeshImportNode_v2::FChaosClothAssetSkeletalMeshImportNode_v2(const Dataflow::FNodeParameters& InParam, FGuid InGuid)
+	: FDataflowNode(InParam, InGuid)
+{
+	RegisterInputConnection(&SkeletalMesh);
+	RegisterOutputConnection(&Collection);
+}
+
+void FChaosClothAssetSkeletalMeshImportNode_v2::Evaluate(Dataflow::FContext& Context, const FDataflowOutput* Out) const
+{
+	using namespace UE::Chaos::ClothAsset;
+
+	if (Out->IsA<FManagedArrayCollection>(&Collection))
+	{
+		const TSharedRef<FManagedArrayCollection> ClothCollection = MakeShared<FManagedArrayCollection>();
+		FCollectionClothFacade ClothFacade(ClothCollection);
+		ClothFacade.DefineSchema();
+
+		if (const USkeletalMesh* const InSkeletalMesh = GetValue(Context, &SkeletalMesh))
+		{
+			const FSkeletalMeshModel* ImportedModel = InSkeletalMesh->GetImportedModel();
+			const bool bIsValidLOD = ImportedModel && ImportedModel->LODModels.IsValidIndex(LODIndex);
+			if (!bIsValidLOD)
+			{
+				FClothDataflowTools::LogAndToastWarning(*this, LOCTEXT("InvalidLODHeadline", "Invalid LOD."),
+					FText::Format(
+						LOCTEXT("InvalidLODDetails", "No valid LOD {0} found for skeletal mesh {1}."),
+						LODIndex,
+						FText::FromString(InSkeletalMesh->GetName())));
+
+				SetValue(Context, MoveTemp(*ClothCollection), &Collection);
+				return;
+			}
+
+			const FSkeletalMeshLODModel& LODModel = ImportedModel->LODModels[LODIndex];
+			const int32 FirstSection = bImportSingleSection ? SectionIndex : 0;
+			const int32 LastSection = bImportSingleSection ? SectionIndex : LODModel.Sections.Num() - 1;
+			constexpr bool bImportSimMeshNormals = true;
+
+			for (int32 Section = FirstSection; Section <= LastSection; ++Section)
+			{
+				const bool bIsValidSection = LODModel.Sections.IsValidIndex(Section);;
+				if (!bIsValidSection)
+				{
+					FClothDataflowTools::LogAndToastWarning(*this, LOCTEXT("InvalidSectionHeadline", "Invalid section."),
+						FText::Format(
+							LOCTEXT("InvalidSectionDetails", "No valid section {0} found for skeletal mesh {1}."),
+							Section,
+							FText::FromString(InSkeletalMesh->GetName())));
+
+					continue;
+				}
+
+				if (bImportSimMesh)
+				{
+					FClothDataflowTools::AddSimPatternsFromSkeletalMeshSection(ClothCollection, LODModel, Section, UVChannel, UVScale, bImportSimMeshNormals);
+				}
+
+				if (bImportRenderMesh)
+				{
+					const TArray<FSkeletalMaterial>& Materials = InSkeletalMesh->GetMaterials();
+					check(Section < Materials.Num());
+					const FString RenderMaterialPathName = Materials[Section].MaterialInterface ? Materials[Section].MaterialInterface->GetPathName() : "";
+					FClothDataflowTools::AddRenderPatternFromSkeletalMeshSection(ClothCollection, LODModel, Section, RenderMaterialPathName);
+				}
+			}
+
+			if (const UPhysicsAsset* PhysicsAsset = bSetPhysicsAsset ? InSkeletalMesh->GetPhysicsAsset() : nullptr)
+			{
+				ClothFacade.SetPhysicsAssetPathName(PhysicsAsset->GetPathName());
+			}
+
+			ClothFacade.SetSkeletalMeshPathName(InSkeletalMesh->GetPathName());
+		}
+		SetValue(Context, MoveTemp(*ClothCollection), &Collection);
+	}
+}
+
+
+
 FChaosClothAssetSkeletalMeshImportNode::FChaosClothAssetSkeletalMeshImportNode(const Dataflow::FNodeParameters& InParam, FGuid InGuid)
 	: FDataflowNode(InParam, InGuid)
 {
@@ -53,6 +132,7 @@ void FChaosClothAssetSkeletalMeshImportNode::Evaluate(Dataflow::FContext& Contex
 			const FSkeletalMeshLODModel &LODModel = ImportedModel->LODModels[LODIndex];
 			const int32 FirstSection = bImportSingleSection ? SectionIndex : 0;
 			const int32 LastSection = bImportSingleSection ? SectionIndex : LODModel.Sections.Num() - 1;
+			constexpr bool bImportSimMeshNormals = false;
 
 			for (int32 Section = FirstSection; Section <= LastSection; ++Section)
 			{
@@ -70,7 +150,7 @@ void FChaosClothAssetSkeletalMeshImportNode::Evaluate(Dataflow::FContext& Contex
 
 				if (bImportSimMesh)
 				{
-					FClothDataflowTools::AddSimPatternsFromSkeletalMeshSection(ClothCollection, LODModel, Section, UVChannel, UVScale);
+					FClothDataflowTools::AddSimPatternsFromSkeletalMeshSection(ClothCollection, LODModel, Section, UVChannel, UVScale, bImportSimMeshNormals);
 				}
 
 				if (bImportRenderMesh)
@@ -87,6 +167,11 @@ void FChaosClothAssetSkeletalMeshImportNode::Evaluate(Dataflow::FContext& Contex
 				ClothFacade.SetPhysicsAssetPathName(PhysicsAsset->GetPathName());
 			}
 
+			// In order to retain existing behavior, flip the sim normals.
+			constexpr bool bReverseSimMeshNormals = true;
+			constexpr bool bReverseFalse = false;
+			FClothGeometryTools::ReverseMesh(ClothCollection, bReverseSimMeshNormals, bReverseFalse, bReverseFalse, bReverseFalse, TArray<int32>(), TArray<int32>());
+
 			ClothFacade.SetSkeletalMeshPathName(InSkeletalMesh->GetPathName());
 		}
 		SetValue(Context, MoveTemp(*ClothCollection), &Collection);
diff --git a/Engine/Plugins/ChaosClothAssetEditor/Source/ChaosClothAssetDataflowNodes/Private/ChaosClothAsset/StaticMeshImportNode.cpp b/Engine/Plugins/ChaosClothAssetEditor/Source/ChaosClothAssetDataflowNodes/Private/ChaosClothAsset/StaticMeshImportNode.cpp
--- a/Engine/Plugins/ChaosClothAssetEditor/Source/ChaosClothAssetDataflowNodes/Private/ChaosClothAsset/StaticMeshImportNode.cpp
+++ b/Engine/Plugins/ChaosClothAssetEditor/Source/ChaosClothAssetDataflowNodes/Private/ChaosClothAsset/StaticMeshImportNode.cpp
@@ -22,7 +22,8 @@ namespace UE::Chaos::ClothAsset::Private
 		const TSharedRef<FManagedArrayCollection>& ClothCollection,
 		int32 InSectionIndex,
 		int32 UVChannelIndex,
-		const FVector2f& UVScale)
+		const FVector2f& UVScale,
+		bool bImportNormals)
 	{
 		if (InSectionIndex == INDEX_NONE || SkeletalMeshLODModel.Sections.IsValidIndex(InSectionIndex))
 		{
@@ -30,7 +31,7 @@ namespace UE::Chaos::ClothAsset::Private
 			const int32 EndSection = (InSectionIndex != INDEX_NONE) ? InSectionIndex + 1 : SkeletalMeshLODModel.Sections.Num();
 			for (int32 SectionIndex = StartSection; SectionIndex < EndSection; ++SectionIndex)
 			{
-				FClothDataflowTools::AddSimPatternsFromSkeletalMeshSection(ClothCollection, SkeletalMeshLODModel, SectionIndex, UVChannelIndex, UVScale);
+				FClothDataflowTools::AddSimPatternsFromSkeletalMeshSection(ClothCollection, SkeletalMeshLODModel, SectionIndex, UVChannelIndex, UVScale, bImportNormals);
 			}
 			return true;
 		}
@@ -58,6 +59,115 @@ namespace UE::Chaos::ClothAsset::Private
 	}
 } // namespace UE::Chaos::ClothAsset::Private
 
+
+FChaosClothAssetStaticMeshImportNode_v2::FChaosClothAssetStaticMeshImportNode_v2(const Dataflow::FNodeParameters& InParam, FGuid InGuid)
+	: FDataflowNode(InParam, InGuid)
+{
+	RegisterInputConnection(&StaticMesh);
+	RegisterOutputConnection(&Collection);
+}
+
+void FChaosClothAssetStaticMeshImportNode_v2::Evaluate(Dataflow::FContext& Context, const FDataflowOutput* Out) const
+{
+	using namespace UE::Chaos::ClothAsset;
+	using namespace UE::Chaos::ClothAsset::Private;
+
+	if (Out->IsA<FManagedArrayCollection>(&Collection))
+	{
+		// Evaluate out collection
+		const TSharedRef<FManagedArrayCollection> ClothCollection = MakeShared<FManagedArrayCollection>();
+		FCollectionClothFacade ClothFacade(ClothCollection);
+		ClothFacade.DefineSchema();
+
+		const UStaticMesh* const InStaticMesh = GetValue(Context, &StaticMesh);
+
+		if (InStaticMesh && (bImportSimMesh || bImportRenderMesh))
+		{
+			const int32 NumLods = InStaticMesh->GetNumSourceModels();
+			if (LODIndex < NumLods)
+			{
+				const FMeshDescription* const MeshDescription = InStaticMesh->GetMeshDescription(LODIndex);
+				check(MeshDescription);
+
+				constexpr bool bImportSimMeshNormals = true;
+				if (bImportSimMesh && SimMeshSection == INDEX_NONE)
+				{
+					// Keep legacy welding behavior when the entire static mesh is imported
+					FMeshDescriptionToDynamicMesh Converter;
+					Converter.bPrintDebugMessages = false;
+					Converter.bEnableOutputGroups = false;
+					Converter.bVIDsFromNonManifoldMeshDescriptionAttr = true;
+					UE::Geometry::FDynamicMesh3 DynamicMesh;
+					Converter.Convert(MeshDescription, DynamicMesh);
+					constexpr bool bAppend = false;
+					FClothGeometryTools::BuildSimMeshFromDynamicMesh(ClothCollection, DynamicMesh, UVChannel, UVScale, bAppend, bImportSimMeshNormals);
+				}
+				const bool bNeedsSkeletalMeshModel = bImportRenderMesh || (bImportSimMesh && SimMeshSection != INDEX_NONE);
+
+				FSkeletalMeshLODModel SkeletalMeshLODModel;
+				const FMeshBuildSettings& BuildSettings = InStaticMesh->GetSourceModel(LODIndex).BuildSettings;
+				if (bNeedsSkeletalMeshModel &&
+					FClothDataflowTools::BuildSkeletalMeshModelFromMeshDescription(MeshDescription, InStaticMesh->GetSourceModel(LODIndex).BuildSettings, SkeletalMeshLODModel))
+				{
+					if (bImportSimMesh && SimMeshSection != INDEX_NONE)
+					{
+						if (!InitializeSimMeshFromSkeletalMeshModel(SkeletalMeshLODModel, ClothCollection, SimMeshSection, UVChannel, UVScale, bImportSimMeshNormals))
+						{
+							FClothDataflowTools::LogAndToastWarning(*this,
+								LOCTEXT("InvalidRenderMeshHeadline", "Invalid render mesh."),
+								FText::Format(
+									LOCTEXT("InvalidRenderMeshDetails", "The input static mesh {0} failed to convert to a valid render mesh."),
+									FText::FromString(InStaticMesh->GetName())));
+						}
+					}
+					if (bImportRenderMesh)
+					{
+						const FStaticMeshConstAttributes MeshAttributes(*MeshDescription);
+						const TPolygonGroupAttributesConstRef<FName> MaterialSlotNames = MeshAttributes.GetPolygonGroupMaterialSlotNames();
+						auto GetMaterialPathName = [&InStaticMesh, &SkeletalMeshLODModel, &MaterialSlotNames](int32 SectionIndex) -> FString
+						{
+							// Section MaterialIndex refers to the polygon group index. Look up which material this corresponds with.
+							const FName& MaterialSlotName = MaterialSlotNames[SkeletalMeshLODModel.Sections[SectionIndex].MaterialIndex];
+							const int32 MaterialIndex = InStaticMesh->GetMaterialIndexFromImportedMaterialSlotName(MaterialSlotName);
+							const TArray<FStaticMaterial>& StaticMaterials = InStaticMesh->GetStaticMaterials();
+							if (StaticMaterials.IsValidIndex(MaterialIndex) && StaticMaterials[MaterialIndex].MaterialInterface)
+							{
+								return StaticMaterials[MaterialIndex].MaterialInterface->GetPathName();
+							}
+							return FString();
+						};
+
+						if (!InitializeRenderMeshFromSkeletalMeshModel(SkeletalMeshLODModel, ClothCollection, RenderMeshSection, GetMaterialPathName))
+						{
+							FClothDataflowTools::LogAndToastWarning(*this,
+								LOCTEXT("InvalidRenderMeshHeadline", "Invalid render mesh."),
+								FText::Format(
+									LOCTEXT("InvalidRenderMeshDetails", "The input static mesh {0} failed to convert to a valid render mesh."),
+									FText::FromString(InStaticMesh->GetName())));
+						}
+					}
+				}
+				// Bind to root bone by default
+				FClothGeometryTools::BindMeshToRootBone(ClothCollection, bImportSimMesh, bImportRenderMesh);
+			}
+			else
+			{
+				FClothDataflowTools::LogAndToastWarning(*this,
+					LOCTEXT("InvalidLODIndexHeadline", "Invalid LOD index."),
+					FText::Format(LOCTEXT("InvalidLODIndexDetails",
+						"{0} is not a valid LOD index for static mesh {1}.\n"
+						"This static mesh has {2} LOD(s)."),
+						FText::FromString(InStaticMesh->GetName()),
+						LODIndex,
+						NumLods));
+			}
+		}
+
+		SetValue(Context, MoveTemp(*ClothCollection), &Collection);
+	}
+}
+
+
 FChaosClothAssetStaticMeshImportNode::FChaosClothAssetStaticMeshImportNode(const Dataflow::FNodeParameters& InParam, FGuid InGuid)
 	: FDataflowNode(InParam, InGuid)
 {
@@ -87,6 +197,7 @@ void FChaosClothAssetStaticMeshImportNode::Evaluate(Dataflow::FContext& Context,
 				const FMeshDescription* const MeshDescription = InStaticMesh->GetMeshDescription(LODIndex);
 				check(MeshDescription);
 
+				constexpr bool bImportSimMeshNormals = false;
 				if (bImportSimMesh && SimMeshSection == INDEX_NONE)
 				{
 					// Keep legacy welding behavior when the entire static mesh is imported
@@ -97,7 +208,7 @@ void FChaosClothAssetStaticMeshImportNode::Evaluate(Dataflow::FContext& Context,
 					UE::Geometry::FDynamicMesh3 DynamicMesh;
 					Converter.Convert(MeshDescription, DynamicMesh);
 					constexpr bool bAppend = false;
-					FClothGeometryTools::BuildSimMeshFromDynamicMesh(ClothCollection, DynamicMesh, UVChannel, UVScale, bAppend);
+					FClothGeometryTools::BuildSimMeshFromDynamicMesh(ClothCollection, DynamicMesh, UVChannel, UVScale, bAppend, bImportSimMeshNormals);
 				}
 				const bool bNeedsSkeletalMeshModel = bImportRenderMesh || (bImportSimMesh && SimMeshSection != INDEX_NONE);
 
@@ -108,7 +219,7 @@ void FChaosClothAssetStaticMeshImportNode::Evaluate(Dataflow::FContext& Context,
 				{
 					if (bImportSimMesh && SimMeshSection != INDEX_NONE)
 					{
-						if (!InitializeSimMeshFromSkeletalMeshModel(SkeletalMeshLODModel, ClothCollection, SimMeshSection, UVChannel, UVScale))
+						if (!InitializeSimMeshFromSkeletalMeshModel(SkeletalMeshLODModel, ClothCollection, SimMeshSection, UVChannel, UVScale, bImportSimMeshNormals))
 						{
 							FClothDataflowTools::LogAndToastWarning(*this,
 								LOCTEXT("InvalidRenderMeshHeadline", "Invalid render mesh."),
@@ -146,6 +257,12 @@ void FChaosClothAssetStaticMeshImportNode::Evaluate(Dataflow::FContext& Context,
 				}
 				// Bind to root bone by default
 				FClothGeometryTools::BindMeshToRootBone(ClothCollection, bImportSimMesh, bImportRenderMesh);
+				
+				// In order to retain existing behavior, flip the sim normals.
+				constexpr bool bReverseSimMeshNormals = true;
+				constexpr bool bReverseFalse = false;
+				FClothGeometryTools::ReverseMesh(ClothCollection, bReverseSimMeshNormals, bReverseFalse, bReverseFalse, bReverseFalse, TArray<int32>(), TArray<int32>());
+				
 			}
 			else
 			{
diff --git a/Engine/Plugins/ChaosClothAssetEditor/Source/ChaosClothAssetDataflowNodes/Public/ChaosClothAsset/SkeletalMeshImportNode.h b/Engine/Plugins/ChaosClothAssetEditor/Source/ChaosClothAssetDataflowNodes/Public/ChaosClothAsset/SkeletalMeshImportNode.h
--- a/Engine/Plugins/ChaosClothAssetEditor/Source/ChaosClothAssetDataflowNodes/Public/ChaosClothAsset/SkeletalMeshImportNode.h
+++ b/Engine/Plugins/ChaosClothAssetEditor/Source/ChaosClothAssetDataflowNodes/Public/ChaosClothAsset/SkeletalMeshImportNode.h
@@ -7,9 +7,65 @@
 #include "SkeletalMeshImportNode.generated.h"
 
 class USkeletalMesh;
-
 /** Import a skeletal mesh asset into the cloth collection simulation and/or render mesh containers. */
 USTRUCT(Meta = (DataflowCloth))
+struct FChaosClothAssetSkeletalMeshImportNode_v2 : public FDataflowNode
+{
+	GENERATED_USTRUCT_BODY()
+	DATAFLOW_NODE_DEFINE_INTERNAL(FChaosClothAssetSkeletalMeshImportNode_v2, "SkeletalMeshImport", "Cloth", "Cloth Skeletal Mesh Import")
+
+public:
+	UPROPERTY(Meta = (DataflowOutput))
+	FManagedArrayCollection Collection;
+
+	/** The skeletal mesh to import. */
+	UPROPERTY(EditAnywhere, Category = "Skeletal Mesh Import", Meta = (DataflowInput))
+	TObjectPtr<const USkeletalMesh> SkeletalMesh;
+
+	/** The skeletal mesh LOD to import. */
+	UPROPERTY(EditAnywhere, Category = "Skeletal Mesh Import", Meta = (ClampMin = "0", DisplayName = "LOD Index"))
+	int32 LODIndex = 0;
+
+	/** Enable single import section mode. */
+	UPROPERTY(EditAnywhere, Category = "Skeletal Mesh Import")
+	bool bImportSingleSection = false;
+
+	/** The skeletal mesh LOD section to import. If not enabled, then all sections will be imported. */
+	UPROPERTY(EditAnywhere, Category = "Skeletal Mesh Import", Meta = (ClampMin = "0", EditCondition = "bImportSingleSection"))
+	int32 SectionIndex = 0;
+
+	/** Whether to import the simulation mesh from the specified skeletal mesh. */
+	UPROPERTY(EditAnywhere, Category = "Skeletal Mesh Import")
+	bool bImportSimMesh = true;
+
+	/** Whether to import the render mesh from the specified skeletal mesh. */
+	UPROPERTY(EditAnywhere, Category = "Skeletal Mesh Import")
+	bool bImportRenderMesh = true;
+
+	/**
+	 * UV channel of the skeletal mesh to import the 2D simulation mesh patterns from.
+	 * If set to -1, or the specified UVChannel doesn't exist then the import will unwrap the 3D simulation mesh into 2D simulation mesh patterns.
+	 */
+	UPROPERTY(EditAnywhere, Category = "Skeletal Mesh Import", Meta = (ClampMin = "-1", EditCondition = "bImportSimMesh"))
+	int32 UVChannel = 0;
+
+	/* Apply this scale to the UVs when populating Sim Mesh positions. */
+	UPROPERTY(EditAnywhere, Category = "Skeletal Mesh Import", Meta = (AllowPreserveRatio, EditCondition = "bImportSimMesh && UVChannel != INDEX_NONE"))
+	FVector2f UVScale = { 1.f, 1.f };
+
+	/** Set the same physics asset as the one used by the imported skeletal mesh. */
+	UPROPERTY(EditAnywhere, Category = "Skeletal Mesh Import")
+	bool bSetPhysicsAsset = false;
+
+	FChaosClothAssetSkeletalMeshImportNode_v2(const Dataflow::FNodeParameters& InParam, FGuid InGuid = FGuid::NewGuid());
+
+private:
+	virtual void Evaluate(Dataflow::FContext& Context, const FDataflowOutput* Out) const override;
+};
+
+/** Import a skeletal mesh asset into the cloth collection simulation and/or render mesh containers. 
+ * This version re-calculates the Sim Mesh Normals, and they are flipped. */
+USTRUCT(Meta = (DataflowCloth, Deprecated = "5.5"))
 struct FChaosClothAssetSkeletalMeshImportNode : public FDataflowNode
 {
 	GENERATED_USTRUCT_BODY()
diff --git a/Engine/Plugins/ChaosClothAssetEditor/Source/ChaosClothAssetDataflowNodes/Public/ChaosClothAsset/StaticMeshImportNode.h b/Engine/Plugins/ChaosClothAssetEditor/Source/ChaosClothAssetDataflowNodes/Public/ChaosClothAsset/StaticMeshImportNode.h
--- a/Engine/Plugins/ChaosClothAssetEditor/Source/ChaosClothAssetDataflowNodes/Public/ChaosClothAsset/StaticMeshImportNode.h
+++ b/Engine/Plugins/ChaosClothAssetEditor/Source/ChaosClothAssetDataflowNodes/Public/ChaosClothAsset/StaticMeshImportNode.h
@@ -8,8 +8,62 @@
 
 class UStaticMesh;
 
-/** Import a static mesh asset into the cloth collection simulation and/or render mesh containers. */
+/** Import a static mesh asset into the cloth collection simulation and/or render mesh containers.*/
 USTRUCT(Meta = (DataflowCloth))
+struct FChaosClothAssetStaticMeshImportNode_v2 : public FDataflowNode
+{
+	GENERATED_USTRUCT_BODY()
+	DATAFLOW_NODE_DEFINE_INTERNAL(FChaosClothAssetStaticMeshImportNode_v2, "StaticMeshImport", "Cloth", "Cloth Static Mesh Import")
+
+public:
+	UPROPERTY(Meta = (DataflowOutput))
+	FManagedArrayCollection Collection;
+
+	/* The Static Mesh to import from. */
+	UPROPERTY(EditAnywhere, Category = "Static Mesh Import", Meta = (DataflowInput))
+	TObjectPtr<const UStaticMesh> StaticMesh;
+
+	/* Which static mesh Lod to import. */
+	UPROPERTY(EditAnywhere, Category = "Static Mesh Import", Meta = (ClampMin = "0", DisplayName = "LOD Index"))
+	int32 LODIndex = 0;
+
+	/* Import static mesh data as a simulation mesh data. */
+	UPROPERTY(EditAnywhere, Category = "Static Mesh Import")
+	bool bImportSimMesh = true;
+
+	/* Material section to import as sim mesh data. Use -1 to import all sections. */
+	UPROPERTY(EditAnywhere, Category = "Static Mesh Import", Meta = (ClampMin = "-1", EditCondition = "bImportSimMesh"))
+	int32 SimMeshSection = INDEX_NONE;
+
+	/**
+	 * UV channel of the static mesh to import the 2D simulation mesh patterns from.
+	 * If set to -1, or the specified UVChannel doesn't exist then the import will unwrap the 3D simulation mesh into 2D simulation mesh patterns.
+	 */
+	UPROPERTY(EditAnywhere, Category = "Static Mesh Import", Meta = (ClampMin = "-1", EditCondition = "bImportSimMesh"))
+	int32 UVChannel = 0;
+
+	/* Apply this scale to the UVs when populating Sim Mesh positions. */
+	UPROPERTY(EditAnywhere, Category = "Static Mesh Import", Meta = (AllowPreserveRatio, EditCondition = "bImportSimMesh && (UVChannel >= 0)"))
+	FVector2f UVScale = { 1.f, 1.f };
+
+	/* Import static mesh data as render mesh data. */
+	UPROPERTY(EditAnywhere, Category = "Static Mesh Import")
+	bool bImportRenderMesh = true;
+
+	/* Material section to import as render mesh data. Use -1 to import all sections. */
+	UPROPERTY(EditAnywhere, Category = "Static Mesh Import", Meta = (ClampMin = "-1", EditCondition = "bImportRenderMesh"))
+	int32 RenderMeshSection = INDEX_NONE;
+
+	FChaosClothAssetStaticMeshImportNode_v2(const Dataflow::FNodeParameters& InParam, FGuid InGuid = FGuid::NewGuid());
+
+private:
+	virtual void Evaluate(Dataflow::FContext& Context, const FDataflowOutput* Out) const override;
+};
+
+
+/** Import a static mesh asset into the cloth collection simulation and/or render mesh containers. 
+* This version re-calculates the Sim Mesh Normals, and they are flipped. */
+USTRUCT(Meta = (DataflowCloth, Deprecated = "5.5"))
 struct FChaosClothAssetStaticMeshImportNode : public FDataflowNode
 {
 	GENERATED_USTRUCT_BODY()
