Movie Pipeline - Added missing Cloth Component substepping to match the same behavior as the SkeletalMesh clothing system.
Cloth Asset - Fixed missing implementation of p.ChaosCloth.UseTimeStepSmoothing. Made all interactor functions available to other modules.

#tests Succesfully rendered a test sequence with a cloth asset simulation and sub samples that was blowing out prior to this fix.
[RN] Movie Pipeline - Fixes the rendering of simulated Cloth Assets when using temporal samples.

diff --git a/Engine/Plugins/ChaosClothAsset/Source/ChaosClothAssetEngine/Private/ChaosClothAsset/ClothSimulationProxy.cpp b/Engine/Plugins/ChaosClothAsset/Source/ChaosClothAssetEngine/Private/ChaosClothAsset/ClothSimulationProxy.cpp
--- a/Engine/Plugins/ChaosClothAsset/Source/ChaosClothAssetEngine/Private/ChaosClothAsset/ClothSimulationProxy.cpp
+++ b/Engine/Plugins/ChaosClothAsset/Source/ChaosClothAssetEngine/Private/ChaosClothAsset/ClothSimulationProxy.cpp
@@ -280,9 +280,11 @@ namespace UE::Chaos::ClothAsset
 			return;
 		}
 		// Filter delta time to smoothen time variations and prevent unwanted vibrations
+		static IConsoleVariable* const UseTimeStepSmoothingCVar = IConsoleManager::Get().FindConsoleVariable(TEXT("p.ChaosCloth.UseTimeStepSmoothing"));
+		const bool bUseTimeStepSmoothing = UseTimeStepSmoothingCVar ? UseTimeStepSmoothingCVar->GetBool() : true;
 		const Softs::FSolverReal DeltaTime = (Softs::FSolverReal)ClothSimulationContext->DeltaTime;
 		const Softs::FSolverReal PrevDeltaTime = Solver->GetDeltaTime() > 0.f ? Solver->GetDeltaTime() : DeltaTime;
-		const Softs::FSolverReal SmoothedDeltaTime = PrevDeltaTime + (DeltaTime - PrevDeltaTime) * (Softs::FSolverReal)DeltaTimeDecay;
+		const Softs::FSolverReal SmoothedDeltaTime = PrevDeltaTime + (DeltaTime - PrevDeltaTime) * (bUseTimeStepSmoothing ? (Softs::FSolverReal)DeltaTimeDecay : 1.f);
 
 		const double StartTime = FPlatformTime::Seconds();
 		const float PrevSimulationTime = SimulationTime;  // Copy the atomic to prevent a re-read
diff --git a/Engine/Plugins/ChaosClothAsset/Source/ChaosClothAssetEngine/Public/ChaosClothAsset/ClothAssetInteractor.h b/Engine/Plugins/ChaosClothAsset/Source/ChaosClothAssetEngine/Public/ChaosClothAsset/ClothAssetInteractor.h
--- a/Engine/Plugins/ChaosClothAsset/Source/ChaosClothAssetEngine/Public/ChaosClothAsset/ClothAssetInteractor.h
+++ b/Engine/Plugins/ChaosClothAsset/Source/ChaosClothAssetEngine/Public/ChaosClothAsset/ClothAssetInteractor.h
@@ -25,63 +25,63 @@ public:
 
 	/** Generate a list of all properties held by this interactor. Properties for all LODs will be returned if LODIndex = -1.*/
 	UFUNCTION(BlueprintCallable, Category = ClothProperty)
-	TArray<FString> GetAllProperties(int32 LODIndex = -1) const;
+	CHAOSCLOTHASSETENGINE_API TArray<FString> GetAllProperties(int32 LODIndex = -1) const;
 
 	/** Get the value for a property cast to float. DefaultValue will be returned if the property is not found.*/
 	UFUNCTION(BlueprintCallable, Category = ClothProperty)
-	float GetFloatValue(const FString& PropertyName, int32 LODIndex = 0, float DefaultValue = 0.f) const;
+	CHAOSCLOTHASSETENGINE_API float GetFloatValue(const FString& PropertyName, int32 LODIndex = 0, float DefaultValue = 0.f) const;
 
 	/** Get the low value for a weighted property value (same as GetFloatValue). DefaultValue will be returned if the property is not found.*/
 	UFUNCTION(BlueprintCallable, Category = ClothProperty)
-	float GetLowFloatValue(const FString& PropertyName, int32 LODIndex = 0, float DefaultValue = 0.f) const;
+	CHAOSCLOTHASSETENGINE_API float GetLowFloatValue(const FString& PropertyName, int32 LODIndex = 0, float DefaultValue = 0.f) const;
 
 	/** Get the high value for a weighted property value. DefaultValue will be returned if the property is not found.*/
 	UFUNCTION(BlueprintCallable, Category = ClothProperty)
-	float GetHighFloatValue(const FString& PropertyName, int32 LODIndex = 0, float DefaultValue = 0.f) const;
+	CHAOSCLOTHASSETENGINE_API float GetHighFloatValue(const FString& PropertyName, int32 LODIndex = 0, float DefaultValue = 0.f) const;
 
 	/** Get the low and high values for a weighted property value. DefaultValue will be returned if the property is not found.*/
 	UFUNCTION(BlueprintCallable, Category = ClothProperty)
-	FVector2D GetWeightedFloatValue(const FString& PropertyName, int32 LODIndex = 0, FVector2D DefaultValue = FVector2D(0., 0.)) const;
+	CHAOSCLOTHASSETENGINE_API FVector2D GetWeightedFloatValue(const FString& PropertyName, int32 LODIndex = 0, FVector2D DefaultValue = FVector2D(0., 0.)) const;
 
 	/** Get the value for a property cast to int. DefaultValue will be returned if the property is not found.*/
 	UFUNCTION(BlueprintCallable, Category = ClothProperty)
-	int32 GetIntValue(const FString& PropertyName, int32 LODIndex = 0, int32 DefaultValue = 0) const;
+	CHAOSCLOTHASSETENGINE_API int32 GetIntValue(const FString& PropertyName, int32 LODIndex = 0, int32 DefaultValue = 0) const;
 
 	/** Get the value for a property cast to vector. DefaultValue will be returned if the property is not found.*/
 	UFUNCTION(BlueprintCallable, Category = ClothProperty)
-	FVector GetVectorValue(const FString& PropertyName, int32 LODIndex = 0, FVector DefaultValue = FVector(0.)) const;
+	CHAOSCLOTHASSETENGINE_API FVector GetVectorValue(const FString& PropertyName, int32 LODIndex = 0, FVector DefaultValue = FVector(0.)) const;
 
 	/** Get the string value for a property (typically the associated map name for weighted values). DefaultValue will be returned if the property is not found.*/
 	UFUNCTION(BlueprintCallable, Category = ClothProperty)
-	FString GetStringValue(const FString& PropertyName, int32 LODIndex = 0, const FString& DefaultValue = "") const;
+	CHAOSCLOTHASSETENGINE_API FString GetStringValue(const FString& PropertyName, int32 LODIndex = 0, const FString& DefaultValue = "") const;
 
 	/**Set the value for a property (if it exists). This sets the Low and High values for weighted values. All LODs will be set when LODIndex = -1.*/
 	UFUNCTION(BlueprintCallable, Category = ClothProperty)
-	void SetFloatValue(const FString& PropertyName, int32 LODIndex = -1, float Value = 0.f);
+	CHAOSCLOTHASSETENGINE_API void SetFloatValue(const FString& PropertyName, int32 LODIndex = -1, float Value = 0.f);
 
 	/**Set the low value for a weighted property (if it exists). All LODs will be set when LODIndex = -1.*/
 	UFUNCTION(BlueprintCallable, Category = ClothProperty)
-	void SetLowFloatValue(const FString& PropertyName, int32 LODIndex = -1, float Value = 0.f);
+	CHAOSCLOTHASSETENGINE_API void SetLowFloatValue(const FString& PropertyName, int32 LODIndex = -1, float Value = 0.f);
 
 	/**Set the high value for a weighted property (if it exists). All LODs will be set when LODIndex = -1.*/
 	UFUNCTION(BlueprintCallable, Category = ClothProperty)
-	void SetHighFloatValue(const FString& PropertyName, int32 LODIndex = -1, float Value = 0.f);
+	CHAOSCLOTHASSETENGINE_API void SetHighFloatValue(const FString& PropertyName, int32 LODIndex = -1, float Value = 0.f);
 
 	/**Set the low and high values for a weighted property (if it exists). All LODs will be set when LODIndex = -1.*/
 	UFUNCTION(BlueprintCallable, Category = ClothProperty)
-	void SetWeightedFloatValue(const FString& PropertyName, int32 LODIndex = -1, FVector2D Value = FVector2D(0., 0.));
+	CHAOSCLOTHASSETENGINE_API void SetWeightedFloatValue(const FString& PropertyName, int32 LODIndex = -1, FVector2D Value = FVector2D(0., 0.));
 
 	/**Set the value for a property (if it exists). All LODs will be set when LODIndex = -1.*/
 	UFUNCTION(BlueprintCallable, Category = ClothProperty)
-	void SetIntValue(const FString& PropertyName, int32 LODIndex = -1, int32 Value = 0);
+	CHAOSCLOTHASSETENGINE_API void SetIntValue(const FString& PropertyName, int32 LODIndex = -1, int32 Value = 0);
 
 	/**Set the value for a property (if it exists). All LODs will be set when LODIndex = -1.*/
 	UFUNCTION(BlueprintCallable, Category = ClothProperty)
-	void SetVectorValue(const FString& PropertyName, int32 LODIndex = -1, FVector Value = FVector(0.));
+	CHAOSCLOTHASSETENGINE_API void SetVectorValue(const FString& PropertyName, int32 LODIndex = -1, FVector Value = FVector(0.));
 
 	/**Set the string value for a property (if it exists). This is typically the map name associated with a property. All LODs will be set when LODIndex = -1.*/
 	UFUNCTION(BlueprintCallable, Category = ClothProperty)
-	void SetStringValue(const FString& PropertyName, int32 LODIndex = -1, const FString& Value = "");
+	CHAOSCLOTHASSETENGINE_API void SetStringValue(const FString& PropertyName, int32 LODIndex = -1, const FString& Value = "");
 
 private:
 	TArray<TWeakPtr<::Chaos::Softs::FCollectionPropertyFacade>> CollectionPropertyFacades;
diff --git a/Engine/Plugins/MovieScene/MovieRenderPipeline/MovieRenderPipeline.uplugin b/Engine/Plugins/MovieScene/MovieRenderPipeline/MovieRenderPipeline.uplugin
--- a/Engine/Plugins/MovieScene/MovieRenderPipeline/MovieRenderPipeline.uplugin
+++ b/Engine/Plugins/MovieScene/MovieRenderPipeline/MovieRenderPipeline.uplugin
@@ -42,6 +42,10 @@
 		{
 			"Name": "SequencerScripting",
 			"Enabled": true
+		},
+		{
+			"Name": "ChaosClothAsset",
+			"Enabled": true
 		}
 	],
 	"Modules": [
diff --git a/Engine/Plugins/MovieScene/MovieRenderPipeline/Source/MovieRenderPipelineCore/MovieRenderPipelineCore.Build.cs b/Engine/Plugins/MovieScene/MovieRenderPipeline/Source/MovieRenderPipelineCore/MovieRenderPipelineCore.Build.cs
--- a/Engine/Plugins/MovieScene/MovieRenderPipeline/Source/MovieRenderPipelineCore/MovieRenderPipelineCore.Build.cs
+++ b/Engine/Plugins/MovieScene/MovieRenderPipeline/Source/MovieRenderPipelineCore/MovieRenderPipelineCore.Build.cs
@@ -25,6 +25,7 @@ public class MovieRenderPipelineCore : ModuleRules
 				"Slate",
 				"OpenColorIOWrapper",
 				"ImageCore",
+				"ChaosClothAssetEngine",
 			}
 		);
 
diff --git a/Engine/Plugins/MovieScene/MovieRenderPipeline/Source/MovieRenderPipelineCore/Private/MoviePipeline.cpp b/Engine/Plugins/MovieScene/MovieRenderPipeline/Source/MovieRenderPipelineCore/Private/MoviePipeline.cpp
--- a/Engine/Plugins/MovieScene/MovieRenderPipeline/Source/MovieRenderPipelineCore/Private/MoviePipeline.cpp
+++ b/Engine/Plugins/MovieScene/MovieRenderPipeline/Source/MovieRenderPipelineCore/Private/MoviePipeline.cpp
@@ -48,6 +48,8 @@
 #include "EngineUtils.h"
 #include "ClothingSimulationInteractor.h"
 #include "ClothingSimulationInterface.h"
+#include "ChaosClothAsset/ClothComponent.h"
+#include "ChaosClothAsset/ClothAssetInteractor.h"
 #include "Components/SkeletalMeshComponent.h"
 #include "MovieSceneCommonHelpers.h"
 
@@ -1716,19 +1718,59 @@ void UMoviePipeline::SetSkeletalMeshClothSubSteps(const int32 InSubdivisionCount
 				UClothingSimulationInteractor* ClothInteractor = Component->GetClothingSimulationInteractor();
 				if (ClothInteractor)
 				{
-					TWeakObjectPtr<UClothingSimulationInteractor> WeakPtr = TWeakObjectPtr< UClothingSimulationInteractor>(ClothInteractor);
+					const int32 LODIndex = 0;  // There is only a NumSubSteps for LOD 0 in the Skeletal Mesh clothing system
+					TWeakObjectPtr<UObject> WeakPtr = TWeakObjectPtr<UClothingSimulationInteractor>(ClothInteractor);
 		
-					FClothSimSettingsCache* ExistingCacheEntry = ClothSimCache.Find(WeakPtr);
+					TArray<FClothSimSettingsCache>* ExistingCacheEntry = ClothSimCache.Find(WeakPtr);
 					if (!ExistingCacheEntry)
 					{
 						ClothSimCache.Add(WeakPtr);
 						ExistingCacheEntry = ClothSimCache.Find(WeakPtr);
+						ExistingCacheEntry->SetNumZeroed(1);  // Only store LOD 0
 						const int32 NumSubsteps = Component->GetClothingSimulation() ? FMath::Max(Component->GetClothingSimulation()->GetNumSubsteps(), 1)
 							: 1; // If there's no clothing simulation component just fall back to assuming they only had 1.
-						ExistingCacheEntry->NumSubSteps = NumSubsteps;
+						(*ExistingCacheEntry)[LODIndex].NumSubSteps = NumSubsteps;
 					}
 
-					ClothInteractor->SetNumSubsteps(ExistingCacheEntry->NumSubSteps * InSubdivisionCount);
+					ClothInteractor->SetNumSubsteps((*ExistingCacheEntry)[LODIndex].NumSubSteps * InSubdivisionCount);
+				}
+			}
+
+			TArray<UChaosClothComponent*> ChaosClothComponents;
+			FoundActor->GetComponents(ChaosClothComponents);
+
+			for (UChaosClothComponent* Component : ChaosClothComponents)
+			{
+				UChaosClothAssetInteractor* ClothAssetInteractor = Component->GetClothOutfitInteractor();
+				if (ClothAssetInteractor)
+				{
+					const int32 NumLODs = Component->GetNumLODs();
+					TWeakObjectPtr<UObject> WeakPtr = TWeakObjectPtr<UChaosClothAssetInteractor>(ClothAssetInteractor);
+		
+					TArray<FClothSimSettingsCache>* ExistingCacheEntry = ClothSimCache.Find(WeakPtr);
+					if (!ExistingCacheEntry)
+					{
+						ClothSimCache.Add(WeakPtr);
+						ExistingCacheEntry = ClothSimCache.Find(WeakPtr);
+						ExistingCacheEntry->SetNumUninitialized(NumLODs);
+						for (int32 LODIndex = 0; LODIndex < NumLODs; ++LODIndex)
+						{
+							constexpr int32 MinNumSubsteps = 1;
+							const int32 NumSubsteps = FMath::Max(ClothAssetInteractor->GetIntValue(TEXT("NumSubsteps"), LODIndex, MinNumSubsteps), MinNumSubsteps);
+							const int32 DynamicSubstepDeltaTime = ClothAssetInteractor->GetFloatValue(TEXT("DynamicSubstepDeltaTime"), LODIndex, MinNumSubsteps);
+							(*ExistingCacheEntry)[LODIndex].NumSubSteps = NumSubsteps;
+							(*ExistingCacheEntry)[LODIndex].DynamicSubstepDeltaTime = DynamicSubstepDeltaTime;
+						}
+					}
+
+					for (int32 LODIndex = 0; LODIndex < NumLODs; ++LODIndex)
+					{
+						if (ExistingCacheEntry->IsValidIndex(LODIndex))
+						{
+							ClothAssetInteractor->SetIntValue(TEXT("NumSubsteps"), LODIndex, (*ExistingCacheEntry)[LODIndex].NumSubSteps * InSubdivisionCount);
+							ClothAssetInteractor->SetFloatValue(TEXT("DynamicSubstepDeltaTime"), LODIndex, 0.f);
+						}
+					}
 				}
 			}
 		}
@@ -1737,11 +1779,23 @@ void UMoviePipeline::SetSkeletalMeshClothSubSteps(const int32 InSubdivisionCount
 
 void UMoviePipeline::RestoreSkeletalMeshClothSubSteps()
 {
-	for (const TPair<TWeakObjectPtr<UClothingSimulationInteractor>, FClothSimSettingsCache>& Pair : ClothSimCache)
+	for (const TPair<TWeakObjectPtr<UObject>, TArray<FClothSimSettingsCache>>& Pair : ClothSimCache)
 	{
-		if (Pair.Key.Get())
+		if (UObject* Object = Pair.Key.Get())
 		{
-			Pair.Key->SetNumSubsteps(Pair.Value.NumSubSteps);
+			if (UClothingSimulationInteractor* ClothInteractor = Cast<UClothingSimulationInteractor>(Object))
+			{
+				const int32 LODIndex = 0;  // There is only a NumSubSteps for LOD 0 in the Skeletal Mesh clothing system
+				ClothInteractor->SetNumSubsteps(Pair.Value[LODIndex].NumSubSteps);
+			}
+			else if (UChaosClothAssetInteractor* ClothAssetInteractor = Cast<UChaosClothAssetInteractor>(Object))
+			{
+				for (int32 LODIndex = 0; LODIndex < Pair.Value.Num(); ++LODIndex)
+				{
+					ClothAssetInteractor->SetIntValue(TEXT("NumSubsteps"), LODIndex, Pair.Value[LODIndex].NumSubSteps);
+					ClothAssetInteractor->SetIntValue(TEXT("DynamicSubstepDeltaTime"), LODIndex, Pair.Value[LODIndex].DynamicSubstepDeltaTime);
+				}
+			}
 		}
 	}
 }
diff --git a/Engine/Plugins/MovieScene/MovieRenderPipeline/Source/MovieRenderPipelineCore/Public/MoviePipeline.h b/Engine/Plugins/MovieScene/MovieRenderPipeline/Source/MovieRenderPipelineCore/Public/MoviePipeline.h
--- a/Engine/Plugins/MovieScene/MovieRenderPipeline/Source/MovieRenderPipelineCore/Public/MoviePipeline.h
+++ b/Engine/Plugins/MovieScene/MovieRenderPipeline/Source/MovieRenderPipelineCore/Public/MoviePipeline.h
@@ -418,9 +418,11 @@ private:
 	struct FClothSimSettingsCache
 	{
 		int32 NumSubSteps;
+		float DynamicSubstepDeltaTime;
 	};
 
-	TMap<TWeakObjectPtr<class UClothingSimulationInteractor>, FClothSimSettingsCache> ClothSimCache;
+	/** Simulation settings cache per cloth interactor object. Needs one per LOD, hence the array. */
+	TMap<TWeakObjectPtr<UObject>, TArray<FClothSimSettingsCache>> ClothSimCache;
 
 	struct FRenderTimeStatistics
 	{
