CIS - Fixed issue Issue 548872: warning C4996: 'Chaos::FClothVisualization': Use FClothVisualizationNoGC instead

#rnx

diff --git a/Engine/Plugins/ChaosClothAsset/Source/ChaosClothAssetEngine/Public/ChaosClothAsset/ClothSimulationProxy.h b/Engine/Plugins/ChaosClothAsset/Source/ChaosClothAssetEngine/Public/ChaosClothAsset/ClothSimulationProxy.h
--- a/Engine/Plugins/ChaosClothAsset/Source/ChaosClothAssetEngine/Public/ChaosClothAsset/ClothSimulationProxy.h
+++ b/Engine/Plugins/ChaosClothAsset/Source/ChaosClothAssetEngine/Public/ChaosClothAsset/ClothSimulationProxy.h
@@ -69,7 +69,9 @@ namespace UE::Chaos::ClothAsset
 		FBoxSphereBounds CalculateBounds_AnyThread() const;
 
 		UE_DEPRECATED(5.5, "Use GetClothVisualization instead.")
+PRAGMA_DISABLE_DEPRECATION_WARNINGS
 		const ::Chaos::FClothVisualization* GetVisualization() const { return Visualization.Get(); }
+PRAGMA_ENABLE_DEPRECATION_WARNINGS
 		const ::Chaos::FClothVisualizationNoGC* GetClothVisualization() const;
 
 		int32 GetNumCloths() const { return NumCloths; }
