Updating Display Name & Tool Tip for UChaosClothComponent class for naming uniformity with other similar classes

#rnx

diff --git a/Engine/Plugins/ChaosClothAsset/Source/ChaosClothAssetEngine/Public/ChaosClothAsset/ClothComponent.h b/Engine/Plugins/ChaosClothAsset/Source/ChaosClothAssetEngine/Public/ChaosClothAsset/ClothComponent.h
--- a/Engine/Plugins/ChaosClothAsset/Source/ChaosClothAssetEngine/Public/ChaosClothAsset/ClothComponent.h
+++ b/Engine/Plugins/ChaosClothAsset/Source/ChaosClothAssetEngine/Public/ChaosClothAsset/ClothComponent.h
@@ -34,8 +34,8 @@ namespace UE::Chaos::ClothAsset
  */
 UCLASS(
 	ClassGroup = Physics, 
-	Meta = (BlueprintSpawnableComponent, ToolTip = "Chaos cloth component."),
-	DisplayName = "Chaos cloth component",
+	Meta = (BlueprintSpawnableComponent, ToolTip = "Chaos Cloth Component"),
+	DisplayName = "Chaos Cloth Component",
 	HideCategories = (Object, "Mesh|SkeletalAsset", Constraints, Advanced, Cooking, Collision, Navigation))
 class CHAOSCLOTHASSETENGINE_API UChaosClothComponent : public USkinnedMeshComponent, public IDataflowPhysicsSolverInterface
 {
