Chaos Cloth Asset: add "Copy Interactor Name" right click command for simulation properties that can be set via the Cloth Asset Interactor.

Also fix some properties to be marked as non-animatable/intrinsic.

diff --git a/Engine/Plugins/ChaosClothAsset/Source/ChaosClothAssetEngine/Private/ChaosClothAsset/ClothAssetInteractor.cpp b/Engine/Plugins/ChaosClothAsset/Source/ChaosClothAssetEngine/Private/ChaosClothAsset/ClothAssetInteractor.cpp
--- a/Engine/Plugins/ChaosClothAsset/Source/ChaosClothAssetEngine/Private/ChaosClothAsset/ClothAssetInteractor.cpp
+++ b/Engine/Plugins/ChaosClothAsset/Source/ChaosClothAssetEngine/Private/ChaosClothAsset/ClothAssetInteractor.cpp
@@ -5,6 +5,172 @@
 
 #include UE_INLINE_GENERATED_CPP_BY_NAME(ClothAssetInteractor)
 
+namespace UE::Chaos::ClothAsset::Private
+{
+	// Put aliases in the same order as how the if/else works in FClothConstraints::Create[Type]Constraints to ensure
+	// the property that is actually used by the solver is used here.
+
+	// SimulationBendingConfigNode
+	static const TArray<FString> BendingStiffnessWarpAliases =
+	{
+		TEXT("XPBDAnisoBendingStiffnessWarp"),
+	};
+	static const TArray<FString> BendingStiffnessWeftAliases =
+	{
+		TEXT("XPBDAnisoBendingStiffnessWeft"),
+	};
+	static const TArray<FString> BendingStiffnessBiasAliases =
+	{
+		TEXT("XPBDAnisoBendingStiffnessBias"),
+	};
+	static const TArray<FString> BendingDampingAliases =
+	{
+		TEXT("XPBDAnisoBendingDamping"),
+		TEXT("XPBDBendingElementDamping"),
+		TEXT("XPBDBendingSpringDamping"),
+	};
+	static const TArray<FString> BucklingRatioAliases =
+	{
+		TEXT("XPBDAnisoBucklingRatio"),
+		TEXT("XPBDBucklingRatio"),
+		TEXT("BucklingRatio"),
+	};
+	static const TArray<FString> BucklingStiffnessWarpAliases =
+	{
+		TEXT("XPBDAnisoBucklingStiffnessWarp"),
+	};
+	static const TArray<FString> BucklingStiffnessWeftAliases =
+	{
+		TEXT("XPBDAnisoBucklingStiffnessWeft"),
+	};
+	static const TArray<FString> BucklingStiffnessBiasAliases =
+	{
+		TEXT("XPBDAnisoBucklingStiffnessBias"),
+	};
+	static const TArray<FString> BendingStiffnessAliases =
+	{
+		TEXT("XPBDBendingElementStiffness"),
+		TEXT("BendingElementStiffness"),
+		TEXT("XPBDBendingSpringStiffness"),
+		TEXT("BendingSpringStiffness"),
+	};
+	static const TArray<FString> BucklingStiffnessAliases =
+	{
+		TEXT("XPBDBucklingStiffness"),
+		TEXT("BucklingStiffness"),
+	};
+
+	// SimulationStretchConfigNode
+	static const TArray<FString> StretchStiffnessWarpAliases =
+	{
+		TEXT("XPBDAnisoStretchStiffnessWarp"),
+		TEXT("XPBDAnisoSpringStiffnessWarp"),
+	};
+	static const TArray<FString> StretchStiffnessWeftAliases =
+	{
+		TEXT("XPBDAnisoStretchStiffnessWeft"),
+		TEXT("XPBDAnisoSpringStiffnessWeft"),
+	};
+	static const TArray<FString> StretchStiffnessBiasAliases =
+	{
+		TEXT("XPBDAnisoStretchStiffnessBias"),
+		TEXT("XPBDAnisoSpringStiffnessBias"),
+	};
+	static const TArray<FString> StretchDampingAliases =
+	{
+		TEXT("XPBDAnisoStretchDamping"),
+		TEXT("XPBDEdgeSpringDamping"),
+		TEXT("XPBDAnisoSpringDamping"),
+	};
+	static const TArray<FString> StretchStiffnessAliases =
+	{
+		TEXT("XPBDEdgeSpringStiffness"),
+		TEXT("EdgeSpringStiffness"),
+	};
+	static const TArray<FString> StretchWarpScaleAliases =
+	{
+		TEXT("XPBDAnisoStretchWarpScale"),
+		TEXT("XPBDAnisoSpringWarpScale"),
+	};
+	static const TArray<FString> StretchWeftScaleAliases =
+	{
+		TEXT("XPBDAnisoStretchWeftScale"),
+		TEXT("XPBDAnisoSpringWeftScale"),
+	};
+	static const TArray<FString> AreaStiffnessAliases =
+	{
+		TEXT("XPBDAreaSpringStiffness"),
+		TEXT("AreaSpringStiffness"),
+	};
+	
+
+	static const TMap<FString, TArray<FString>> Aliases =
+	{
+		{TEXT("BendingStiffnessWarp"), BendingStiffnessWarpAliases},
+		{TEXT("BendingStiffnessWeft"), BendingStiffnessWeftAliases},
+		{TEXT("BendingStiffnessBias"), BendingStiffnessBiasAliases},
+		{TEXT("BendingDamping"), BendingDampingAliases},
+		{TEXT("BucklingRatio"), BucklingRatioAliases},
+		{TEXT("BucklingStiffnessWarp"), BucklingStiffnessWarpAliases},
+		{TEXT("BucklingStiffnessWeft"), BucklingStiffnessWeftAliases},
+		{TEXT("BucklingStiffnessBias"), BucklingStiffnessBiasAliases},
+		{TEXT("BendingStiffness"), BendingStiffnessAliases},
+		{TEXT("BucklingStiffness"), BucklingStiffnessAliases},
+		{TEXT("StretchStiffnessWarp"), StretchStiffnessWarpAliases},
+		{TEXT("StretchStiffnessWeft"), StretchStiffnessWeftAliases},
+		{TEXT("StretchStiffnessBias"), StretchStiffnessBiasAliases},
+		{TEXT("StretchDamping"), StretchDampingAliases},
+		{TEXT("StretchStiffness"), StretchStiffnessAliases},
+		{TEXT("StretchWarpScale"), StretchWarpScaleAliases},
+		{TEXT("StretchWeftScale"), StretchWeftScaleAliases},
+		{TEXT("AreaStiffness"), AreaStiffnessAliases},
+	};
+
+
+	template<typename T>
+	static T GetValueWithAlias(TSharedPtr<::Chaos::Softs::FCollectionPropertyFacade>& PropertyFacade, const FString& PropertyName, const T& DefaultValue, const TFunctionRef<T(int32 KeyIndex)>& GetValue)
+	{
+		check(PropertyFacade);
+		{
+			const int32 KeyIndex = PropertyFacade->GetKeyIndex(PropertyName);
+			if (KeyIndex != INDEX_NONE)
+			{
+				return GetValue(KeyIndex);
+			}
+		}
+
+		if (const TArray<FString>* FoundAliases = Aliases.Find(PropertyName))
+		{
+			for (const FString& FoundAlias : *FoundAliases)
+			{
+				const int32 KeyIndex = PropertyFacade->GetKeyIndex(FoundAlias);
+				if (KeyIndex != INDEX_NONE)
+				{
+					return GetValue(KeyIndex);
+				}
+			}
+		}
+
+		return DefaultValue;
+	}
+
+	void SetValueWithAlias(TSharedPtr<::Chaos::Softs::FCollectionPropertyFacade>& PropertyFacade, const FString& PropertyName, const TFunctionRef<void(const FString&)>& SetValue)
+	{
+		check(PropertyFacade);
+		if (const TArray<FString>* FoundAliases = Aliases.Find(PropertyName))
+		{
+			for (const FString& FoundAlias : *FoundAliases)
+			{
+				SetValue(FoundAlias);
+			}
+		}
+		else
+		{
+			SetValue(PropertyName);
+		}
+	}
+}
+
 void UChaosClothAssetInteractor::SetProperties(const TArray<TSharedPtr<::Chaos::Softs::FCollectionPropertyFacade>>& InCollectionPropertyFacades)
 {
 	CollectionPropertyFacades.Reset(InCollectionPropertyFacades.Num());
@@ -21,7 +187,9 @@ void UChaosClothAssetInteractor::ResetProperties()
 
 TArray<FString> UChaosClothAssetInteractor::GetAllProperties(int32 LODIndex) const
 {
-	TArray<FString> Keys;
+	using namespace UE::Chaos::ClothAsset::Private;
+
+	TSet<FString> Keys;
 	if (LODIndex == INDEX_NONE)
 	{
 		for (int32 Index = 0; Index < CollectionPropertyFacades.Num(); ++Index)
@@ -31,7 +199,7 @@ TArray<FString> UChaosClothAssetInteractor::GetAllProperties(int32 LODIndex) con
 				if (Keys.IsEmpty())
 				{
 					// This is the first non-empty LOD. We can add all keys without worrying about uniqueness.
-					Keys.Reserve(PropertyFacade->Num());
+					Keys.Reserve(PropertyFacade->Num() + Aliases.Num());
 					for (int32 KeyIndex = 0; KeyIndex < PropertyFacade->Num(); ++KeyIndex)
 					{
 						Keys.Add(PropertyFacade->GetKey(KeyIndex));
@@ -41,7 +209,7 @@ TArray<FString> UChaosClothAssetInteractor::GetAllProperties(int32 LODIndex) con
 				{
 					for (int32 KeyIndex = 0; KeyIndex < PropertyFacade->Num(); ++KeyIndex)
 					{
-						Keys.AddUnique(PropertyFacade->GetKey(KeyIndex));
+						Keys.Add(PropertyFacade->GetKey(KeyIndex));
 					}
 				}
 			}
@@ -61,16 +229,32 @@ TArray<FString> UChaosClothAssetInteractor::GetAllProperties(int32 LODIndex) con
 			}
 		}
 	}
-	return Keys;
+
+	for (TMap<FString, TArray<FString>>::TConstIterator AliasIter = Aliases.CreateConstIterator(); AliasIter; ++AliasIter)
+	{
+		for (const FString& OtherName : AliasIter.Value())
+		{
+			if (Keys.Contains(OtherName))
+			{
+				Keys.Add(AliasIter.Key());
+			}
+		}
+	}
+	return Keys.Array();
 }
 
 float UChaosClothAssetInteractor::GetFloatValue(const FString& PropertyName, int32 LODIndex, float DefaultValue) const
 {
+	using namespace UE::Chaos::ClothAsset::Private;
 	if (CollectionPropertyFacades.IsValidIndex(LODIndex))
 	{
 		if (TSharedPtr<::Chaos::Softs::FCollectionPropertyFacade> PropertyFacade = CollectionPropertyFacades[LODIndex].Pin())
 		{
-			return PropertyFacade->GetValue<float>(PropertyName, DefaultValue);
+			return GetValueWithAlias<float>(PropertyFacade, PropertyName, DefaultValue,
+				[&PropertyFacade](int32 KeyIndex)
+				{
+					return PropertyFacade->GetValue<float>(KeyIndex);
+				});
 		}
 	}
 	return DefaultValue;
@@ -78,11 +262,16 @@ float UChaosClothAssetInteractor::GetFloatValue(const FString& PropertyName, int
 
 float UChaosClothAssetInteractor::GetLowFloatValue(const FString& PropertyName, int32 LODIndex, float DefaultValue) const
 {
+	using namespace UE::Chaos::ClothAsset::Private;
 	if (CollectionPropertyFacades.IsValidIndex(LODIndex))
 	{
 		if (TSharedPtr<::Chaos::Softs::FCollectionPropertyFacade> PropertyFacade = CollectionPropertyFacades[LODIndex].Pin())
 		{
-			return PropertyFacade->GetLowValue<float>(PropertyName, DefaultValue);
+			return GetValueWithAlias<float>(PropertyFacade, PropertyName, DefaultValue,
+				[&PropertyFacade](int32 KeyIndex)
+				{
+					return PropertyFacade->GetLowValue<float>(KeyIndex);
+				});
 		}
 	}
 	return DefaultValue;
@@ -90,11 +279,16 @@ float UChaosClothAssetInteractor::GetLowFloatValue(const FString& PropertyName,
 
 float UChaosClothAssetInteractor::GetHighFloatValue(const FString& PropertyName, int32 LODIndex, float DefaultValue) const
 {
+	using namespace UE::Chaos::ClothAsset::Private;
 	if (CollectionPropertyFacades.IsValidIndex(LODIndex))
 	{
 		if (TSharedPtr<::Chaos::Softs::FCollectionPropertyFacade> PropertyFacade = CollectionPropertyFacades[LODIndex].Pin())
 		{
-			return PropertyFacade->GetHighValue<float>(PropertyName, DefaultValue);
+			return GetValueWithAlias<float>(PropertyFacade, PropertyName, DefaultValue,
+				[&PropertyFacade](int32 KeyIndex)
+				{
+					return PropertyFacade->GetHighValue<float>(KeyIndex);
+				});
 		}
 	}
 	return DefaultValue;
@@ -102,11 +296,16 @@ float UChaosClothAssetInteractor::GetHighFloatValue(const FString& PropertyName,
 
 FVector2D UChaosClothAssetInteractor::GetWeightedFloatValue(const FString& PropertyName, int32 LODIndex, FVector2D DefaultValue) const
 {
+	using namespace UE::Chaos::ClothAsset::Private;
 	if (CollectionPropertyFacades.IsValidIndex(LODIndex))
 	{
 		if (TSharedPtr<::Chaos::Softs::FCollectionPropertyFacade> PropertyFacade = CollectionPropertyFacades[LODIndex].Pin())
 		{
-			return FVector2D(PropertyFacade->GetWeightedFloatValue(PropertyName, FVector2f(DefaultValue)));
+			return GetValueWithAlias<FVector2D>(PropertyFacade, PropertyName, DefaultValue,
+				[&PropertyFacade](int32 KeyIndex)
+				{
+					return FVector2D(PropertyFacade->GetWeightedFloatValue(KeyIndex));
+				});
 		}
 	}
 	return DefaultValue;
@@ -114,11 +313,16 @@ FVector2D UChaosClothAssetInteractor::GetWeightedFloatValue(const FString& Prope
 
 int32 UChaosClothAssetInteractor::GetIntValue(const FString& PropertyName, int32 LODIndex, int32 DefaultValue) const
 {
+	using namespace UE::Chaos::ClothAsset::Private;
 	if (CollectionPropertyFacades.IsValidIndex(LODIndex))
 	{
 		if (TSharedPtr<::Chaos::Softs::FCollectionPropertyFacade> PropertyFacade = CollectionPropertyFacades[LODIndex].Pin())
 		{
-			return PropertyFacade->GetValue<int32>(PropertyName, DefaultValue);
+			return GetValueWithAlias<int32>(PropertyFacade, PropertyName, DefaultValue,
+				[&PropertyFacade](int32 KeyIndex)
+				{
+					return PropertyFacade->GetValue<int32>(KeyIndex);
+				});
 		}
 	}
 	return DefaultValue;
@@ -126,11 +330,16 @@ int32 UChaosClothAssetInteractor::GetIntValue(const FString& PropertyName, int32
 
 FVector UChaosClothAssetInteractor::GetVectorValue(const FString& PropertyName, int32 LODIndex, FVector DefaultValue) const
 {
+	using namespace UE::Chaos::ClothAsset::Private;
 	if (CollectionPropertyFacades.IsValidIndex(LODIndex))
 	{
 		if (TSharedPtr<::Chaos::Softs::FCollectionPropertyFacade> PropertyFacade = CollectionPropertyFacades[LODIndex].Pin())
 		{
-			return FVector(PropertyFacade->GetValue<FVector3f>(PropertyName, FVector3f(DefaultValue)));
+			return GetValueWithAlias<FVector>(PropertyFacade, PropertyName, DefaultValue,
+				[&PropertyFacade](int32 KeyIndex)
+				{
+					return FVector(PropertyFacade->GetValue<FVector3f>(KeyIndex));
+				});
 		}
 	}
 	return DefaultValue;
@@ -138,11 +347,16 @@ FVector UChaosClothAssetInteractor::GetVectorValue(const FString& PropertyName,
 
 FString UChaosClothAssetInteractor::GetStringValue(const FString& PropertyName, int32 LODIndex, const FString& DefaultValue) const
 {
+	using namespace UE::Chaos::ClothAsset::Private;
 	if (CollectionPropertyFacades.IsValidIndex(LODIndex))
 	{
 		if (TSharedPtr<::Chaos::Softs::FCollectionPropertyFacade> PropertyFacade = CollectionPropertyFacades[LODIndex].Pin())
 		{
-			return PropertyFacade->GetStringValue(PropertyName, DefaultValue);
+			return GetValueWithAlias<FString>(PropertyFacade, PropertyName, DefaultValue,
+				[&PropertyFacade](int32 KeyIndex)
+				{
+					return PropertyFacade->GetStringValue(KeyIndex);
+				});
 		}
 	}
 	return DefaultValue;
@@ -150,13 +364,18 @@ FString UChaosClothAssetInteractor::GetStringValue(const FString& PropertyName,
 
 void UChaosClothAssetInteractor::SetFloatValue(const FString& PropertyName, int32 LODIndex, float Value)
 {
+	using namespace UE::Chaos::ClothAsset::Private;
 	if (LODIndex == INDEX_NONE)
 	{
 		for (int32 Index = 0; Index < CollectionPropertyFacades.Num(); ++Index)
 		{
 			if (TSharedPtr<::Chaos::Softs::FCollectionPropertyFacade> PropertyFacade = CollectionPropertyFacades[Index].Pin())
 			{
-				PropertyFacade->SetValue(PropertyName, Value);
+				SetValueWithAlias(PropertyFacade, PropertyName,
+					[&PropertyFacade, Value](const FString& Name)
+					{
+						PropertyFacade->SetValue(Name, Value);
+					});
 			}
 		}
 	}
@@ -164,20 +383,29 @@ void UChaosClothAssetInteractor::SetFloatValue(const FString& PropertyName, int3
 	{
 		if (TSharedPtr<::Chaos::Softs::FCollectionPropertyFacade> PropertyFacade = CollectionPropertyFacades[LODIndex].Pin())
 		{
-			PropertyFacade->SetValue(PropertyName, Value);
+			SetValueWithAlias(PropertyFacade, PropertyName,
+				[&PropertyFacade, Value](const FString& Name)
+				{
+					PropertyFacade->SetValue(Name, Value);
+				});
 		}
 	}
 }
 
 void UChaosClothAssetInteractor::SetLowFloatValue(const FString& PropertyName, int32 LODIndex, float Value)
 {
+	using namespace UE::Chaos::ClothAsset::Private;
 	if (LODIndex == INDEX_NONE)
 	{
 		for (int32 Index = 0; Index < CollectionPropertyFacades.Num(); ++Index)
 		{
 			if (TSharedPtr<::Chaos::Softs::FCollectionPropertyFacade> PropertyFacade = CollectionPropertyFacades[Index].Pin())
 			{
-				PropertyFacade->SetLowValue(PropertyName, Value);
+				SetValueWithAlias(PropertyFacade, PropertyName,
+					[&PropertyFacade, Value](const FString& Name)
+					{
+						PropertyFacade->SetLowValue(Name, Value);
+					});
 			}
 		}
 	}
@@ -185,20 +413,29 @@ void UChaosClothAssetInteractor::SetLowFloatValue(const FString& PropertyName, i
 	{
 		if (TSharedPtr<::Chaos::Softs::FCollectionPropertyFacade> PropertyFacade = CollectionPropertyFacades[LODIndex].Pin())
 		{
-			PropertyFacade->SetLowValue(PropertyName, Value);
+			SetValueWithAlias(PropertyFacade, PropertyName,
+				[&PropertyFacade, Value](const FString& Name)
+				{
+					PropertyFacade->SetLowValue(Name, Value);
+				});
 		}
 	}
 }
 
 void UChaosClothAssetInteractor::SetHighFloatValue(const FString& PropertyName, int32 LODIndex, float Value)
 {
+	using namespace UE::Chaos::ClothAsset::Private;
 	if (LODIndex == INDEX_NONE)
 	{
 		for (int32 Index = 0; Index < CollectionPropertyFacades.Num(); ++Index)
 		{
 			if (TSharedPtr<::Chaos::Softs::FCollectionPropertyFacade> PropertyFacade = CollectionPropertyFacades[Index].Pin())
 			{
-				PropertyFacade->SetHighValue(PropertyName, Value);
+				SetValueWithAlias(PropertyFacade, PropertyName,
+					[&PropertyFacade, Value](const FString& Name)
+					{
+						PropertyFacade->SetHighValue(Name, Value);
+					});
 			}
 		}
 	}
@@ -206,20 +443,29 @@ void UChaosClothAssetInteractor::SetHighFloatValue(const FString& PropertyName,
 	{
 		if (TSharedPtr<::Chaos::Softs::FCollectionPropertyFacade> PropertyFacade = CollectionPropertyFacades[LODIndex].Pin())
 		{
-			PropertyFacade->SetHighValue(PropertyName, Value);
+			SetValueWithAlias(PropertyFacade, PropertyName,
+				[&PropertyFacade, Value](const FString& Name)
+				{
+					PropertyFacade->SetHighValue(Name, Value);
+				});
 		}
 	}
 }
 
 void UChaosClothAssetInteractor::SetWeightedFloatValue(const FString& PropertyName, int32 LODIndex, FVector2D Value)
 {
+	using namespace UE::Chaos::ClothAsset::Private;
 	if (LODIndex == INDEX_NONE)
 	{
 		for (int32 Index = 0; Index < CollectionPropertyFacades.Num(); ++Index)
 		{
 			if (TSharedPtr<::Chaos::Softs::FCollectionPropertyFacade> PropertyFacade = CollectionPropertyFacades[Index].Pin())
 			{
-				PropertyFacade->SetWeightedFloatValue(PropertyName, FVector2f(Value));
+				SetValueWithAlias(PropertyFacade, PropertyName,
+					[&PropertyFacade, &Value](const FString& Name)
+					{
+						PropertyFacade->SetWeightedFloatValue(Name, FVector2f(Value));
+					});
 			}
 		}
 	}
@@ -227,20 +473,29 @@ void UChaosClothAssetInteractor::SetWeightedFloatValue(const FString& PropertyNa
 	{
 		if (TSharedPtr<::Chaos::Softs::FCollectionPropertyFacade> PropertyFacade = CollectionPropertyFacades[LODIndex].Pin())
 		{
-			PropertyFacade->SetWeightedFloatValue(PropertyName, FVector2f(Value));
+			SetValueWithAlias(PropertyFacade, PropertyName,
+				[&PropertyFacade, &Value](const FString& Name)
+				{
+					PropertyFacade->SetWeightedFloatValue(Name, FVector2f(Value));
+				});
 		}
 	}
 }
 
 void UChaosClothAssetInteractor::SetIntValue(const FString& PropertyName, int32 LODIndex, int32 Value)
 {
+	using namespace UE::Chaos::ClothAsset::Private;
 	if (LODIndex == INDEX_NONE)
 	{
 		for (int32 Index = 0; Index < CollectionPropertyFacades.Num(); ++Index)
 		{
 			if (TSharedPtr<::Chaos::Softs::FCollectionPropertyFacade> PropertyFacade = CollectionPropertyFacades[Index].Pin())
 			{
-				PropertyFacade->SetValue(PropertyName, Value);
+				SetValueWithAlias(PropertyFacade, PropertyName,
+					[&PropertyFacade, Value](const FString& Name)
+					{
+						PropertyFacade->SetValue(Name, Value);
+					});
 			}
 		}
 	}
@@ -248,20 +503,29 @@ void UChaosClothAssetInteractor::SetIntValue(const FString& PropertyName, int32
 	{
 		if (TSharedPtr<::Chaos::Softs::FCollectionPropertyFacade> PropertyFacade = CollectionPropertyFacades[LODIndex].Pin())
 		{
-			PropertyFacade->SetValue(PropertyName, Value);
+			SetValueWithAlias(PropertyFacade, PropertyName,
+				[&PropertyFacade, Value](const FString& Name)
+				{
+					PropertyFacade->SetValue(Name, Value);
+				});
 		}
 	}
 }
 
 void UChaosClothAssetInteractor::SetVectorValue(const FString& PropertyName, int32 LODIndex, FVector Value)
 {
+	using namespace UE::Chaos::ClothAsset::Private;
 	if (LODIndex == INDEX_NONE)
 	{
 		for (int32 Index = 0; Index < CollectionPropertyFacades.Num(); ++Index)
 		{
 			if (TSharedPtr<::Chaos::Softs::FCollectionPropertyFacade> PropertyFacade = CollectionPropertyFacades[Index].Pin())
 			{
-				PropertyFacade->SetValue(PropertyName, FVector3f(Value));
+				SetValueWithAlias(PropertyFacade, PropertyName,
+					[&PropertyFacade, &Value](const FString& Name)
+					{
+						PropertyFacade->SetValue(Name, FVector3f(Value));
+					});
 			}
 		}
 	}
@@ -269,20 +533,29 @@ void UChaosClothAssetInteractor::SetVectorValue(const FString& PropertyName, int
 	{
 		if (TSharedPtr<::Chaos::Softs::FCollectionPropertyFacade> PropertyFacade = CollectionPropertyFacades[LODIndex].Pin())
 		{
-			PropertyFacade->SetValue(PropertyName, FVector3f(Value));
+			SetValueWithAlias(PropertyFacade, PropertyName,
+				[&PropertyFacade, &Value](const FString& Name)
+				{
+					PropertyFacade->SetValue(Name, FVector3f(Value));
+				});
 		}
 	}
 }
 
 void UChaosClothAssetInteractor::SetStringValue(const FString& PropertyName, int32 LODIndex, const FString& Value)
 {
+	using namespace UE::Chaos::ClothAsset::Private;
 	if (LODIndex == INDEX_NONE)
 	{
 		for (int32 Index = 0; Index < CollectionPropertyFacades.Num(); ++Index)
 		{
 			if (TSharedPtr<::Chaos::Softs::FCollectionPropertyFacade> PropertyFacade = CollectionPropertyFacades[Index].Pin())
 			{
-				PropertyFacade->SetStringValue(PropertyName, Value);
+				SetValueWithAlias(PropertyFacade, PropertyName,
+					[&PropertyFacade, &Value](const FString& Name)
+					{
+						PropertyFacade->SetStringValue(Name, Value);
+					});
 			}
 		}
 	}
@@ -290,7 +563,11 @@ void UChaosClothAssetInteractor::SetStringValue(const FString& PropertyName, int
 	{
 		if (TSharedPtr<::Chaos::Softs::FCollectionPropertyFacade> PropertyFacade = CollectionPropertyFacades[LODIndex].Pin())
 		{
-			PropertyFacade->SetStringValue(PropertyName, Value);
+			SetValueWithAlias(PropertyFacade, PropertyName,
+				[&PropertyFacade, &Value](const FString& Name)
+				{
+					PropertyFacade->SetStringValue(Name, Value);
+				});
 		}
 	}
 }
diff --git a/Engine/Plugins/ChaosClothAssetEditor/Source/ChaosClothAssetDataflowNodes/Private/ChaosClothAsset/SimulationAerodynamicsConfigNode.cpp b/Engine/Plugins/ChaosClothAssetEditor/Source/ChaosClothAssetDataflowNodes/Private/ChaosClothAsset/SimulationAerodynamicsConfigNode.cpp
--- a/Engine/Plugins/ChaosClothAssetEditor/Source/ChaosClothAssetDataflowNodes/Private/ChaosClothAsset/SimulationAerodynamicsConfigNode.cpp
+++ b/Engine/Plugins/ChaosClothAssetEditor/Source/ChaosClothAssetDataflowNodes/Private/ChaosClothAsset/SimulationAerodynamicsConfigNode.cpp
@@ -27,7 +27,7 @@ FChaosClothAssetSimulationAerodynamicsConfigNode::FChaosClothAssetSimulationAero
 void FChaosClothAssetSimulationAerodynamicsConfigNode::AddProperties(FPropertyHelper& PropertyHelper) const
 {
 	PropertyHelper.SetProperty(this, &FluidDensity);
-	PropertyHelper.SetPropertyEnum(this, &WindVelocitySpace, {}, ECollectionPropertyFlags::None);
+	PropertyHelper.SetPropertyEnum(this, &WindVelocitySpace);
 	PropertyHelper.SetProperty(this, &WindVelocity);
 	
 	PropertyHelper.SetSolverPropertyWeighted(FName(TEXT("Drag")), Drag, [](
diff --git a/Engine/Plugins/ChaosClothAssetEditor/Source/ChaosClothAssetDataflowNodes/Private/ChaosClothAsset/SimulationLongRangeAttachmentConfigNode.cpp b/Engine/Plugins/ChaosClothAssetEditor/Source/ChaosClothAssetDataflowNodes/Private/ChaosClothAsset/SimulationLongRangeAttachmentConfigNode.cpp
--- a/Engine/Plugins/ChaosClothAssetEditor/Source/ChaosClothAssetDataflowNodes/Private/ChaosClothAsset/SimulationLongRangeAttachmentConfigNode.cpp
+++ b/Engine/Plugins/ChaosClothAssetEditor/Source/ChaosClothAssetDataflowNodes/Private/ChaosClothAsset/SimulationLongRangeAttachmentConfigNode.cpp
@@ -24,7 +24,7 @@ void FChaosClothAssetSimulationLongRangeAttachmentConfigNode_v2::AddProperties(F
 {
 	PropertyHelper.SetPropertyWeighted(this, &TetherStiffness);
 	PropertyHelper.SetPropertyWeighted(this, &TetherScale);
-	PropertyHelper.SetPropertyBool(this, &bUseGeodesicTethers);
+	PropertyHelper.SetPropertyBool(this, &bUseGeodesicTethers, {}, ECollectionPropertyFlags::Intrinsic);  // Intrinsic since the tethers need to be recalculated.
 	PropertyHelper.SetPropertyString(this, &FixedEndSet);
 }
 
@@ -54,7 +54,7 @@ void FChaosClothAssetSimulationLongRangeAttachmentConfigNode::AddProperties(FPro
 PRAGMA_DISABLE_DEPRECATION_WARNINGS  // SetProperty functions are templated and cause deprecation warnings with the now deprecated v1
 	PropertyHelper.SetPropertyWeighted(this, &TetherStiffness);
 	PropertyHelper.SetPropertyWeighted(this, &TetherScale);
-	PropertyHelper.SetPropertyBool(this, &bUseGeodesicTethers);
+	PropertyHelper.SetPropertyBool(this, &bUseGeodesicTethers, {}, ECollectionPropertyFlags::Intrinsic);  // Intrinsic since the tethers need to be recalculated.
 	PropertyHelper.SetPropertyString(this, &FixedEndWeightMap);
 PRAGMA_ENABLE_DEPRECATION_WARNINGS
 }
diff --git a/Engine/Plugins/ChaosClothAssetEditor/Source/ChaosClothAssetDataflowNodes/Private/ChaosClothAsset/SimulationSelfCollisionConfigNode.cpp b/Engine/Plugins/ChaosClothAssetEditor/Source/ChaosClothAssetDataflowNodes/Private/ChaosClothAsset/SimulationSelfCollisionConfigNode.cpp
--- a/Engine/Plugins/ChaosClothAssetEditor/Source/ChaosClothAssetDataflowNodes/Private/ChaosClothAsset/SimulationSelfCollisionConfigNode.cpp
+++ b/Engine/Plugins/ChaosClothAssetEditor/Source/ChaosClothAssetDataflowNodes/Private/ChaosClothAsset/SimulationSelfCollisionConfigNode.cpp
@@ -30,7 +30,7 @@ FChaosClothAssetSimulationSelfCollisionConfigNode::FChaosClothAssetSimulationSel
 
 void FChaosClothAssetSimulationSelfCollisionConfigNode::AddProperties(FPropertyHelper& PropertyHelper) const
 {
-	PropertyHelper.SetPropertyBool(FName("UseSelfCollisions"), true);
+	PropertyHelper.SetPropertyBool(this, &bUseSelfCollisions);
 	PropertyHelper.SetProperty(this, &SelfCollisionStiffness);
 	PropertyHelper.SetProperty(this, &SelfCollisionDisableNeighborDistance, {}, ECollectionPropertyFlags::None); // Non animatable
 	PropertyHelper.SetPropertyString(this, &SelfCollisionDisabledFaces);
diff --git a/Engine/Plugins/ChaosClothAssetEditor/Source/ChaosClothAssetDataflowNodes/Private/ChaosClothAsset/SimulationSelfCollisionSpheresConfigNode.cpp b/Engine/Plugins/ChaosClothAssetEditor/Source/ChaosClothAssetDataflowNodes/Private/ChaosClothAsset/SimulationSelfCollisionSpheresConfigNode.cpp
--- a/Engine/Plugins/ChaosClothAssetEditor/Source/ChaosClothAssetDataflowNodes/Private/ChaosClothAsset/SimulationSelfCollisionSpheresConfigNode.cpp
+++ b/Engine/Plugins/ChaosClothAssetEditor/Source/ChaosClothAssetDataflowNodes/Private/ChaosClothAsset/SimulationSelfCollisionSpheresConfigNode.cpp
@@ -20,7 +20,7 @@ FChaosClothAssetSimulationSelfCollisionSpheresConfigNode::FChaosClothAssetSimula
 void FChaosClothAssetSimulationSelfCollisionSpheresConfigNode::AddProperties(FPropertyHelper& PropertyHelper) const
 {
 	PropertyHelper.SetProperty(this, &SelfCollisionSphereRadius);
-	PropertyHelper.SetProperty(this, &SelfCollisionSphereRadiusCullMultiplier);
+	PropertyHelper.SetProperty(this, &SelfCollisionSphereRadiusCullMultiplier, {}, ECollectionPropertyFlags::Intrinsic); // Need to rebuild the sphere set to see this effect.
 	PropertyHelper.SetProperty(this, &SelfCollisionSphereStiffness);
 	PropertyHelper.SetPropertyString(this, &SelfCollisionSphereSetName, {}, ECollectionPropertyFlags::None);  // Non animatable
 }
diff --git a/Engine/Plugins/ChaosClothAssetEditor/Source/ChaosClothAssetDataflowNodes/Private/ChaosClothAsset/SimulationVelocityScaleConfigNode.cpp b/Engine/Plugins/ChaosClothAssetEditor/Source/ChaosClothAssetDataflowNodes/Private/ChaosClothAsset/SimulationVelocityScaleConfigNode.cpp
--- a/Engine/Plugins/ChaosClothAssetEditor/Source/ChaosClothAssetDataflowNodes/Private/ChaosClothAsset/SimulationVelocityScaleConfigNode.cpp
+++ b/Engine/Plugins/ChaosClothAssetEditor/Source/ChaosClothAssetDataflowNodes/Private/ChaosClothAsset/SimulationVelocityScaleConfigNode.cpp
@@ -14,7 +14,7 @@ FChaosClothAssetSimulationVelocityScaleConfigNode::FChaosClothAssetSimulationVel
 void FChaosClothAssetSimulationVelocityScaleConfigNode::AddProperties(FPropertyHelper& PropertyHelper) const
 {
 	static const FVector3f DefaultLinearClamp = { TNumericLimits<float>::Max(), TNumericLimits<float>::Max(), TNumericLimits<float>::Max() };
-	PropertyHelper.SetPropertyEnum(this, &VelocityScaleSpace, {}, ECollectionPropertyFlags::None);
+	PropertyHelper.SetPropertyEnum(this, &VelocityScaleSpace);
 	PropertyHelper.SetProperty(this, &LinearVelocityScale);
 	PropertyHelper.SetProperty(TEXT("MaxLinearVelocity"), bEnableLinearVelocityClamping ? MaxLinearVelocity : DefaultLinearClamp);
 	PropertyHelper.SetProperty(TEXT("MaxLinearAcceleration"), bEnableLinearAccelerationClamping ? MaxLinearAcceleration : DefaultLinearClamp);
diff --git a/Engine/Plugins/ChaosClothAssetEditor/Source/ChaosClothAssetDataflowNodes/Public/ChaosClothAsset/SimulationAerodynamicsConfigNode.h b/Engine/Plugins/ChaosClothAssetEditor/Source/ChaosClothAssetDataflowNodes/Public/ChaosClothAsset/SimulationAerodynamicsConfigNode.h
--- a/Engine/Plugins/ChaosClothAssetEditor/Source/ChaosClothAssetDataflowNodes/Public/ChaosClothAsset/SimulationAerodynamicsConfigNode.h
+++ b/Engine/Plugins/ChaosClothAssetEditor/Source/ChaosClothAssetDataflowNodes/Public/ChaosClothAsset/SimulationAerodynamicsConfigNode.h
@@ -21,20 +21,20 @@ public:
 	 * The fluid density is given in kg/m^3.
 	 * Air density is considered to be around 1.225 kg/m^3 in average atmospheric conditions.
 	 */
-	UPROPERTY(EditAnywhere, Category = "Aerodynamics", Meta = (UIMin = "0", ClampMin = "0", ClampMax = "10"))
+	UPROPERTY(EditAnywhere, Category = "Aerodynamics", Meta = (UIMin = "0", ClampMin = "0", ClampMax = "10", InteractorName = "FluidDensity"))
 	float FluidDensity = 1.225f;
 
 	/**
 	 * Wind velocity is specified in this space.
 	 */
-	UPROPERTY(EditAnywhere, Category = "Aerodynamics")
+	UPROPERTY(EditAnywhere, Category = "Aerodynamics", Meta = (InteractorName = "WindVelocitySpace"))
 	EChaosSoftsSimulationSpace WindVelocitySpace = EChaosSoftsSimulationSpace::WorldSpace;
 	
 	/**
 	 * The fixed wind velocity [m/s] for this asset.
 	 * For reference a wind gust is above 8m/s (18mph).
 	 */
-	UPROPERTY(EditAnywhere, Category = "Aerodynamics", Meta = (UIMin = "0", UIMax = "10"))
+	UPROPERTY(EditAnywhere, Category = "Aerodynamics", Meta = (UIMin = "0", UIMax = "10", InteractorName = "WindVelocity"))
 	FVector3f WindVelocity = { 0.f, 0.f, 0.f };
 
 	/**
@@ -45,7 +45,7 @@ public:
 	 * are interpolated with the per particle weight to make the final value used for the simulation.
 	 * Otherwise all particles are considered to have a zero weight, and only the Low value is meaningful.
 	 */
-	UPROPERTY(EditAnywhere, Category = "Aerodynamics", Meta = (UIMin = "0", UIMax = "1", ClampMin = "0", ClampMax = "10"))
+	UPROPERTY(EditAnywhere, Category = "Aerodynamics", Meta = (UIMin = "0", UIMax = "1", ClampMin = "0", ClampMax = "10", InteractorName = "Drag"))
 	FChaosClothAssetWeightedValue Drag = { true, 0.035f, 1.f, TEXT("Drag"), true };
 
 	UPROPERTY(EditAnywhere, Category = "Aerodynamics", Meta = (InlineEditConditionToggle))
@@ -58,7 +58,7 @@ public:
 	 * are interpolated with the per particle weight to make the final value used for the simulation.
 	 * Otherwise all particles are considered to have a zero weight, and only the Low value is meaningful.
 	 */
-	UPROPERTY(EditAnywhere, Category = "Aerodynamics", Meta = (UIMin = "0", UIMax = "1", ClampMin = "0", ClampMax = "10", EditCondition = "bEnableOuterDrag"))
+	UPROPERTY(EditAnywhere, Category = "Aerodynamics", Meta = (UIMin = "0", UIMax = "1", ClampMin = "0", ClampMax = "10", EditCondition = "bEnableOuterDrag", InteractorName = "OuterDrag"))
 	FChaosClothAssetWeightedValue OuterDrag = { true, 0.035f, 1.f, TEXT("OuterDrag"), true };
 
 
@@ -70,7 +70,7 @@ public:
 	 * are interpolated with the per particle weight to make the final value used for the simulation.
 	 * Otherwise all particles are considered to have a zero weight, and only the Low value is meaningful.
 	 */
-	UPROPERTY(EditAnywhere, Category = "Aerodynamics", Meta = (UIMin = "0", UIMax = "1", ClampMin = "0", ClampMax = "10"))
+	UPROPERTY(EditAnywhere, Category = "Aerodynamics", Meta = (UIMin = "0", UIMax = "1", ClampMin = "0", ClampMax = "10", InteractorName = "Lift"))
 	FChaosClothAssetWeightedValue Lift = { true, 0.035f, 1.f, TEXT("Lift"), true };
 
 	UPROPERTY(EditAnywhere, Category = "Aerodynamics", Meta = (InlineEditConditionToggle))
@@ -83,7 +83,7 @@ public:
 	 * are interpolated with the per particle weight to make the final value used for the simulation.
 	 * Otherwise all particles are considered to have a zero weight, and only the Low value is meaningful.
 	 */
-	UPROPERTY(EditAnywhere, Category = "Aerodynamics", Meta = (UIMin = "0", UIMax = "1", ClampMin = "0", ClampMax = "10", EditCondition = "bEnableOuterLift"))
+	UPROPERTY(EditAnywhere, Category = "Aerodynamics", Meta = (UIMin = "0", UIMax = "1", ClampMin = "0", ClampMax = "10", EditCondition = "bEnableOuterLift", InteractorName = "OuterLift"))
 	FChaosClothAssetWeightedValue OuterLift = { true, 0.035f, 1.f, TEXT("OuterLift"), true };
 	
 	FChaosClothAssetSimulationAerodynamicsConfigNode(const UE::Dataflow::FNodeParameters& InParam, FGuid InGuid = FGuid::NewGuid());
diff --git a/Engine/Plugins/ChaosClothAssetEditor/Source/ChaosClothAssetDataflowNodes/Public/ChaosClothAsset/SimulationAnimDriveConfigNode.h b/Engine/Plugins/ChaosClothAssetEditor/Source/ChaosClothAssetDataflowNodes/Public/ChaosClothAsset/SimulationAnimDriveConfigNode.h
--- a/Engine/Plugins/ChaosClothAssetEditor/Source/ChaosClothAssetDataflowNodes/Public/ChaosClothAsset/SimulationAnimDriveConfigNode.h
+++ b/Engine/Plugins/ChaosClothAssetEditor/Source/ChaosClothAssetDataflowNodes/Public/ChaosClothAsset/SimulationAnimDriveConfigNode.h
@@ -20,7 +20,7 @@ public:
 	 * are interpolated with the per particle weight to make the final value used for the simulation.
 	 * Otherwise all particles are considered to have a zero weight, and only the Low value is meaningful.
 	 */
-	UPROPERTY(EditAnywhere, Category = "AnimDrive Properties", Meta = (UIMin = "0", UIMax = "1", ClampMin = "0", ClampMax = "1"))
+	UPROPERTY(EditAnywhere, Category = "AnimDrive Properties", Meta = (UIMin = "0", UIMax = "1", ClampMin = "0", ClampMax = "1", InteractorName = "AnimDriveStiffness"))
 	FChaosClothAssetWeightedValue AnimDriveStiffness = { true, 0.f, 1.f, TEXT("AnimDriveStiffness") };
 
 	/**
@@ -29,7 +29,7 @@ public:
 	 * are interpolated with the per particle weight to make the final value used for the simulation.
 	 * Otherwise all particles are considered to have a zero weight, and only the Low value is meaningful.
 	 */
-	UPROPERTY(EditAnywhere, Category = "AnimDrive Properties", Meta = (UIMin = "0", UIMax = "1", ClampMin = "0", ClampMax = "1"))
+	UPROPERTY(EditAnywhere, Category = "AnimDrive Properties", Meta = (UIMin = "0", UIMax = "1", ClampMin = "0", ClampMax = "1", InteractorName = "AnimDriveDamping"))
 	FChaosClothAssetWeightedValue AnimDriveDamping = { true, 0.f, 1.f, TEXT("AnimDriveDamping") };
 	
 	FChaosClothAssetSimulationAnimDriveConfigNode(const UE::Dataflow::FNodeParameters& InParam, FGuid InGuid = FGuid::NewGuid());
diff --git a/Engine/Plugins/ChaosClothAssetEditor/Source/ChaosClothAssetDataflowNodes/Public/ChaosClothAsset/SimulationBackstopConfigNode.h b/Engine/Plugins/ChaosClothAssetEditor/Source/ChaosClothAssetDataflowNodes/Public/ChaosClothAsset/SimulationBackstopConfigNode.h
--- a/Engine/Plugins/ChaosClothAssetEditor/Source/ChaosClothAssetDataflowNodes/Public/ChaosClothAsset/SimulationBackstopConfigNode.h
+++ b/Engine/Plugins/ChaosClothAssetEditor/Source/ChaosClothAssetDataflowNodes/Public/ChaosClothAsset/SimulationBackstopConfigNode.h
@@ -21,7 +21,7 @@ public:
 	 * are interpolated with the per particle weight to make the final value used for the simulation.
 	 * Otherwise all particles are considered to have a zero weight, and only the Low value is meaningful.
 	 */
-	UPROPERTY(EditAnywhere, Category = "Backstop Properties", Meta = (UIMin = "0", UIMax = "100"))
+	UPROPERTY(EditAnywhere, Category = "Backstop Properties", Meta = (UIMin = "0", UIMax = "100", InteractorName = "BackstopDistance"))
 	FChaosClothAssetWeightedValue BackstopDistance = { true, 0.f, 100.f, TEXT("BackstopDistance") };
 
 	/**
@@ -30,7 +30,7 @@ public:
 	 * are interpolated with the per particle weight to make the final value used for the simulation.
 	 * Otherwise all particles are considered to have a zero weight, and only the Low value is meaningful.
 	 */
-	UPROPERTY(EditAnywhere, Category = "Backstop Properties", Meta = (UIMin = "0", UIMax = "100", ClampMin = "0"))
+	UPROPERTY(EditAnywhere, Category = "Backstop Properties", Meta = (UIMin = "0", UIMax = "100", ClampMin = "0", InteractorName = "BackstopRadius"))
 	FChaosClothAssetWeightedValue BackstopRadius = { true, 0.f, 100.f, TEXT("BackstopRadius") };
 
 	FChaosClothAssetSimulationBackstopConfigNode(const UE::Dataflow::FNodeParameters& InParam, FGuid InGuid = FGuid::NewGuid());
diff --git a/Engine/Plugins/ChaosClothAssetEditor/Source/ChaosClothAssetDataflowNodes/Public/ChaosClothAsset/SimulationBendingConfigNode.h b/Engine/Plugins/ChaosClothAssetEditor/Source/ChaosClothAssetDataflowNodes/Public/ChaosClothAsset/SimulationBendingConfigNode.h
--- a/Engine/Plugins/ChaosClothAssetEditor/Source/ChaosClothAssetDataflowNodes/Public/ChaosClothAsset/SimulationBendingConfigNode.h
+++ b/Engine/Plugins/ChaosClothAssetEditor/Source/ChaosClothAssetDataflowNodes/Public/ChaosClothAsset/SimulationBendingConfigNode.h
@@ -63,7 +63,7 @@ public:
 	 * are interpolated with the per particle weight to make the final value used for the simulation.
 	 * Otherwise all particles are considered to have a zero weight, and only the Low value is meaningful.
 	 */
-	UPROPERTY(EditAnywhere, Category = "Bending Properties", Meta = (UIMin = "0", UIMax = "10000", ClampMin = "0", ClampMax = "10000000", EditCondition = "DistributionType == EChaosClothAssetConstraintDistributionType::Anisotropic && SolverType == EChaosClothAssetConstraintSolverType::XPBD", EditConditionHides))
+	UPROPERTY(EditAnywhere, Category = "Bending Properties", Meta = (UIMin = "0", UIMax = "10000", ClampMin = "0", ClampMax = "10000000", EditCondition = "DistributionType == EChaosClothAssetConstraintDistributionType::Anisotropic && SolverType == EChaosClothAssetConstraintSolverType::XPBD", EditConditionHides, InteractorName = "BendingStiffnessWarp"))
 	FChaosClothAssetWeightedValue BendingStiffnessWarp = { true,  100.f, 100.f, TEXT("BendingStiffnessWarp"), true };
 
 	/**
@@ -72,7 +72,7 @@ public:
 	 * are interpolated with the per particle weight to make the final value used for the simulation.
 	 * Otherwise all particles are considered to have a zero weight, and only the Low value is meaningful.
 	 */
-	UPROPERTY(EditAnywhere, Category = "Bending Properties", Meta = (UIMin = "0", UIMax = "10000", ClampMin = "0", ClampMax = "10000000", EditCondition = "DistributionType == EChaosClothAssetConstraintDistributionType::Anisotropic && SolverType == EChaosClothAssetConstraintSolverType::XPBD", EditConditionHides))
+	UPROPERTY(EditAnywhere, Category = "Bending Properties", Meta = (UIMin = "0", UIMax = "10000", ClampMin = "0", ClampMax = "10000000", EditCondition = "DistributionType == EChaosClothAssetConstraintDistributionType::Anisotropic && SolverType == EChaosClothAssetConstraintSolverType::XPBD", EditConditionHides, InteractorName = "BendingStiffnessWeft"))
 	FChaosClothAssetWeightedValue BendingStiffnessWeft = { true, 100.f, 100.f, TEXT("BendingStiffnessWeft"),  true };
 
 	/**
@@ -81,7 +81,7 @@ public:
 	 * are interpolated with the per particle weight to make the final value used for the simulation.
 	 * Otherwise all particles are considered to have a zero weight, and only the Low value is meaningful.
 	 */
-	UPROPERTY(EditAnywhere, Category = "Bending Properties", Meta = (UIMin = "0", UIMax = "10000", ClampMin = "0", ClampMax = "10000000", EditCondition = "DistributionType == EChaosClothAssetConstraintDistributionType::Anisotropic && SolverType == EChaosClothAssetConstraintSolverType::XPBD", EditConditionHides))
+	UPROPERTY(EditAnywhere, Category = "Bending Properties", Meta = (UIMin = "0", UIMax = "10000", ClampMin = "0", ClampMax = "10000000", EditCondition = "DistributionType == EChaosClothAssetConstraintDistributionType::Anisotropic && SolverType == EChaosClothAssetConstraintSolverType::XPBD", EditConditionHides, InteractorName = "BendingStiffnessBias"))
 	FChaosClothAssetWeightedValue BendingStiffnessBias = { true, 100.f, 100.f, TEXT("BendingStiffnessBias"), true };
 
 	/**
@@ -90,14 +90,14 @@ public:
 	 * are interpolated with the per particle weight to make the final value used for the simulation.
 	 * Otherwise all particles are considered to have a zero weight, and only the Low value is meaningful.
 	 */
-	UPROPERTY(EditAnywhere, Category = "Bending Properties", Meta = (UIMin = "0", UIMax = "10", ClampMin = "0", ClampMax = "1000"), Meta = (EditCondition = "SolverType == EChaosClothAssetConstraintSolverType::XPBD && DistributionType == EChaosClothAssetConstraintDistributionType::Anisotropic", EditConditionHides))
+	UPROPERTY(EditAnywhere, Category = "Bending Properties", Meta = (UIMin = "0", UIMax = "10", ClampMin = "0", ClampMax = "1000"), Meta = (EditCondition = "SolverType == EChaosClothAssetConstraintSolverType::XPBD && DistributionType == EChaosClothAssetConstraintDistributionType::Anisotropic", EditConditionHides, InteractorName = "BendingDamping"))
 	FChaosClothAssetWeightedValue BendingAnisoDamping = { true, 1.f, 1.f, TEXT("BendingAnisoDamping"), true };
 
 	/**
 	 * Once the element has bent such that it's folded more than this ratio from its rest angle ("buckled"), switch to using Buckling Stiffness instead of BendingElement Stiffness.
 	 * When Buckling Ratio = 0, the Buckling Stiffness will never be used. When BucklingRatio = 1, the Buckling Stiffness will be used as soon as it's bent past its rest configuration.
 	 */
-	UPROPERTY(EditAnywhere, Category = "Bending Properties", DisplayName = "Aniso Buckling Ratio", Meta = (UIMin = "0", UIMax = "1", ClampMin = "0", ClampMax = "1"), Meta = (EditCondition = "SolverType == EChaosClothAssetConstraintSolverType::XPBD && DistributionType == EChaosClothAssetConstraintDistributionType::Anisotropic", EditConditionHides))
+	UPROPERTY(EditAnywhere, Category = "Bending Properties", DisplayName = "Aniso Buckling Ratio", Meta = (UIMin = "0", UIMax = "1", ClampMin = "0", ClampMax = "1"), Meta = (EditCondition = "SolverType == EChaosClothAssetConstraintSolverType::XPBD && DistributionType == EChaosClothAssetConstraintDistributionType::Anisotropic", EditConditionHides, InteractorName = "BucklingRatio"))
 	FChaosClothAssetWeightedValue AnisoBucklingRatio= { true, 0.5f, 0.5f, TEXT("BucklingRatio"), true };
 	
 	/**
@@ -109,7 +109,7 @@ public:
 	 * are interpolated with the per particle weight to make the final value used for the simulation.
 	 * Otherwise all particles are considered to have a zero weight, and only the Low value is meaningful.
 	 */
-	UPROPERTY(EditAnywhere, Category = "Bending Properties", Meta = (UIMin = "0", UIMax = "10000", ClampMin = "0", ClampMax = "10000000", EditCondition = "DistributionType == EChaosClothAssetConstraintDistributionType::Anisotropic  && SolverType == EChaosClothAssetConstraintSolverType::XPBD", EditConditionHides))
+	UPROPERTY(EditAnywhere, Category = "Bending Properties", Meta = (UIMin = "0", UIMax = "10000", ClampMin = "0", ClampMax = "10000000", EditCondition = "DistributionType == EChaosClothAssetConstraintDistributionType::Anisotropic  && SolverType == EChaosClothAssetConstraintSolverType::XPBD", EditConditionHides, InteractorName = "BucklingStiffnessWarp"))
 	FChaosClothAssetWeightedValue BucklingStiffnessWarp = { true, 50.f, 50.f, TEXT("BucklingStiffnessWarp"), true };
 
 	/**
@@ -121,7 +121,7 @@ public:
 	 * are interpolated with the per particle weight to make the final value used for the simulation.
 	 * Otherwise all particles are considered to have a zero weight, and only the Low value is meaningful.
 	 */
-	UPROPERTY(EditAnywhere, Category = "Bending Properties", Meta = (UIMin = "0", UIMax = "10000", ClampMin = "0", ClampMax = "10000000", EditCondition = "DistributionType == EChaosClothAssetConstraintDistributionType::Anisotropic && SolverType == EChaosClothAssetConstraintSolverType::XPBD", EditConditionHides))
+	UPROPERTY(EditAnywhere, Category = "Bending Properties", Meta = (UIMin = "0", UIMax = "10000", ClampMin = "0", ClampMax = "10000000", EditCondition = "DistributionType == EChaosClothAssetConstraintDistributionType::Anisotropic && SolverType == EChaosClothAssetConstraintSolverType::XPBD", EditConditionHides, InteractorName = "BucklingStiffnessWeft"))
 	FChaosClothAssetWeightedValue BucklingStiffnessWeft = { true, 50.f, 50.f, TEXT("BucklingStiffnessWeft"),  true };
 
 	/**
@@ -133,16 +133,17 @@ public:
 	 * are interpolated with the per particle weight to make the final value used for the simulation.
 	 * Otherwise all particles are considered to have a zero weight, and only the Low value is meaningful.
 	 */
-	UPROPERTY(EditAnywhere, Category = "Bending Properties", Meta = (UIMin = "0", UIMax = "10000", ClampMin = "0", ClampMax = "10000000", EditCondition = "DistributionType == EChaosClothAssetConstraintDistributionType::Anisotropic && SolverType == EChaosClothAssetConstraintSolverType::XPBD", EditConditionHides))
+	UPROPERTY(EditAnywhere, Category = "Bending Properties", Meta = (UIMin = "0", UIMax = "10000", ClampMin = "0", ClampMax = "10000000", EditCondition = "DistributionType == EChaosClothAssetConstraintDistributionType::Anisotropic && SolverType == EChaosClothAssetConstraintSolverType::XPBD", EditConditionHides, InteractorName = "BucklingStiffnessBias"))
 	FChaosClothAssetWeightedValue BucklingStiffnessBias = { true, 50.f, 50.f, TEXT("BucklingStiffnessBias"), true };
 
 	/**
-	 * The Stiffness of the bending constraints. Increase the iteration count for stiffer materials.
+	 * The Stiffness of the bending constraints. Increase the iteration count for stiffer materials. 
+	 * Note that PBD stiffnesses will be internally clamped to [0,1].
 	 * If a valid weight map is found with the given Weight Map name, then both Low and High values
 	 * are interpolated with the per particle weight to make the final value used for the simulation.
 	 * Otherwise all particles are considered to have a zero weight, and only the Low value is meaningful.
 	 */
-	UPROPERTY(EditAnywhere, Category = "Bending Properties", Meta = (UIMin = "0", UIMax = "1", ClampMin = "0", ClampMax = "1"), Meta = (EditCondition = "((DistributionType == EChaosClothAssetConstraintDistributionType::Isotropic && SolverType == EChaosClothAssetConstraintSolverType::XPBD) || SolverType == EChaosClothAssetConstraintSolverType::PBD)", EditConditionHides))
+	UPROPERTY(EditAnywhere, Category = "Bending Properties", Meta = (UIMin = "0", UIMax = "10000", ClampMin = "0", ClampMax = "10000000"), Meta = (EditCondition = "((DistributionType == EChaosClothAssetConstraintDistributionType::Isotropic && SolverType == EChaosClothAssetConstraintSolverType::XPBD) || SolverType == EChaosClothAssetConstraintSolverType::PBD)", EditConditionHides, InteractorName = "BendingStiffness"))
 	FChaosClothAssetWeightedValue BendingStiffness = { true, 1.f, 1.f, TEXT("BendingStiffness") };
 
 	/**
@@ -151,26 +152,27 @@ public:
 	 * are interpolated with the per particle weight to make the final value used for the simulation.
 	 * Otherwise all particles are considered to have a zero weight, and only the Low value is meaningful.
 	 */
-	UPROPERTY(EditAnywhere, Category = "Bending Properties", Meta = (UIMin = "0", UIMax = "10", ClampMin = "0", ClampMax = "1000"), Meta = (EditCondition = "SolverType == EChaosClothAssetConstraintSolverType::XPBD && DistributionType == EChaosClothAssetConstraintDistributionType::Isotropic ", EditConditionHides))
+	UPROPERTY(EditAnywhere, Category = "Bending Properties", Meta = (UIMin = "0", UIMax = "10", ClampMin = "0", ClampMax = "1000"), Meta = (EditCondition = "SolverType == EChaosClothAssetConstraintSolverType::XPBD && DistributionType == EChaosClothAssetConstraintDistributionType::Isotropic ", EditConditionHides, InteractorName = "BendingDamping"))
 	FChaosClothAssetWeightedValue BendingDamping = { true, 1.f, 1.f, TEXT("BendingDamping") };
 	
 	/**
 	 * The stiffness after buckling.
 	 * The constraint will use this stiffness instead of bending Stiffness once the cloth has buckled, i.e., bent beyond a certain angle.
 	 * Typically, Buckling Stiffness is set to be less than Bending Stiffness.
+	 * Note that PBD stiffnesses will be internally clamped to [0,1].
 	 * Buckling Ratio determines the switch point between using BendingElement Stiffness and Buckling Stiffness.
 	 * If a valid weight map is found with the given Weight Map name, then both Low and High values
 	 * are interpolated with the per particle weight to make the final value used for the simulation.
 	 * Otherwise all particles are considered to have a zero weight, and only the Low value is meaningful.
 	 */
-	UPROPERTY(EditAnywhere, Category = "Bending Properties", Meta = (UIMin = "0", UIMax = "1", ClampMin = "0", ClampMax = "1", EditCondition = "((DistributionType == EChaosClothAssetConstraintDistributionType::Isotropic && SolverType == EChaosClothAssetConstraintSolverType::XPBD) || SolverType == EChaosClothAssetConstraintSolverType::PBD) && ConstraintType == EChaosClothAssetBendingConstraintType::HingeAngles", EditConditionHides))
+	UPROPERTY(EditAnywhere, Category = "Bending Properties", Meta = (UIMin = "0", UIMax = "10", ClampMin = "0", ClampMax = "1000", EditCondition = "((DistributionType == EChaosClothAssetConstraintDistributionType::Isotropic && SolverType == EChaosClothAssetConstraintSolverType::XPBD) || SolverType == EChaosClothAssetConstraintSolverType::PBD) && ConstraintType == EChaosClothAssetBendingConstraintType::HingeAngles", EditConditionHides, InteractorName = "BucklingStiffness"))
 	FChaosClothAssetWeightedValue BucklingStiffness = { true, 0.9f, 0.9f, TEXT("BucklingStiffness") };
 
 	/**
 	 * Once the element has bent such that it's folded more than this ratio from its rest angle ("buckled"), switch to using Buckling Stiffness instead of BendingElement Stiffness.
 	 * When Buckling Ratio = 0, the Buckling Stiffness will never be used. When BucklingRatio = 1, the Buckling Stiffness will be used as soon as it's bent past its rest configuration.
 	 */
-	UPROPERTY(EditAnywhere, Category = "Bending Properties", DisplayName = "Buckling Ratio", Meta = (UIMin = "0", UIMax = "1", ClampMin = "0", ClampMax = "1"), Meta = (EditCondition = "((DistributionType == EChaosClothAssetConstraintDistributionType::Isotropic && SolverType == EChaosClothAssetConstraintSolverType::XPBD) || SolverType == EChaosClothAssetConstraintSolverType::PBD) && ConstraintType == EChaosClothAssetBendingConstraintType::HingeAngles", EditConditionHides))
+	UPROPERTY(EditAnywhere, Category = "Bending Properties", DisplayName = "Buckling Ratio", Meta = (UIMin = "0", UIMax = "1", ClampMin = "0", ClampMax = "1"), Meta = (EditCondition = "((DistributionType == EChaosClothAssetConstraintDistributionType::Isotropic && SolverType == EChaosClothAssetConstraintSolverType::XPBD) || SolverType == EChaosClothAssetConstraintSolverType::PBD) && ConstraintType == EChaosClothAssetBendingConstraintType::HingeAngles", EditConditionHides, InteractorName = "BucklingRatio"))
 	FChaosClothAssetWeightedValue BucklingRatioWeighted = { true, 0.5f, 0.5f, TEXT("BucklingRatio") };
 
 	FChaosClothAssetSimulationBendingConfigNode(const UE::Dataflow::FNodeParameters& InParam, FGuid InGuid = FGuid::NewGuid());
diff --git a/Engine/Plugins/ChaosClothAssetEditor/Source/ChaosClothAssetDataflowNodes/Public/ChaosClothAsset/SimulationCollisionConfigNode.h b/Engine/Plugins/ChaosClothAssetEditor/Source/ChaosClothAssetDataflowNodes/Public/ChaosClothAsset/SimulationCollisionConfigNode.h
--- a/Engine/Plugins/ChaosClothAssetEditor/Source/ChaosClothAssetDataflowNodes/Public/ChaosClothAsset/SimulationCollisionConfigNode.h
+++ b/Engine/Plugins/ChaosClothAssetEditor/Source/ChaosClothAssetDataflowNodes/Public/ChaosClothAsset/SimulationCollisionConfigNode.h
@@ -14,22 +14,22 @@ struct FChaosClothAssetSimulationCollisionConfigNode : public FChaosClothAssetSi
 
 public:
 	/** The added thickness of collision shapes. */
-	UPROPERTY(EditAnywhere, Category = "Collision Properties", DisplayName = "Collision Thickness", Meta = (UIMin = "0", UIMax = "100", ClampMin = "0", ClampMax = "1000"))
+	UPROPERTY(EditAnywhere, Category = "Collision Properties", DisplayName = "Collision Thickness", Meta = (UIMin = "0", UIMax = "100", ClampMin = "0", ClampMax = "1000", InteractorName = "CollisionThickness"))
 	FChaosClothAssetImportedFloatValue CollisionThicknessImported = {UE::Chaos::ClothAsset::FDefaultFabric::CollisionThickness};
 
 	/** Friction coefficient for cloth - collider interaction. */
-	UPROPERTY(EditAnywhere, Category = "Collision Properties", DisplayName = "Friction Coefficient", Meta = (UIMin = "0", UIMax = "1", ClampMin = "0", ClampMax = "10"))
+	UPROPERTY(EditAnywhere, Category = "Collision Properties", DisplayName = "Friction Coefficient", Meta = (UIMin = "0", UIMax = "1", ClampMin = "0", ClampMax = "10", InteractorName = "FrictionCoefficient"))
 	FChaosClothAssetImportedFloatValue FrictionCoefficientImported = {UE::Chaos::ClothAsset::FDefaultFabric::Friction};
 
 	/** Stiffness for proximity repulsion forces (Force-based solver only). Units = kg cm/ s^2 (same as XPBD springs)*/
-	UPROPERTY(EditAnywhere, Category = "Proximity Force Properties", Meta = (UIMin = "0", UIMax = "10000", ClampMin = "0", ClampMax = "10000000"))
+	UPROPERTY(EditAnywhere, Category = "Proximity Force Properties", Meta = (UIMin = "0", UIMax = "10000", ClampMin = "0", ClampMax = "10000000", InteractorName = "ProximityStiffness"))
 	float ProximityStiffness = 100.f;
 
 	/**
 	 * Use continuous collision detection (CCD) to prevent any missed collisions between fast moving particles and colliders.
 	 * This has a negative effect on performance compared to when resolving collision without using CCD.
 	 */
-	UPROPERTY(EditAnywhere, Category = "Collision Properties")
+	UPROPERTY(EditAnywhere, Category = "Collision Properties", Meta = (InteractorName = "UseCCD"))
 	bool bUseCCD = false;
 
 	FChaosClothAssetSimulationCollisionConfigNode(const UE::Dataflow::FNodeParameters& InParam, FGuid InGuid = FGuid::NewGuid());
diff --git a/Engine/Plugins/ChaosClothAssetEditor/Source/ChaosClothAssetDataflowNodes/Public/ChaosClothAsset/SimulationDampingConfigNode.h b/Engine/Plugins/ChaosClothAssetEditor/Source/ChaosClothAssetDataflowNodes/Public/ChaosClothAsset/SimulationDampingConfigNode.h
--- a/Engine/Plugins/ChaosClothAssetEditor/Source/ChaosClothAssetDataflowNodes/Public/ChaosClothAsset/SimulationDampingConfigNode.h
+++ b/Engine/Plugins/ChaosClothAssetEditor/Source/ChaosClothAssetDataflowNodes/Public/ChaosClothAsset/SimulationDampingConfigNode.h
@@ -17,7 +17,7 @@ public:
 	 * The amount of global damping applied to the cloth velocities, also known as point damping.
 	 * Point damping improves simulation stability, but can also cause an overall slow-down effect and therefore is best left to very small percentage amounts.
 	 */
-	UPROPERTY(EditAnywhere, Category = "Damping Properties", Meta = (UIMin = "0", UIMax = "1", ClampMin = "0", ClampMax = "1"))
+	UPROPERTY(EditAnywhere, DisplayName = "Damping Coefficient", Category = "Damping Properties", Meta = (UIMin = "0", UIMax = "1", ClampMin = "0", ClampMax = "1", InteractorName = "DampingCoefficient"))
 	FChaosClothAssetWeightedValue DampingCoefficientWeighted = {true, 0.01f, 0.01f, TEXT("DampingCoefficient")};
 
 	/**
@@ -26,7 +26,7 @@ public:
 	 * It makes the cloth deformations more cohesive and reduces jitter without affecting the overall movement.
 	 * It can also produce synchronization artifacts where part of the cloth mesh are disconnected (which might well be desirable, or not), and is more expensive than global damping.
 	 */
-	UPROPERTY(EditAnywhere, Category = "Damping Properties", Meta = (UIMin = "0", UIMax = "1", ClampMin = "0", ClampMax = "1"))
+	UPROPERTY(EditAnywhere, Category = "Damping Properties", Meta = (UIMin = "0", UIMax = "1", ClampMin = "0", ClampMax = "1", InteractorName = "LocalDampingCoefficient"))
 	float LocalDampingCoefficient = 0.f;
 
 	FChaosClothAssetSimulationDampingConfigNode(const UE::Dataflow::FNodeParameters& InParam, FGuid InGuid = FGuid::NewGuid());
diff --git a/Engine/Plugins/ChaosClothAssetEditor/Source/ChaosClothAssetDataflowNodes/Public/ChaosClothAsset/SimulationGravityConfigNode.h b/Engine/Plugins/ChaosClothAssetEditor/Source/ChaosClothAssetDataflowNodes/Public/ChaosClothAsset/SimulationGravityConfigNode.h
--- a/Engine/Plugins/ChaosClothAssetEditor/Source/ChaosClothAssetDataflowNodes/Public/ChaosClothAsset/SimulationGravityConfigNode.h
+++ b/Engine/Plugins/ChaosClothAssetEditor/Source/ChaosClothAssetDataflowNodes/Public/ChaosClothAsset/SimulationGravityConfigNode.h
@@ -14,15 +14,15 @@ struct FChaosClothAssetSimulationGravityConfigNode : public FChaosClothAssetSimu
 
 public:
 	/** Use the config gravity value instead of world gravity. */
-	UPROPERTY(EditAnywhere, Category = "Gravity Properties", Meta = (InlineEditConditionToggle))
+	UPROPERTY(EditAnywhere, Category = "Gravity Properties", Meta = (InteractorName = "UseGravityOverride"))
 	bool bUseGravityOverride = false;
 
 	/** Scale factor applied to the world gravity and also to the clothing simulation interactor gravity. Does not affect the gravity if set using the override below. */
-	UPROPERTY(EditAnywhere, Category = "Gravity Properties", DisplayName = "Gravity Scale", Meta = (UIMin = "0", UIMax = "10", EditCondition = "!bUseGravityOverride"))
+	UPROPERTY(EditAnywhere, Category = "Gravity Properties", DisplayName = "Gravity Scale", Meta = (UIMin = "0", UIMax = "10", EditCondition = "!bUseGravityOverride", InteractorName = "GravityScale"))
 	FChaosClothAssetWeightedValue GravityScaleWeighted = {true, 1.f, 1.f, TEXT("GravityScale")};
 
 	/** The gravitational acceleration vector [cm/s^2]. */
-	UPROPERTY(EditAnywhere, Category = "Gravity Properties", DisplayName = "Gravity Override", Meta = (UIMin = "0", UIMax = "10", EditCondition = "bUseGravityOverride"))
+	UPROPERTY(EditAnywhere, Category = "Gravity Properties", DisplayName = "Gravity Override", Meta = (UIMin = "0", UIMax = "10", EditCondition = "bUseGravityOverride", InteractorName = "GravityOverride"))
 	FChaosClothAssetImportedVectorValue GravityOverrideImported = {UE::Chaos::ClothAsset::FDefaultSolver::Gravity};  // TODO: Should we make this a S.I. unit?
 	
 	FChaosClothAssetSimulationGravityConfigNode(const UE::Dataflow::FNodeParameters& InParam, FGuid InGuid = FGuid::NewGuid());
diff --git a/Engine/Plugins/ChaosClothAssetEditor/Source/ChaosClothAssetDataflowNodes/Public/ChaosClothAsset/SimulationLongRangeAttachmentConfigNode.h b/Engine/Plugins/ChaosClothAssetEditor/Source/ChaosClothAssetDataflowNodes/Public/ChaosClothAsset/SimulationLongRangeAttachmentConfigNode.h
--- a/Engine/Plugins/ChaosClothAssetEditor/Source/ChaosClothAssetDataflowNodes/Public/ChaosClothAsset/SimulationLongRangeAttachmentConfigNode.h
+++ b/Engine/Plugins/ChaosClothAssetEditor/Source/ChaosClothAssetDataflowNodes/Public/ChaosClothAsset/SimulationLongRangeAttachmentConfigNode.h
@@ -24,7 +24,7 @@ public:
 	 * are interpolated with the per particle weight to make the final value used for the simulation.
 	 * Otherwise all particles are considered to have a zero weight, and only the Low value is meaningful.
 	 */
-	UPROPERTY(EditAnywhere, Category = "Long Range Attachment Properties", Meta = (UIMin = "0", UIMax = "1", ClampMin = "0", ClampMax = "1"))
+	UPROPERTY(EditAnywhere, Category = "Long Range Attachment Properties", Meta = (UIMin = "0", UIMax = "1", ClampMin = "0", ClampMax = "1", InteractorName = "TetherStiffness"))
 	FChaosClothAssetWeightedValue TetherStiffness = { true, 1.f, 1.f, TEXT("TetherStiffness") };
 
 	/**
@@ -33,7 +33,7 @@ public:
 	 * are interpolated with the per particle weight to make the final value used for the simulation.
 	 * Otherwise all particles are considered to have a zero weight, and only the Low value is meaningful.
 	*/
-	UPROPERTY(EditAnywhere, Category = "Long Range Attachment Properties", Meta = (UIMin = "1.", UIMax = "1.1", ClampMin = "0.01", ClampMax = "10"))
+	UPROPERTY(EditAnywhere, Category = "Long Range Attachment Properties", Meta = (UIMin = "1.", UIMax = "1.1", ClampMin = "0.01", ClampMax = "10", InteractorName = "TetherScale"))
 	FChaosClothAssetWeightedValue TetherScale = { true, 1.f, 1.f, TEXT("TetherScale") };
 
 	/**
@@ -74,7 +74,7 @@ public:
 	 * are interpolated with the per particle weight to make the final value used for the simulation.
 	 * Otherwise all particles are considered to have a zero weight, and only the Low value is meaningful.
 	 */
-	UPROPERTY(EditAnywhere, Category = "Long Range Attachment Properties", Meta = (UIMin = "0", UIMax = "1", ClampMin = "0", ClampMax = "1"))
+	UPROPERTY(EditAnywhere, Category = "Long Range Attachment Properties", Meta = (UIMin = "0", UIMax = "1", ClampMin = "0", ClampMax = "1", InteractorName = "TetherStiffness"))
 	FChaosClothAssetWeightedValue TetherStiffness = { true, 1.f, 1.f, TEXT("TetherStiffness") };
 
 	/**
@@ -83,7 +83,7 @@ public:
 	 * are interpolated with the per particle weight to make the final value used for the simulation.
 	 * Otherwise all particles are considered to have a zero weight, and only the Low value is meaningful.
 	*/
-	UPROPERTY(EditAnywhere, Category = "Long Range Attachment Properties", Meta = (UIMin = "1.", UIMax = "1.1", ClampMin = "0.01", ClampMax = "10"))
+	UPROPERTY(EditAnywhere, Category = "Long Range Attachment Properties", Meta = (UIMin = "1.", UIMax = "1.1", ClampMin = "0.01", ClampMax = "10", InteractorName = "TetherScale"))
 	FChaosClothAssetWeightedValue TetherScale = { true, 1.f, 1.f, TEXT("TetherScale") };
 
 	/**
diff --git a/Engine/Plugins/ChaosClothAssetEditor/Source/ChaosClothAssetDataflowNodes/Public/ChaosClothAsset/SimulationMultiResConfigNode.h b/Engine/Plugins/ChaosClothAssetEditor/Source/ChaosClothAssetDataflowNodes/Public/ChaosClothAsset/SimulationMultiResConfigNode.h
--- a/Engine/Plugins/ChaosClothAssetEditor/Source/ChaosClothAssetDataflowNodes/Public/ChaosClothAsset/SimulationMultiResConfigNode.h
+++ b/Engine/Plugins/ChaosClothAssetEditor/Source/ChaosClothAssetDataflowNodes/Public/ChaosClothAsset/SimulationMultiResConfigNode.h
@@ -19,23 +19,23 @@ public:
 	bool bIsFineLOD = false;
 
 	/** This LOD is a coarse LOD for a finer LOD's multi-res simulation.*/
-	UPROPERTY(EditAnywhere, Category = "Multi-resolution", meta = (EditCondition = "!bIsFineLOD"))
+	UPROPERTY(EditAnywhere, Category = "Multi-resolution", meta = (EditCondition = "!bIsFineLOD", InteractorName = "IsCoarseMultiResLOD"))
 	bool bIsCoarseMultiResLOD = false;
 
 	/** Index of the coarse LOD for multi-res simulation. That LOD also needs a Multi Res Config node with "Is Coarse Multi Res LOD" = true.*/
-	UPROPERTY(EditAnywhere, Category = "Multi-resolution", meta = (EditCondition = "bIsFineLOD", EditConditionHides))
+	UPROPERTY(EditAnywhere, Category = "Multi-resolution", meta = (EditCondition = "bIsFineLOD", EditConditionHides, InteractorName = "MultiResCoarseLODIndex"))
 	int32 MultiResCoarseLODIndex = INDEX_NONE;
 
 	/** Use XPBD-style constraints */
-	UPROPERTY(EditAnywhere, Category = "Multi-resolution", meta = (EditCondition = "bIsFineLOD", EditConditionHides))
+	UPROPERTY(EditAnywhere, Category = "Multi-resolution", meta = (EditCondition = "bIsFineLOD", EditConditionHides, InteractorName = "MultiResUseXPBD"))
 	bool bMultiResUseXPBD = false;
 
 	/** MultiRes Spring Stiffness */
-	UPROPERTY(EditAnywhere, Category = "Multi-resolution", meta = (UIMin = "0", UIMax = "100000", ClampMin = "0", ClampMax = "100000", EditCondition = "bIsFineLOD", EditConditionHides))
+	UPROPERTY(EditAnywhere, Category = "Multi-resolution", meta = (UIMin = "0", UIMax = "100000", ClampMin = "0", ClampMax = "100000", EditCondition = "bIsFineLOD", EditConditionHides, InteractorName = "MultiResStiffness"))
 	FChaosClothAssetWeightedValue MultiResStiffness = {true, 1.f, 1.f};
 
 	/** MultiRes Velocity Target Spring Stiffness (non-XPBD only) */
-	UPROPERTY(EditAnywhere, Category = "Multi-resolution", meta = (UIMin = "0", UIMax = "1", ClampMin = "0", ClampMax = "1", EditCondition = "bIsFineLOD && !bMultiResUseXPBD", EditConditionHides))
+	UPROPERTY(EditAnywhere, Category = "Multi-resolution", meta = (UIMin = "0", UIMax = "1", ClampMin = "0", ClampMax = "1", EditCondition = "bIsFineLOD && !bMultiResUseXPBD", EditConditionHides, InteractorName = "MultiResVelocityTargetStiffness"))
 	FChaosClothAssetWeightedValue MultiResVelocityTargetStiffness = { true, 1.f, 1.f };
 
 	FChaosClothAssetSimulationMultiResConfigNode(const UE::Dataflow::FNodeParameters& InParam, FGuid InGuid = FGuid::NewGuid());
diff --git a/Engine/Plugins/ChaosClothAssetEditor/Source/ChaosClothAssetDataflowNodes/Public/ChaosClothAsset/SimulationPressureConfigNode.h b/Engine/Plugins/ChaosClothAssetEditor/Source/ChaosClothAssetDataflowNodes/Public/ChaosClothAsset/SimulationPressureConfigNode.h
--- a/Engine/Plugins/ChaosClothAssetEditor/Source/ChaosClothAssetDataflowNodes/Public/ChaosClothAsset/SimulationPressureConfigNode.h
+++ b/Engine/Plugins/ChaosClothAssetEditor/Source/ChaosClothAssetDataflowNodes/Public/ChaosClothAsset/SimulationPressureConfigNode.h
@@ -20,7 +20,7 @@ public:
 	 * are interpolated with the per particle weight to make the final value used for the simulation.
 	 * Otherwise all particles are considered to have a zero weight, and only the Low value is meaningful.
 	 */
-	UPROPERTY(EditAnywhere, Category = "Pressure Properties", meta = (UIMin = "-10", UIMax = "10", ClampMin = "-100", ClampMax = "100"))
+	UPROPERTY(EditAnywhere, Category = "Pressure Properties", meta = (UIMin = "-10", UIMax = "10", ClampMin = "-100", ClampMax = "100", InteractorName = "Pressure"))
 	FChaosClothAssetWeightedValue Pressure = { true, 0.0f, 1.f, TEXT("Pressure"),true };
 
 	FChaosClothAssetSimulationPressureConfigNode(const UE::Dataflow::FNodeParameters& InParam, FGuid InGuid = FGuid::NewGuid());
diff --git a/Engine/Plugins/ChaosClothAssetEditor/Source/ChaosClothAssetDataflowNodes/Public/ChaosClothAsset/SimulationSelfCollisionConfigNode.h b/Engine/Plugins/ChaosClothAssetEditor/Source/ChaosClothAssetDataflowNodes/Public/ChaosClothAsset/SimulationSelfCollisionConfigNode.h
--- a/Engine/Plugins/ChaosClothAssetEditor/Source/ChaosClothAssetDataflowNodes/Public/ChaosClothAsset/SimulationSelfCollisionConfigNode.h
+++ b/Engine/Plugins/ChaosClothAssetEditor/Source/ChaosClothAssetDataflowNodes/Public/ChaosClothAsset/SimulationSelfCollisionConfigNode.h
@@ -14,17 +14,21 @@ struct FChaosClothAssetSimulationSelfCollisionConfigNode : public FChaosClothAss
 	DATAFLOW_NODE_DEFINE_INTERNAL(FChaosClothAssetSimulationSelfCollisionConfigNode, "SimulationSelfCollisionConfig", "Cloth", "Cloth Simulation Self Collision Config")
 
 public:
+	/** Activating this node will enable self collisions.*/
+	UPROPERTY(VisibleAnywhere, Category = "Self-Collision Properties", meta = (InteractorName = "UseSelfCollisions"))
+	bool bUseSelfCollisions = true;
+
 	/** The self collision offset per side. Total thickness of cloth is 2x this value. */
-	UPROPERTY(EditAnywhere, Category = "Self-Collision Properties", meta = (UIMin = "0", UIMax = "100", ClampMin = "0", ClampMax = "1000"))
+	UPROPERTY(EditAnywhere, Category = "Self-Collision Properties", meta = (UIMin = "0", UIMax = "100", ClampMin = "0", ClampMax = "1000", InteractorName = "SelfCollisionThickness"))
 	FChaosClothAssetWeightedValue SelfCollisionThicknessWeighted = {true, UE::Chaos::ClothAsset::FDefaultFabric::SelfCollisionThickness,
 		UE::Chaos::ClothAsset::FDefaultFabric::SelfCollisionThickness, TEXT("SelfCollisionThickness"), true};
 
 	/** The stiffness of the springs used to control self collision (PBD Solver). */
-	UPROPERTY(EditAnywhere, Category = "Self-Collision Properties", meta = (UIMin = "0", UIMax = "1", ClampMin = "0", ClampMax = "1", EditCondition = "!bSelfCollideAgainstKinematicCollidersOnly"))
+	UPROPERTY(EditAnywhere, Category = "Self-Collision Properties", meta = (UIMin = "0", UIMax = "1", ClampMin = "0", ClampMax = "1", EditCondition = "!bSelfCollideAgainstKinematicCollidersOnly", InteractorName = "SelfCollisionStiffness"))
 	float SelfCollisionStiffness = 0.5f;
 
 	/** Friction coefficient for cloth - cloth interaction. */
-	UPROPERTY(EditAnywhere, Category = "Self-Collision Properties", DisplayName = "Self Collision Friction", meta = (UIMin = "0", UIMax = "1", ClampMin = "0", ClampMax = "1", EditCondition = "!bSelfCollideAgainstKinematicCollidersOnly"))
+	UPROPERTY(EditAnywhere, Category = "Self-Collision Properties", DisplayName = "Self Collision Friction", meta = (UIMin = "0", UIMax = "1", ClampMin = "0", ClampMax = "1", EditCondition = "!bSelfCollideAgainstKinematicCollidersOnly", InteractorName = "SelfCollisionFriction"))
 	FChaosClothAssetImportedFloatValue SelfCollisionFrictionImported = {UE::Chaos::ClothAsset::FDefaultFabric::SelfFriction};
 
 	/** Disabled neighbor collision ring. Collisions are disabled between vertices within this N-ring connectivity distance.*/
@@ -35,59 +39,59 @@ public:
 	* Faces labeled with -1 will collide normally without any layering behavior.
 	* Faces labeled with any other number will keep higher layer numbers outside lower layer numbers (outside = front facing normal direction).
 	*/
-	UPROPERTY(EditAnywhere, Category = "Self-Collision Properties", meta = (EditCondition = "!bSelfCollideAgainstKinematicCollidersOnly"))
+	UPROPERTY(EditAnywhere, Category = "Self-Collision Properties", meta = (EditCondition = "!bSelfCollideAgainstKinematicCollidersOnly", InteractorName = "SelfCollisionLayers"))
 	FChaosClothAssetConnectableIStringValue SelfCollisionLayers = { TEXT("SelfCollisionLayers"), true };
 
 	/** Sim face selection set of faces which should not self collide */
-	UPROPERTY(EditAnywhere, Category = "Self-Collision Properties")
+	UPROPERTY(EditAnywhere, Category = "Self-Collision Properties", meta = (InteractorName = "SelfCollisionDisabledFaces"))
 	FChaosClothAssetConnectableIStringValue SelfCollisionDisabledFaces = { TEXT("SelfCollisionDisabledFaces") };
 	
 	/** Collide only against kinematic colliders (no dynamic self collisions). Kinematic colliders do not do Self Intersections. They always collide against the front-face.*/
-	UPROPERTY(EditAnywhere, Category = "Self-Collision Properties - Kinematic Colliders")
+	UPROPERTY(EditAnywhere, Category = "Self-Collision Properties - Kinematic Colliders", meta = (InteractorName = "SelfCollideAgainstKinematicCollidersOnly"))
 	bool bSelfCollideAgainstKinematicCollidersOnly = false;
 
 	/** Sim face selection set of kinematic faces which should self collide. Kinematic colliders do not do Self Intersections. They always collide against the front-face. */
-	UPROPERTY(EditAnywhere, Category = "Self-Collision Properties - Kinematic Colliders", meta = (EditCondition = "!bSelfCollideAgainstAllKinematicVertices"))
+	UPROPERTY(EditAnywhere, Category = "Self-Collision Properties - Kinematic Colliders", meta = (EditCondition = "!bSelfCollideAgainstAllKinematicVertices", InteractorName = "SelfCollisionEnabledKinematicFaces"))
 	FChaosClothAssetConnectableIStringValue SelfCollisionEnabledKinematicFaces = { TEXT("SelfCollisionEnabledKinematicFaces") };
 
 	/** Thickness of kinematic colliders. Total offset between cloth and kinematic colliders is SelfCollisionThickness + SelfCollisionKinematicColliderThickness. */
-	UPROPERTY(EditAnywhere, Category = "Self-Collision Properties - Kinematic Colliders", meta = (UIMin = "0", UIMax = "100", ClampMin = "0", ClampMax = "1000"))
+	UPROPERTY(EditAnywhere, Category = "Self-Collision Properties - Kinematic Colliders", meta = (UIMin = "0", UIMax = "100", ClampMin = "0", ClampMax = "1000", InteractorName = "SelfCollisionKinematicColliderThickness"))
 	float SelfCollisionKinematicColliderThickness = 0.f;
 
 	/** The stiffness of the springs used to control self collision (PBD Solver). */
-	UPROPERTY(EditAnywhere, Category = "Self-Collision Properties - Kinematic Colliders", meta = (UIMin = "0", UIMax = "1", ClampMin = "0", ClampMax = "1"))
+	UPROPERTY(EditAnywhere, Category = "Self-Collision Properties - Kinematic Colliders", meta = (UIMin = "0", UIMax = "1", ClampMin = "0", ClampMax = "1", InteractorName = "SelfCollisionKinematicColliderStiffness"))
 	float SelfCollisionKinematicColliderStiffness = 1.f;
 
 	/** Friction coefficient for cloth - kinematic cloth interaction. Weight map is on the dynamic cloth, not the collider.*/
-	UPROPERTY(EditAnywhere, Category = "Self-Collision Properties - Kinematic Colliders", meta = (UIMin = "0", UIMax = "1", ClampMin = "0", ClampMax = "1"))
+	UPROPERTY(EditAnywhere, Category = "Self-Collision Properties - Kinematic Colliders", meta = (UIMin = "0", UIMax = "1", ClampMin = "0", ClampMax = "1", InteractorName = "SelfCollisionKinematicColliderFriction"))
 	FChaosClothAssetWeightedValue SelfCollisionKinematicColliderFrictionWeighted = {true, 0.0f, 0.f, TEXT("SelfCollisionKinematicColliderFriction")};
 
 	/** Self collide against all kinematic vertices. Kinematic colliders do not do Self Intersections. They always collide against the front-face.*/
-	UPROPERTY(EditAnywhere, Category = "Self-Collision Properties - Kinematic Colliders")
+	UPROPERTY(EditAnywhere, Category = "Self-Collision Properties - Kinematic Colliders", meta = (InteractorName = "SelfCollideAgainstAllKinematicVertices"))
 	bool bSelfCollideAgainstAllKinematicVertices = false;
 
 	/** Enable self intersection resolution. This will try to fix any cloth intersections that are not handled by collision repulsions. */
-	UPROPERTY(EditAnywhere, Category = "Experimental", meta = (EditCondition = "!bSelfCollideAgainstKinematicCollidersOnly"))
+	UPROPERTY(EditAnywhere, Category = "Experimental", meta = (EditCondition = "!bSelfCollideAgainstKinematicCollidersOnly", InteractorName = "UseSelfIntersections"))
 	bool bUseSelfIntersections = false;
 
 	/** Do global intersection analysis to determine the correct normals for the collision springs */
-	UPROPERTY(EditAnywhere, Category = "Experimental", meta = (EditCondition = "bUseSelfIntersections && !bSelfCollideAgainstKinematicCollidersOnly"))
+	UPROPERTY(EditAnywhere, Category = "Experimental", meta = (EditCondition = "bUseSelfIntersections && !bSelfCollideAgainstKinematicCollidersOnly", InteractorName = "UseGlobalIntersectionAnalysis"))
 	bool bUseGlobalIntersectionAnalysis = true;
 
 	/** Do a step of contour minimization at the beginning of the timestep. */
-	UPROPERTY(EditAnywhere, Category = "Experimental", meta = (EditCondition = "bUseSelfIntersections && !bSelfCollideAgainstKinematicCollidersOnly"))
+	UPROPERTY(EditAnywhere, Category = "Experimental", meta = (EditCondition = "bUseSelfIntersections && !bSelfCollideAgainstKinematicCollidersOnly", InteractorName = "UseContourMinimization"))
 	bool bUseContourMinimization = true;
 
 	/** Number of post timestep contour minimization steps to do. (Expensive!)*/
-	UPROPERTY(EditAnywhere, Category = "Experimental", meta = (ClampMin = "0", EditCondition = "bUseSelfIntersections && !bSelfCollideAgainstKinematicCollidersOnly"))
+	UPROPERTY(EditAnywhere, Category = "Experimental", meta = (ClampMin = "0", EditCondition = "bUseSelfIntersections && !bSelfCollideAgainstKinematicCollidersOnly", InteractorName = "NumContourMinimizationPostSteps"))
 	int32 NumContourMinimizationPostSteps = 0;
 
 	/** Use global contour gradients when doing post timestep contour minimization */
-	UPROPERTY(EditAnywhere, Category = "Experimental", meta = (EditCondition = "bUseSelfIntersections && NumContourMinimizationPostSteps > 0 && !bSelfCollideAgainstKinematicCollidersOnly"))
+	UPROPERTY(EditAnywhere, Category = "Experimental", meta = (EditCondition = "bUseSelfIntersections && NumContourMinimizationPostSteps > 0 && !bSelfCollideAgainstKinematicCollidersOnly", InteractorName = "UseGlobalPostStepContours"))
 	bool bUseGlobalPostStepContours = true;
 
 	/** The stiffness of the proximity repulsions used to control self collision (Force-based Solver). Units = kg cm/ s^2 (same as XPBD springs) */
-	UPROPERTY(EditAnywhere, Category = "Experimental", meta = (UIMin = "0", UIMax = "10000", ClampMin = "0", ClampMax = "10000000", EditCondition = "!bSelfCollideAgainstAllKinematicVertices"))
+	UPROPERTY(EditAnywhere, Category = "Experimental", meta = (UIMin = "0", UIMax = "10000", ClampMin = "0", ClampMax = "10000000", EditCondition = "!bSelfCollideAgainstAllKinematicVertices", InteractorName = "SelfCollisionProximityStiffness"))
 	float SelfCollisionProximityStiffness = 1.f;
 
 	FChaosClothAssetSimulationSelfCollisionConfigNode(const UE::Dataflow::FNodeParameters& InParam, FGuid InGuid = FGuid::NewGuid());
diff --git a/Engine/Plugins/ChaosClothAssetEditor/Source/ChaosClothAssetDataflowNodes/Public/ChaosClothAsset/SimulationSelfCollisionSpheresConfigNode.h b/Engine/Plugins/ChaosClothAssetEditor/Source/ChaosClothAssetDataflowNodes/Public/ChaosClothAsset/SimulationSelfCollisionSpheresConfigNode.h
--- a/Engine/Plugins/ChaosClothAssetEditor/Source/ChaosClothAssetDataflowNodes/Public/ChaosClothAsset/SimulationSelfCollisionSpheresConfigNode.h
+++ b/Engine/Plugins/ChaosClothAssetEditor/Source/ChaosClothAssetDataflowNodes/Public/ChaosClothAsset/SimulationSelfCollisionSpheresConfigNode.h
@@ -14,11 +14,11 @@ struct FChaosClothAssetSimulationSelfCollisionSpheresConfigNode : public FChaosC
 
 public:
 	/** The radius of the spheres used in self collision centered at each vertex. */
-	UPROPERTY(EditAnywhere, Category = "Self-Collision Spheres Properties", meta = (UIMin = "0", UIMax = "100", ClampMin = "0", ClampMax = "1000"))
+	UPROPERTY(EditAnywhere, Category = "Self-Collision Spheres Properties", meta = (UIMin = "0", UIMax = "100", ClampMin = "0", ClampMax = "1000", InteractorName = "SelfCollisionSphereRadius"))
 	float SelfCollisionSphereRadius = 0.5f;
 
 	/** The stiffness of the springs used to control self collision. */
-	UPROPERTY(EditAnywhere, Category = "Self-Collision Spheres Properties", meta = (UIMin = "0", UIMax = "100", ClampMin = "0", ClampMax = "1000"))
+	UPROPERTY(EditAnywhere, Category = "Self-Collision Spheres Properties", meta = (UIMin = "0", UIMax = "100", ClampMin = "0", ClampMax = "1000", InteractorName = "SelfCollisionSphereStiffness"))
 	float SelfCollisionSphereStiffness = 1.f;
 
 	/** 
diff --git a/Engine/Plugins/ChaosClothAssetEditor/Source/ChaosClothAssetDataflowNodes/Public/ChaosClothAsset/SimulationSolverConfigNode.h b/Engine/Plugins/ChaosClothAssetEditor/Source/ChaosClothAssetDataflowNodes/Public/ChaosClothAsset/SimulationSolverConfigNode.h
--- a/Engine/Plugins/ChaosClothAssetEditor/Source/ChaosClothAssetDataflowNodes/Public/ChaosClothAsset/SimulationSolverConfigNode.h
+++ b/Engine/Plugins/ChaosClothAssetEditor/Source/ChaosClothAssetDataflowNodes/Public/ChaosClothAsset/SimulationSolverConfigNode.h
@@ -19,21 +19,21 @@ public:
 	 * At lower fps up to MaxNumIterations may be used instead. At higher fps as low as one single iteration might be used.
 	 * Higher number of iterations will increase the stiffness of all constraints and improve convergence, but will also increase the CPU cost of the simulation.
 	 */
-	UPROPERTY(EditAnywhere, Category = Simulation, meta = (UIMin = "1", UIMax = "10", ClampMin = "0", ClampMax = "100"))
+	UPROPERTY(EditAnywhere, Category = Simulation, meta = (UIMin = "1", UIMax = "10", ClampMin = "0", ClampMax = "100", InteractorName = "NumIterations"))
 	int32 NumIterations = 1;
 
 	/**
 	 * The maximum number of solver iterations.
 	 * This is the upper limit of the number of iterations set in solver, when the frame rate is lower than 60fps.
 	 */
-	UPROPERTY(EditAnywhere, Category = Simulation, meta = (UIMin = "1", UIMax = "10", ClampMin = "1", ClampMax = "100"))
+	UPROPERTY(EditAnywhere, Category = Simulation, meta = (UIMin = "1", UIMax = "10", ClampMin = "1", ClampMax = "100", InteractorName = "MaxNumIterations"))
 	int32 MaxNumIterations = 6;
 
 	/**
 	 * The number of solver substeps.
 	 * This will increase the precision of the collision inputs and help with constraint resolutions but will increase the CPU cost.
 	 */
-	UPROPERTY(EditAnywhere, Category = Simulation, DisplayName = "Num Substeps", meta = (UIMin = "1", UIMax = "10", ClampMin = "1", ClampMax = "100"))
+	UPROPERTY(EditAnywhere, Category = Simulation, DisplayName = "Num Substeps", meta = (UIMin = "1", UIMax = "10", ClampMin = "1", ClampMax = "100", InteractorName = "NumSubsteps"))
 	FChaosClothAssetImportedIntValue NumSubstepsImported = {UE::Chaos::ClothAsset::FDefaultSolver::SubSteps};
 
 	/**
@@ -45,20 +45,20 @@ public:
 	/**
 	 * Choose the number of substeps based on a target substep delta time in milliseconds. Substeps are clamped to [1, NumSubsteps]. 
 	 */
-	UPROPERTY(EditAnywhere, Category = Simulation, meta = (UIMin = "1", UIMax = "30", ClampMin = "0", ClampMax = "1000", EditCondition = "bEnableDynamicSubstepping"))
+	UPROPERTY(EditAnywhere, Category = Simulation, meta = (UIMin = "1", UIMax = "30", ClampMin = "0", ClampMax = "1000", EditCondition = "bEnableDynamicSubstepping", InteractorName = "DynamicSubstepDeltaTime"))
 	float DynamicSubstepDeltaTime = 16.67f;
 
 	/**
 	* Enable setting separate SelfCollisionSubsteps. Otherwise, self collisions will be detected every substep.
 	*/
-	UPROPERTY(EditAnywhere, Category = Simulation, meta = (InlineEditConditionToggle))
+	UPROPERTY(EditAnywhere, Category = Simulation, meta = (InteractorName = "EnableNumSelfCollisionSubsteps"))
 	bool bEnableNumSelfCollisionSubsteps = false;
 
 	/**
 	 * Set a separate number of self collision substeps. Lower this number to increase speed at the expense of lower self collision accuracy.
 	 * Actual value always clamped between [1, NumSubsteps].
 	 */
-	UPROPERTY(EditAnywhere, Category = Simulation, meta = (UIMin = "1", UIMax = "10", ClampMin = "1", ClampMax = "100", EditCondition = "bEnableNumSelfCollisionSubsteps"))
+	UPROPERTY(EditAnywhere, Category = Simulation, meta = (UIMin = "1", UIMax = "10", ClampMin = "1", ClampMax = "100", EditCondition = "bEnableNumSelfCollisionSubsteps", InteractorName = "NumSelfCollisionSubsteps"))
 	int32 NumSelfCollisionSubsteps = 1;
 
 	FChaosClothAssetSimulationSolverConfigNode(const UE::Dataflow::FNodeParameters& InParam, FGuid InGuid = FGuid::NewGuid());
diff --git a/Engine/Plugins/ChaosClothAssetEditor/Source/ChaosClothAssetDataflowNodes/Public/ChaosClothAsset/SimulationStretchConfigNode.h b/Engine/Plugins/ChaosClothAssetEditor/Source/ChaosClothAssetDataflowNodes/Public/ChaosClothAsset/SimulationStretchConfigNode.h
--- a/Engine/Plugins/ChaosClothAssetEditor/Source/ChaosClothAssetDataflowNodes/Public/ChaosClothAsset/SimulationStretchConfigNode.h
+++ b/Engine/Plugins/ChaosClothAssetEditor/Source/ChaosClothAssetDataflowNodes/Public/ChaosClothAsset/SimulationStretchConfigNode.h
@@ -37,7 +37,7 @@ public:
 	 * are interpolated with the per particle weight to make the final value used for the simulation.
 	 * Otherwise all particles are considered to have a zero weight, and only the Low value is meaningful.
 	 */
-	UPROPERTY(EditAnywhere, Category = "Stretch Properties", Meta = (UIMin = "0", UIMax = "10000", ClampMin = "0", ClampMax = "10000000", EditCondition = "DistributionType == EChaosClothAssetConstraintDistributionType::Anisotropic && SolverType == EChaosClothAssetConstraintSolverType::XPBD", EditConditionHides))
+	UPROPERTY(EditAnywhere, Category = "Stretch Properties", Meta = (UIMin = "0", UIMax = "10000", ClampMin = "0", ClampMax = "10000000", EditCondition = "DistributionType == EChaosClothAssetConstraintDistributionType::Anisotropic && SolverType == EChaosClothAssetConstraintSolverType::XPBD", EditConditionHides, InteractorName = "StretchStiffnessWarp"))
 	FChaosClothAssetWeightedValue StretchStiffnessWarp = { true, 100.f, 100.f, TEXT("StretchStiffnessWarp"),true };
 
 	/**
@@ -46,7 +46,7 @@ public:
 	 * are interpolated with the per particle weight to make the final value used for the simulation.
 	 * Otherwise all particles are considered to have a zero weight, and only the Low value is meaningful.
 	 */
-	UPROPERTY(EditAnywhere, Category = "Stretch Properties", Meta = (UIMin = "0", UIMax = "10000", ClampMin = "0", ClampMax = "10000000", EditCondition = "DistributionType == EChaosClothAssetConstraintDistributionType::Anisotropic && SolverType == EChaosClothAssetConstraintSolverType::XPBD", EditConditionHides))
+	UPROPERTY(EditAnywhere, Category = "Stretch Properties", Meta = (UIMin = "0", UIMax = "10000", ClampMin = "0", ClampMax = "10000000", EditCondition = "DistributionType == EChaosClothAssetConstraintDistributionType::Anisotropic && SolverType == EChaosClothAssetConstraintSolverType::XPBD", EditConditionHides, InteractorName = "StretchStiffnessWeft"))
 	FChaosClothAssetWeightedValue StretchStiffnessWeft = { true, 100.f, 100.f, TEXT("StretchStiffnessWeft"), true };
 
 	/**
@@ -55,7 +55,7 @@ public:
 	 * are interpolated with the per particle weight to make the final value used for the simulation.
 	 * Otherwise all particles are considered to have a zero weight, and only the Low value is meaningful.
 	 */
-	UPROPERTY(EditAnywhere, Category = "Stretch Properties", Meta = (UIMin = "0", UIMax = "10000", ClampMin = "0", ClampMax = "10000000", EditCondition = "DistributionType == EChaosClothAssetConstraintDistributionType::Anisotropic && SolverType == EChaosClothAssetConstraintSolverType::XPBD", EditConditionHides))
+	UPROPERTY(EditAnywhere, Category = "Stretch Properties", Meta = (UIMin = "0", UIMax = "10000", ClampMin = "0", ClampMax = "10000000", EditCondition = "DistributionType == EChaosClothAssetConstraintDistributionType::Anisotropic && SolverType == EChaosClothAssetConstraintSolverType::XPBD", EditConditionHides, InteractorName = "StretchStiffnessBias"))
 	FChaosClothAssetWeightedValue StretchStiffnessBias = { true, 100.f, 100.f, TEXT("StretchStiffnessBias"), true };
 
 	/**
@@ -65,7 +65,7 @@ public:
 	 * are interpolated with the per particle weight to make the final value used for the simulation.
 	 * Otherwise all particles are considered to have a zero weight, and only the Low value is meaningful.
 	 */
-	UPROPERTY(EditAnywhere, Category = "Stretch Properties", Meta = (UIMin = "0", UIMax = "10", ClampMin = "0", ClampMax = "1000", EditCondition = "SolverType == EChaosClothAssetConstraintSolverType::XPBD && DistributionType == EChaosClothAssetConstraintDistributionType::Anisotropic", EditConditionHides))
+	UPROPERTY(EditAnywhere, Category = "Stretch Properties", Meta = (UIMin = "0", UIMax = "10", ClampMin = "0", ClampMax = "1000", EditCondition = "SolverType == EChaosClothAssetConstraintSolverType::XPBD && DistributionType == EChaosClothAssetConstraintDistributionType::Anisotropic", EditConditionHides, InteractorName = "StretchDamping"))
 	FChaosClothAssetWeightedValue StretchAnisoDamping = { true, 1.f, 1.f, TEXT("StretchAnisoDamping"), true };
 	
 	/**
@@ -74,7 +74,7 @@ public:
      * are interpolated with the per particle weight to make the final value used for the simulation.
      * Otherwise all particles are considered to have a zero weight, and only the Low value is meaningful.
      */
-    UPROPERTY(EditAnywhere, Category = "Stretch Properties", Meta = (UIMin = "0", UIMax = "10000", ClampMin = "0", ClampMax = "1000000000", EditCondition = "(SolverType == EChaosClothAssetConstraintSolverType::PBD || DistributionType == EChaosClothAssetConstraintDistributionType::Isotropic)", EditConditionHides))
+    UPROPERTY(EditAnywhere, Category = "Stretch Properties", Meta = (UIMin = "0", UIMax = "10000", ClampMin = "0", ClampMax = "1000000000", EditCondition = "(SolverType == EChaosClothAssetConstraintSolverType::PBD || DistributionType == EChaosClothAssetConstraintDistributionType::Isotropic)", EditConditionHides, InteractorName = "StretchStiffness"))
     FChaosClothAssetWeightedValue StretchStiffness = { true, 1.f, 1.f, TEXT("StretchStiffness") };
 	
 	/**
@@ -84,7 +84,7 @@ public:
 	 * are interpolated with the per particle weight to make the final value used for the simulation.
 	 * Otherwise all particles are considered to have a zero weight, and only the Low value is meaningful.
 	 */
-	UPROPERTY(EditAnywhere, Category = "Stretch Properties", Meta = (UIMin = "0", UIMax = "10", ClampMin = "0", ClampMax = "1000", EditCondition = "SolverType == EChaosClothAssetConstraintSolverType::XPBD && DistributionType == EChaosClothAssetConstraintDistributionType::Isotropic", EditConditionHides))
+	UPROPERTY(EditAnywhere, Category = "Stretch Properties", Meta = (UIMin = "0", UIMax = "10", ClampMin = "0", ClampMax = "1000", EditCondition = "SolverType == EChaosClothAssetConstraintSolverType::XPBD && DistributionType == EChaosClothAssetConstraintDistributionType::Isotropic", EditConditionHides, InteractorName = "StretchDamping"))
 	FChaosClothAssetWeightedValue StretchDamping = { true, 1.f, 1.f, TEXT("StretchDamping") };
 
 	/**
@@ -93,7 +93,7 @@ public:
 	 * are interpolated with the per particle weight to make the final value used for the simulation.
 	 * Otherwise all particles are considered to have a zero weight, and only the Low value is meaningful.
 	 */
-	UPROPERTY(EditAnywhere, Category = "Stretch Properties", Meta = (UIMin = "0", UIMax = "10", ClampMin = "0", ClampMax = "10000000", EditCondition = "(SolverType == EChaosClothAssetConstraintSolverType::XPBD && DistributionType == EChaosClothAssetConstraintDistributionType::Anisotropic)", EditConditionHides))
+	UPROPERTY(EditAnywhere, Category = "Stretch Properties", Meta = (UIMin = "0", UIMax = "10", ClampMin = "0", ClampMax = "10000000", EditCondition = "(SolverType == EChaosClothAssetConstraintSolverType::XPBD && DistributionType == EChaosClothAssetConstraintDistributionType::Anisotropic)", EditConditionHides, InteractorName = "StretchWarpScale"))
 	FChaosClothAssetWeightedValue StretchWarpScale = { true, 1.f, 1.f, TEXT("StretchWarpScale") };
 
 	/**
@@ -102,7 +102,7 @@ public:
 	 * are interpolated with the per particle weight to make the final value used for the simulation.
 	 * Otherwise all particles are considered to have a zero weight, and only the Low value is meaningful.
 	 */
-	UPROPERTY(EditAnywhere, Category = "Stretch Properties", Meta = (UIMin = "0", UIMax = "10", ClampMin = "0", ClampMax = "10000000", EditCondition = "(SolverType == EChaosClothAssetConstraintSolverType::XPBD && DistributionType == EChaosClothAssetConstraintDistributionType::Anisotropic)", EditConditionHides))
+	UPROPERTY(EditAnywhere, Category = "Stretch Properties", Meta = (UIMin = "0", UIMax = "10", ClampMin = "0", ClampMax = "10000000", EditCondition = "(SolverType == EChaosClothAssetConstraintSolverType::XPBD && DistributionType == EChaosClothAssetConstraintDistributionType::Anisotropic)", EditConditionHides, InteractorName = "StretchWeftScale"))
 	FChaosClothAssetWeightedValue StretchWeftScale = { true, 1.f, 1.f, TEXT("StretchWeftScale") };
 
 	/**
@@ -112,7 +112,7 @@ public:
 	 * are interpolated with the per particle weight to make the final value used for the simulation.
 	 * Otherwise all particles are considered to have a zero weight, and only the Low value is meaningful.
 	 */
-	UPROPERTY(EditAnywhere, Category = "Stretch Properties", Meta = (UIMin = "0", UIMax = "10000", ClampMin = "0", ClampMax = "10000000", EditCondition = "bAddAreaConstraint == true && ((SolverType == EChaosClothAssetConstraintSolverType::XPBD && DistributionType == EChaosClothAssetConstraintDistributionType::Isotropic) || SolverType == EChaosClothAssetConstraintSolverType::PBD)", EditConditionHides))
+	UPROPERTY(EditAnywhere, Category = "Stretch Properties", Meta = (UIMin = "0", UIMax = "10000", ClampMin = "0", ClampMax = "10000000", EditCondition = "bAddAreaConstraint == true && ((SolverType == EChaosClothAssetConstraintSolverType::XPBD && DistributionType == EChaosClothAssetConstraintDistributionType::Isotropic) || SolverType == EChaosClothAssetConstraintSolverType::PBD)", EditConditionHides, InteractorName = "AreaStiffness"))
 	FChaosClothAssetWeightedValue AreaStiffness = { true, 1.f, 1.f, TEXT("AreaStiffness") };
 
 	FChaosClothAssetSimulationStretchConfigNode(const UE::Dataflow::FNodeParameters& InParam, FGuid InGuid = FGuid::NewGuid());
diff --git a/Engine/Plugins/ChaosClothAssetEditor/Source/ChaosClothAssetDataflowNodes/Public/ChaosClothAsset/SimulationVelocityScaleConfigNode.h b/Engine/Plugins/ChaosClothAssetEditor/Source/ChaosClothAssetDataflowNodes/Public/ChaosClothAsset/SimulationVelocityScaleConfigNode.h
--- a/Engine/Plugins/ChaosClothAssetEditor/Source/ChaosClothAssetDataflowNodes/Public/ChaosClothAsset/SimulationVelocityScaleConfigNode.h
+++ b/Engine/Plugins/ChaosClothAssetEditor/Source/ChaosClothAssetDataflowNodes/Public/ChaosClothAsset/SimulationVelocityScaleConfigNode.h
@@ -20,7 +20,7 @@ public:
 	 * All vector properties on this node (e.g., Linear Velocity Scale, Max Linear Acceleration)
 	 * will be evaluated in this space. 
 	 */
-	UPROPERTY(EditAnywhere, Category = "Velocity Scale Properties")
+	UPROPERTY(EditAnywhere, Category = "Velocity Scale Properties", meta = (InteractorName = "VelocityScaleSpace"))
 	EChaosSoftsSimulationSpace VelocityScaleSpace = EChaosSoftsSimulationSpace::ReferenceBoneSpace;
 
 	/**
@@ -28,7 +28,7 @@ public:
 	 * (the closest bone to the root on which the cloth section has been skinned, or the root itself if the cloth isn't skinned).
 	 * This value will be clamped by "Max Velocity Scale". A velocity scale of > 1 will amplify the velocities from the reference bone.
 	 */
-	UPROPERTY(EditAnywhere, Category = "Velocity Scale Properties", meta = (UIMin = "0", UIMax = "1", ClampMin = "0", ClampMax = "100"))
+	UPROPERTY(EditAnywhere, Category = "Velocity Scale Properties", meta = (UIMin = "0", UIMax = "1", ClampMin = "0", ClampMax = "100", InteractorName = "LinearVelocityScale"))
 	FVector3f LinearVelocityScale = { 0.75f, 0.75f, 0.75f };
 
 	/**
@@ -41,7 +41,7 @@ public:
 	 * The maximum amount of linear velocity sent to the local cloth space from the reference bone
 	 * (the closest bone to the root on which the cloth section has been skinned, or the root itself if the cloth isn't skinned).
 	 */
-	UPROPERTY(EditAnywhere, Category = "Velocity Scale Properties", meta = (ClampMin = "0", EditCondition = "bEnableLinearVelocityClamping"))
+	UPROPERTY(EditAnywhere, Category = "Velocity Scale Properties", meta = (ClampMin = "0", EditCondition = "bEnableLinearVelocityClamping", InteractorName = "MaxLinearVelocity"))
 	FVector3f MaxLinearVelocity = { 1000.f, 1000.f, 1000.f }; // Approx 22mph or 36kph per direction
 
 	/**
@@ -54,7 +54,7 @@ public:
 	 * The maximum amount of linear acceleration sent to the local cloth space from the reference bone
 	 * (the closest bone to the root on which the cloth section has been skinned, or the root itself if the cloth isn't skinned).
 	 */
-	UPROPERTY(EditAnywhere, Category = "Velocity Scale Properties", meta = (ClampMin = "0", EditCondition = "bEnableLinearAccelerationClamping"))
+	UPROPERTY(EditAnywhere, Category = "Velocity Scale Properties", meta = (ClampMin = "0", EditCondition = "bEnableLinearAccelerationClamping", InteractorName = "MaxLinearAcceleration"))
 	FVector3f MaxLinearAcceleration = { 60000.f, 60000.f, 60000.f };
 
 	/**
@@ -62,7 +62,7 @@ public:
 	 * (the closest bone to the root on which the cloth section has been skinned, or the root itself if the cloth isn't skinned).
 	 * This value will be clamped by "Max Velocity Scale". A velocity scale of > 1 will amplify the velocities from the reference bone.
 	 */
-	UPROPERTY(EditAnywhere, Category = "Velocity Scale Properties", meta = (UIMin = "0", UIMax = "1", ClampMin = "0", ClampMax = "100"))
+	UPROPERTY(EditAnywhere, Category = "Velocity Scale Properties", meta = (UIMin = "0", UIMax = "1", ClampMin = "0", ClampMax = "100", InteractorName = "AngularVelocityScale"))
 	float AngularVelocityScale = 0.75f;
 
 	/**
@@ -75,7 +75,7 @@ public:
 	 * The maximum amount of angular velocity sent to the local cloth space from the reference bone
 	 * (the closest bone to the root on which the cloth section has been skinned, or the root itself if the cloth isn't skinned).
 	 */
-	UPROPERTY(EditAnywhere, Category = "Velocity Scale Properties", meta = (ClampMin = "0", EditCondition = "bEnableAngularVelocityClamping"))
+	UPROPERTY(EditAnywhere, Category = "Velocity Scale Properties", meta = (ClampMin = "0", EditCondition = "bEnableAngularVelocityClamping", InteractorName = "MaxAngularVelocity"))
 	float MaxAngularVelocity = 200.f;
 
 	/**
@@ -88,13 +88,13 @@ public:
 	 * The maximum amount of angular acceleration sent to the local cloth space from the reference bone
 	 * (the closest bone to the root on which the cloth section has been skinned, or the root itself if the cloth isn't skinned).
 	 */
-	UPROPERTY(EditAnywhere, Category = "Velocity Scale Properties", meta = (ClampMin = "0", EditCondition = "bEnableAngularAccelerationClamping"))
+	UPROPERTY(EditAnywhere, Category = "Velocity Scale Properties", meta = (ClampMin = "0", EditCondition = "bEnableAngularAccelerationClamping", InteractorName = "MaxAngularAcceleration"))
 	float MaxAngularAcceleration = 12000.f;
 
 	/**
 	 * Clamp on Linear and Angular Velocity Scale. The final velocity scale (e.g., including contributions from blueprints) will be clamped to this value.
 	 */
-	UPROPERTY(EditAnywhere, Category = "Velocity Scale Properties", meta = (UIMin = "0", UIMax = "1", ClampMin = "0", ClampMax = "100"))
+	UPROPERTY(EditAnywhere, Category = "Velocity Scale Properties", meta = (UIMin = "0", UIMax = "1", ClampMin = "0", ClampMax = "100", InteractorName = "MaxVelocityScale"))
 	float MaxVelocityScale = 1.f;
 
 	/**
@@ -105,7 +105,7 @@ public:
 	 * velocities and subjected to the true physical world inertial forces.
 	 * Values range from 0 to 2, with 0 showing no centrifugal effect, 1 full centrifugal effect, and 2 an overdriven centrifugal effect.
 	 */
-	UPROPERTY(EditAnywhere, Category = "Velocity Scale Properties", meta = (UIMin = "0", UIMax = "1", ClampMin = "0", ClampMax = "2"))
+	UPROPERTY(EditAnywhere, Category = "Velocity Scale Properties", meta = (UIMin = "0", UIMax = "1", ClampMin = "0", ClampMax = "2", InteractorName = "FictitiousAngularScale"))
 	float FictitiousAngularScale = 1.f;
 
 	FChaosClothAssetSimulationVelocityScaleConfigNode(const UE::Dataflow::FNodeParameters& InParam, FGuid InGuid = FGuid::NewGuid());
diff --git a/Engine/Plugins/ChaosClothAssetEditor/Source/ChaosClothAssetEditor/ChaosClothAssetEditor.Build.cs b/Engine/Plugins/ChaosClothAssetEditor/Source/ChaosClothAssetEditor/ChaosClothAssetEditor.Build.cs
--- a/Engine/Plugins/ChaosClothAssetEditor/Source/ChaosClothAssetEditor/ChaosClothAssetEditor.Build.cs
+++ b/Engine/Plugins/ChaosClothAssetEditor/Source/ChaosClothAssetEditor/ChaosClothAssetEditor.Build.cs
@@ -18,6 +18,7 @@ public class ChaosClothAssetEditor : ModuleRules
 		PrivateDependencyModuleNames.AddRange(
 			new string[]
 			{
+				"ApplicationCore",
 				"BaseCharacterFXEditor",
 				"Core",
 				"UnrealEd",
diff --git a/Engine/Plugins/ChaosClothAssetEditor/Source/ChaosClothAssetEditor/Private/ChaosClothAsset/ClothEditorToolkit.cpp b/Engine/Plugins/ChaosClothAssetEditor/Source/ChaosClothAssetEditor/Private/ChaosClothAsset/ClothEditorToolkit.cpp
--- a/Engine/Plugins/ChaosClothAssetEditor/Source/ChaosClothAssetEditor/Private/ChaosClothAsset/ClothEditorToolkit.cpp
+++ b/Engine/Plugins/ChaosClothAssetEditor/Source/ChaosClothAssetEditor/Private/ChaosClothAsset/ClothEditorToolkit.cpp
@@ -15,6 +15,7 @@
 #include "ChaosClothAsset/ClothEditorModeToolkit.h"
 #include "ChaosClothAsset/CollectionClothFacade.h"
 #include "ChaosClothAsset/ClothEditorCommands.h"
+#include "ChaosClothAsset/ClothSimulationNodeDetailExtender.h"
 #include "ChaosClothAsset/AddWeightMapNode.h"
 #include "ChaosClothAsset/TransferSkinWeightsNode.h"
 #include "EdModeInteractiveToolsContext.h"
@@ -1048,6 +1049,9 @@ TSharedPtr<IStructureDetailsView> FChaosClothAssetEditorToolkit::CreateNodeDetai
 	NodeDetailsView->GetDetailsView()->SetObject(ObjectToEdit);
 	NodeDetailsView->GetOnFinishedChangingPropertiesDelegate().AddSP(this, &FChaosClothAssetEditorToolkit::OnPropertyValueChanged);
 
+	NodeDetailsExtender = MakeShared<FClothSimulationNodeDetailExtender>();
+	NodeDetailsView->GetDetailsView()->SetExtensionHandler(NodeDetailsExtender);
+
 	return NodeDetailsView;
 }
 
diff --git a/Engine/Plugins/ChaosClothAssetEditor/Source/ChaosClothAssetEditor/Private/ChaosClothAsset/ClothEditorToolkit.h b/Engine/Plugins/ChaosClothAssetEditor/Source/ChaosClothAssetEditor/Private/ChaosClothAsset/ClothEditorToolkit.h
--- a/Engine/Plugins/ChaosClothAssetEditor/Source/ChaosClothAssetEditor/Private/ChaosClothAsset/ClothEditorToolkit.h
+++ b/Engine/Plugins/ChaosClothAssetEditor/Source/ChaosClothAssetEditor/Private/ChaosClothAsset/ClothEditorToolkit.h
@@ -22,6 +22,7 @@ namespace UE::Chaos::ClothAsset
 {
 class FClothEditorSimulationVisualization;
 class FChaosClothAssetEditor3DViewportClient;
+struct FClothSimulationNodeDetailExtender;
 }
 
 namespace UE::Dataflow
@@ -177,6 +178,7 @@ private:
 	static const FName NodeDetailsTabId;
 	TSharedPtr<SDockTab> NodeDetailsTab;
 	TSharedPtr<IStructureDetailsView> NodeDetailsEditor;
+	TSharedPtr<FClothSimulationNodeDetailExtender> NodeDetailsExtender;
 
 	FDelegateHandle OnPackageReloadedDelegateHandle;
 
diff --git a/Engine/Plugins/ChaosClothAssetEditor/Source/ChaosClothAssetEditor/Private/ChaosClothAsset/ClothSimulationNodeDetailExtender.cpp b/Engine/Plugins/ChaosClothAssetEditor/Source/ChaosClothAssetEditor/Private/ChaosClothAsset/ClothSimulationNodeDetailExtender.cpp
new file mode 100644
--- /dev/null
+++ b/Engine/Plugins/ChaosClothAssetEditor/Source/ChaosClothAssetEditor/Private/ChaosClothAsset/ClothSimulationNodeDetailExtender.cpp
@@ -0,0 +1,39 @@
+// Copyright Epic Games, Inc. All Rights Reserved.
+
+#include "ClothSimulationNodeDetailExtender.h"
+#include "HAL/PlatformApplicationMisc.h"
+
+#define LOCTEXT_NAMESPACE "ChaosClothAssetSimulationNodeDetailExtender"
+
+namespace UE::Chaos::ClothAsset
+{
+	bool FClothSimulationNodeDetailExtender::IsPropertyExtendable(const UClass* InObjectClass, const IPropertyHandle& PropertyHandle) const
+	{
+		return PropertyHandle.HasMetaData(TEXT("InteractorName"));
+	}
+
+	void FClothSimulationNodeDetailExtender::ExtendWidgetRow(FDetailWidgetRow& InWidgetRow, const IDetailLayoutBuilder& InDetailBuilder, const UClass* InObjectClass, TSharedPtr<IPropertyHandle> PropertyHandle)
+	{
+		if (!PropertyHandle)
+		{
+			return;
+		}
+
+		check(PropertyHandle->HasMetaData(TEXT("InteractorName")));
+		FDetailWidgetRow::FCustomMenuData CustomMenuData(
+			FUIAction(
+				FExecuteAction::CreateLambda([PropertyHandle]()
+					{
+						const FString& InteractorName = PropertyHandle->GetMetaData(TEXT("InteractorName"));
+						FPlatformApplicationMisc::ClipboardCopy(*InteractorName);
+					})),
+			LOCTEXT("CopyInteractorName", "Copy Interactor Name"),
+			LOCTEXT("CopyInteractorNameToolTip", "Copy name that can be used to set this property via the Cloth Asset Interactor"),
+			FSlateIcon(FCoreStyle::Get().GetStyleSetName(), "GenericCommands.Copy")
+		);
+		InWidgetRow.CustomMenuItems.Add(CustomMenuData);
+	}
+
+} // namespace UE::Chaos::ClothAsset
+
+#undef LOCTEXT_NAMESPACE
diff --git a/Engine/Plugins/ChaosClothAssetEditor/Source/ChaosClothAssetEditor/Private/ChaosClothAsset/ClothSimulationNodeDetailExtender.h b/Engine/Plugins/ChaosClothAssetEditor/Source/ChaosClothAssetEditor/Private/ChaosClothAsset/ClothSimulationNodeDetailExtender.h
new file mode 100644
--- /dev/null
+++ b/Engine/Plugins/ChaosClothAssetEditor/Source/ChaosClothAssetEditor/Private/ChaosClothAsset/ClothSimulationNodeDetailExtender.h
@@ -0,0 +1,14 @@
+// Copyright Epic Games, Inc. All Rights Reserved.
+
+#pragma once
+
+#include "IDetailPropertyExtensionHandler.h"
+
+namespace UE::Chaos::ClothAsset
+{
+	struct FClothSimulationNodeDetailExtender : public IDetailPropertyExtensionHandler
+	{
+		virtual bool IsPropertyExtendable(const UClass* InObjectClass, const IPropertyHandle& PropertyHandle) const override;
+		virtual void ExtendWidgetRow(FDetailWidgetRow& InWidgetRow, const IDetailLayoutBuilder& InDetailBuilder, const UClass* InObjectClass, TSharedPtr<IPropertyHandle> PropertyHandle) override;
+	};
+} // namespace UE::Chaos::ClothAsset
