Cloth Asset - Refactored the proxy setup function into multiples stages to cope with the Simulation Manager disabled state, and allows LOD switching in the Dataflow editor while maintaining the correct behavior when replaying caches.

#rnx
#tests Tried cloth asset simulation in both the Cloth Asset and Dataflow editors. Successfully cached the simulation in the Dataflow editor.

diff --git a/Engine/Plugins/ChaosClothAsset/Source/ChaosClothAssetEngine/Private/ChaosClothAsset/ClothComponent.cpp b/Engine/Plugins/ChaosClothAsset/Source/ChaosClothAssetEngine/Private/ChaosClothAsset/ClothComponent.cpp
--- a/Engine/Plugins/ChaosClothAsset/Source/ChaosClothAssetEngine/Private/ChaosClothAsset/ClothComponent.cpp
+++ b/Engine/Plugins/ChaosClothAsset/Source/ChaosClothAssetEngine/Private/ChaosClothAsset/ClothComponent.cpp
@@ -263,10 +263,8 @@ void UChaosClothComponent::TickComponent(float DeltaTime, enum ELevelTick TickTy
 	// Make sure that the previous frame simulation has completed
 	HandleExistingParallelSimulation();
 	
-	if(!SimulationAsset.DataflowAsset)
+	if (!SimulationAsset.DataflowAsset)
 	{
-		// < This would be the right place to update the preset/use an interactor, ...etc.
-
 		// Update the proxy and start the simulation parallel task
 		StartNewParallelSimulation(DeltaTime);
 
@@ -644,13 +642,22 @@ void UChaosClothComponent::ResetSimulationProxy()
 	ClothSimulationProxy.Reset();
 }
 
+void UChaosClothComponent::PreProcessSimulation(const float DeltaTime)
+{
+	if (ClothSimulationProxy.IsValid())
+	{
+		ClothSimulationProxy->PreProcess_GameThread(DeltaTime);
+	}
+}
+
 void UChaosClothComponent::WriteToSimulation(const float DeltaTime, const bool bAsyncTask)
 {
 	if (ClothSimulationProxy.IsValid())
 	{
 		const bool bIsSimulating = ClothSimulationProxy->PreSimulate_GameThread(DeltaTime);
-		const int32 CurrentLOD = GetPredictedLODLevel();
 
+		// Clear property dirty flags
+		const int32 CurrentLOD = GetPredictedLODLevel();
 		if (bIsSimulating && CollectionPropertyFacades.IsValidIndex(CurrentLOD) && CollectionPropertyFacades[CurrentLOD].IsValid())
 		{
 			CollectionPropertyFacades[CurrentLOD]->ClearDirtyFlags();
@@ -658,19 +665,19 @@ void UChaosClothComponent::WriteToSimulation(const float DeltaTime, const bool b
 	}
 }
 
-void UChaosClothComponent::PreProcessSimulation(const float DeltaTime)
+void UChaosClothComponent::ReadFromSimulation(const float /*DeltaTime*/, const bool /*bAsyncTask*/)
 {
-	if (ClothSimulationProxy.IsValid() && ClothSimulationProxy->HasCacheData())
+	if (ClothSimulationProxy.IsValid())
 	{
-		WriteToSimulation(DeltaTime, false);
+		ClothSimulationProxy->PostSimulate_GameThread();
 	}
 }
 
-void UChaosClothComponent::ReadFromSimulation(const float DeltaTime, const bool bAsyncTask)
+void UChaosClothComponent::PostProcessSimulation(const float /*DeltaTime*/)
 {
 	if (ClothSimulationProxy.IsValid())
 	{
-		ClothSimulationProxy->PostSimulate_GameThread();
+		ClothSimulationProxy->PostProcess_GameThread();
 	}
 }
 
diff --git a/Engine/Plugins/ChaosClothAsset/Source/ChaosClothAssetEngine/Private/ChaosClothAsset/ClothSimulationProxy.cpp b/Engine/Plugins/ChaosClothAsset/Source/ChaosClothAssetEngine/Private/ChaosClothAsset/ClothSimulationProxy.cpp
--- a/Engine/Plugins/ChaosClothAsset/Source/ChaosClothAssetEngine/Private/ChaosClothAsset/ClothSimulationProxy.cpp
+++ b/Engine/Plugins/ChaosClothAsset/Source/ChaosClothAssetEngine/Private/ChaosClothAsset/ClothSimulationProxy.cpp
@@ -179,19 +179,48 @@ PRAGMA_ENABLE_DEPRECATION_WARNINGS
 		CompleteParallelSimulation_GameThread();
 	}
 	
-	void FClothSimulationProxy::PreSolveProxy(const float DeltaTime)
+	bool FClothSimulationProxy::SetupSimulationData(float DeltaTime)
 	{
-		SetupSimulationData(DeltaTime);
+		PreProcess_GameThread(DeltaTime);
+		PreSimulate_GameThread(DeltaTime);
+		PostProcess_GameThread();
+		return bIsSimulating;
 	}
 
-	bool FClothSimulationProxy::SetupSimulationData(float DeltaTime)
+	void FClothSimulationProxy::PreProcess_GameThread(float DeltaTime)
 	{
 		SCOPE_CYCLE_COUNTER(STAT_ClothSimulationProxy_TickGame);
 
-		// Fill a new context, note the context is also needed when the simulation is suspended
+		// Set bIsSimulating to its default value, this will be changed in PreSimulate_GameThread if the simulation ever runs
+		bIsSimulating = false;
+
+		// Fill a new context, note the context is also needed when the simulation is suspended or playing back the cache
 		constexpr bool bIsInitializationFalse = false;
 		FillSimulationContext(DeltaTime, bIsInitializationFalse);
 
+		// Check whether the solver should be enabled for caching purpose
+		const bool bShouldEnableSolver = ShouldEnableSolver(Solver->GetEnableSolver());  // Note: needs to be called after filling the context
+		Solver->SetEnableSolver(bShouldEnableSolver);
+
+		// If a cache is provided, then read it now in the Tick function
+		if (!bShouldEnableSolver && ClothSimulationContext->CacheData.HasData())
+		{
+			// Tick with the solver disabled to read the sim from the cache, the simulation won't be called
+			Tick();
+		}
+	}
+
+	bool FClothSimulationProxy::PreSimulate_GameThread(float DeltaTime)
+	{
+		SCOPE_CYCLE_COUNTER(STAT_ClothSimulationProxy_TickGame);
+
+		// Check if it is playing the cache back
+		const bool bIsSolverEnabled = Solver->GetEnableSolver();
+		if (!bIsSolverEnabled && ClothSimulationContext->CacheData.HasData())
+		{
+			return false;
+		}
+
 		// Check that the render mesh current LOD isn't just fully skinned
 		const int32 LodIndex = ClothSimulationContext->LodIndex;
 		const UChaosClothAsset* const ClothAsset = ClothComponent.GetClothAsset();
@@ -200,78 +229,84 @@ PRAGMA_ENABLE_DEPRECATION_WARNINGS
 			!ClothAsset->GetResourceForRendering()->LODRenderData.IsValidIndex(LodIndex) ||
 			!ClothAsset->GetResourceForRendering()->LODRenderData[LodIndex].HasClothData())
 		{
-			return false; // Not Simulating
+			return false;
 		}
 
-		// Check whether the solver should be enabled for caching purpose
-		const bool bShouldEnableSolver = ShouldEnableSolver(Solver->GetEnableSolver());
-		Solver->SetEnableSolver(bShouldEnableSolver);
-
-		// Prepare the solver task
-		const bool bUseCache = ClothSimulationContext->CacheData.HasData();
-		const bool bCreateParallelTask = (DeltaTime > 0.f && !ClothComponent.IsSimulationSuspended()) || bUseCache;
-		if (bCreateParallelTask)
+		// Check whether there some actual simulation needs to happen
+		bIsSimulating = (DeltaTime > 0.f && !ClothComponent.IsSimulationSuspended() && bIsSolverEnabled);
+		if (!bIsSimulating)
 		{
-			if (bShouldEnableSolver)
-			{
-				// Update the config properties
-				InitializeConfigs();
+			return false;
+		}
 
-				// Update world forces
-				if (UWorld* const World = ClothComponent.GetWorld())
-				{
-					if (UPhysicsFieldComponent* const PhysicsField = World->PhysicsField)
-					{
-						const FBox BoundingBox = CalculateBounds_AnyThread().GetBox().TransformBy(ClothComponent.GetComponentTransform());
+		// Update the config properties
+		InitializeConfigs();
 
-						PhysicsField->FillTransientCommands(false, BoundingBox, Solver->GetTime(), Solver->GetPerSolverField().GetTransientCommands());
-						PhysicsField->FillPersistentCommands(false, BoundingBox, Solver->GetTime(), Solver->GetPerSolverField().GetPersistentCommands());
-					}
-				}
+		// Update world forces
+		if (UWorld* const World = ClothComponent.GetWorld())
+		{
+			if (UPhysicsFieldComponent* const PhysicsField = World->PhysicsField)
+			{
+				const FBox BoundingBox = CalculateBounds_AnyThread().GetBox().TransformBy(ClothComponent.GetComponentTransform());
 
-				// Update external collision sources
-				CollisionSourcesProxy->ExtractCollisionData();
+				PhysicsField->FillTransientCommands(false, BoundingBox, Solver->GetTime(), Solver->GetPerSolverField().GetTransientCommands());
+				PhysicsField->FillPersistentCommands(false, BoundingBox, Solver->GetTime(), Solver->GetPerSolverField().GetPersistentCommands());
 			}
-
-			return true;  // Simulating
 		}
 
-		// It still needs to write back to the GT cache as the context has changed
-		UpdateClothLODs();
-		if (ClothComponent.IsSimulationEnabled())
+		// Update external collision sources
+		CollisionSourcesProxy->ExtractCollisionData();
+
+		return bIsSimulating;
+	}
+
+	void FClothSimulationProxy::PostSimulate_GameThread()
+	{
+		if (bIsSimulating)
 		{
 			WriteSimulationData();
 		}
-		else
-		{
-			CurrentSimulationData.Reset();
-		}
-
-		return false;  // Not simulating
 	}
 
-	bool FClothSimulationProxy::PreSimulate_GameThread(float DeltaTime)
+	void FClothSimulationProxy::PostProcess_GameThread()
 	{
-		SCOPE_CYCLE_COUNTER(STAT_ClothSimulationProxy_TickGame);
-
-		const bool bIsSimulating = SetupSimulationData(DeltaTime);
-		if(bIsSimulating && ClothSimulationContext->CacheData.HasData())
+		if (!bIsSimulating)
 		{
-			// Start the the cloth simulation thread only when reading cache
-			ParallelTask = TGraphTask<FClothSimulationProxyParallelTask>::CreateTask(nullptr, ENamedThreads::GameThread).ConstructAndDispatchWhenReady(*this);
+			// Take care of the LOD switching, as the simulation won't do it
+			UpdateClothLODs();
+
+			// If the simulation is enabled, then it is suspended or reading from cache and the simulation data still needs updating (transforms and LODs)
+			if (ClothComponent.IsSimulationEnabled())
+			{
+				WriteSimulationData();
+			}
+			else
+			{
+				CurrentSimulationData.Reset();
+			}
 		}
-		return bIsSimulating;
 	}
 
 	bool FClothSimulationProxy::Tick_GameThread(float DeltaTime)
 	{
 		SCOPE_CYCLE_COUNTER(STAT_ClothSimulationProxy_TickGame);
 
-		const bool bIsSimulating = SetupSimulationData(DeltaTime);
-		if(bIsSimulating)
+		PreProcess_GameThread(DeltaTime);
+		PreSimulate_GameThread(DeltaTime);
+
+		if (bIsSimulating)
 		{
 			// Start the the cloth simulation thread
 			ParallelTask = TGraphTask<FClothSimulationProxyParallelTask>::CreateTask(nullptr, ENamedThreads::GameThread).ConstructAndDispatchWhenReady(*this);
+
+			// Note: Post simulate and post process will be handled when the parallel task completes
+		}
+		else
+		{
+			// The post simulate doesn't do much right now, but better still call it for completion in case it one day does
+			PostSimulate_GameThread();
+			// The post process only happens when there isn't any simulation
+			PostProcess_GameThread();
 		}
 		return bIsSimulating;
 	}
@@ -433,7 +468,8 @@ PRAGMA_ENABLE_DEPRECATION_WARNINGS
 			ParallelTask.SafeRelease();
 
 			// Write back to the GT cache
-			WriteSimulationData();
+			PostSimulate_GameThread();
+			PostProcess_GameThread();
 		}
 	}
 
diff --git a/Engine/Plugins/ChaosClothAsset/Source/ChaosClothAssetEngine/Public/ChaosClothAsset/ClothComponent.h b/Engine/Plugins/ChaosClothAsset/Source/ChaosClothAssetEngine/Public/ChaosClothAsset/ClothComponent.h
--- a/Engine/Plugins/ChaosClothAsset/Source/ChaosClothAssetEngine/Public/ChaosClothAsset/ClothComponent.h
+++ b/Engine/Plugins/ChaosClothAsset/Source/ChaosClothAssetEngine/Public/ChaosClothAsset/ClothComponent.h
@@ -204,6 +204,7 @@ protected:
 	virtual void WriteToSimulation(const float DeltaTime, const bool bAsyncTask) override;
 	virtual void ReadFromSimulation(const float DeltaTime, const bool bAsyncTask) override;
 	virtual void PreProcessSimulation(const float DeltaTime) override;
+	virtual void PostProcessSimulation(const float DeltaTime) override;
 	// End IDataflowPhysicsSolverInterface overrides
 
 	/** Override this function for setting up custom simulation proxies when the component is registered. */
diff --git a/Engine/Plugins/ChaosClothAsset/Source/ChaosClothAssetEngine/Public/ChaosClothAsset/ClothSimulationProxy.h b/Engine/Plugins/ChaosClothAsset/Source/ChaosClothAssetEngine/Public/ChaosClothAsset/ClothSimulationProxy.h
--- a/Engine/Plugins/ChaosClothAsset/Source/ChaosClothAssetEngine/Public/ChaosClothAsset/ClothSimulationProxy.h
+++ b/Engine/Plugins/ChaosClothAsset/Source/ChaosClothAssetEngine/Public/ChaosClothAsset/ClothSimulationProxy.h
@@ -59,11 +59,14 @@ namespace UE::Chaos::ClothAsset
 		/** Wait for the parallel task to complete if one was running, and update the simulation data. */
 		void CompleteParallelSimulation_GameThread();
 
-		/** write simulation data back onto GT after the simulation is done */
-		void PostSimulate_GameThread() {WriteSimulationData();}
-		
-		/** setup simulation data from GT before the simulation starts */
+		/** Prepare the proxy before the simulation is setup. For specialized usage when not calling Tick_GameThread. */
+		void PreProcess_GameThread(float DeltaTime);
+		/** Prepare the simulation data. For specialized usage when not calling Tick_GameThread. */
         bool PreSimulate_GameThread(float DeltaTime);
+		/** Write data back onto the game thread once the simulation is done. For specialized usage when not calling Tick_GameThread. */
+		void PostSimulate_GameThread();
+		/** Post setup, required for when the simulation didn't run. For specialized usage when not calling Tick_GameThread. */
+		void PostProcess_GameThread();
 
 		/**
 		 * Return a map of all simulation data as used by the skeletal rendering code.
@@ -89,11 +92,13 @@ PRAGMA_ENABLE_DEPRECATION_WARNINGS
 		float GetLinearSolveError() const { return LastLinearSolveError; }
 		float GetSimulationTime() const { return SimulationTime; }
 		bool IsTeleported() const { return bIsTeleported; }
+		UE_DEPRECATED(5.6, "CacheData is reset after the context is updated, making this method unreliable.")
 		bool HasCacheData() const { return CacheData.IsValid(); }
 
 	protected:
 		void Tick();
 		void WriteSimulationData();
+		UE_DEPRECATED(5.6, "Use PreProcess_GameThread, PreSimulate_GameThread, PostSimulate_GameThread, and PostProcess_GameThread instead.")
 		bool SetupSimulationData(float DeltaTime);
 		void InitializeConfigs();
 		void FillSimulationContext(float DeltaTime, bool bIsInitialization = false);
@@ -103,8 +108,10 @@ PRAGMA_ENABLE_DEPRECATION_WARNINGS
 		void UpdateClothLODs();
 
 		// Begin FDataflowPhysicsSolverProxy overrides
-		virtual void AdvanceSolverDatas(const float DeltaTime) override {Tick();}
-		virtual void PreSolveProxy(const float DeltaTime);
+		virtual void AdvanceSolverDatas(const float DeltaTime) override
+		{
+			Tick();
+		}
 		// End FDataflowPhysicsSolverProxy overrides
 
 		// Internal physics thread object
@@ -166,6 +173,9 @@ PRAGMA_ENABLE_DEPRECATION_WARNINGS
 
 		// Cached value of the MaxPhysicsDeltaTime setting for the life of this proxy
 		const float MaxDeltaTime;
+
+		// Whether the proxy has been running the simulation, for internal use only
+		bool bIsSimulating;
 	};
 
 #if !defined(CHAOS_TRANSFORM_CLOTH_SIMUL_DATA_ISPC_ENABLED_DEFAULT)
