Add the option to move the most expensive part of the cloth proxy initialization (Solver::AddCloth, where all constraints are created) to a worker task.

Two new cvars which can be configured for a project:
* p.ChaosClothAsset.EnableAsyncClothInitialization which enables this new async initialization
* p.ChaosClothAsset.WaitForAsyncClothInitialization which controls what happens if the async initialization hasn't completed when Tick is called (wait vs. disable cloth sim)

Default behavior for both is false. Setup the IDataflowPhysicsSolverInterface calls to always the synchronous version when simulating since this is a caching workflow where per-frame accuracy is more important.

diff --git a/Engine/Plugins/Animation/MLDeformer/ChaosClothGenerator/Source/ChaosClothGenerator/Private/ClothGeneratorComponent.cpp b/Engine/Plugins/Animation/MLDeformer/ChaosClothGenerator/Source/ChaosClothGenerator/Private/ClothGeneratorComponent.cpp
--- a/Engine/Plugins/Animation/MLDeformer/ChaosClothGenerator/Source/ChaosClothGenerator/Private/ClothGeneratorComponent.cpp
+++ b/Engine/Plugins/Animation/MLDeformer/ChaosClothGenerator/Source/ChaosClothGenerator/Private/ClothGeneratorComponent.cpp
@@ -10,6 +10,12 @@ namespace UE::Chaos::ClothGenerator
     }
 
     FClothGeneratorProxy::~FClothGeneratorProxy() = default;
+
+	void FClothGeneratorProxy::PostConstructor()
+	{
+		constexpr bool bAsyncInitialization = false;
+		PostConstructorInternal(bAsyncInitialization);
+	}
 };
 
 UClothGeneratorComponent::UClothGeneratorComponent(const FObjectInitializer& ObjectInitializer)
diff --git a/Engine/Plugins/Animation/MLDeformer/ChaosClothGenerator/Source/ChaosClothGenerator/Private/ClothGeneratorComponent.h b/Engine/Plugins/Animation/MLDeformer/ChaosClothGenerator/Source/ChaosClothGenerator/Private/ClothGeneratorComponent.h
--- a/Engine/Plugins/Animation/MLDeformer/ChaosClothGenerator/Source/ChaosClothGenerator/Private/ClothGeneratorComponent.h
+++ b/Engine/Plugins/Animation/MLDeformer/ChaosClothGenerator/Source/ChaosClothGenerator/Private/ClothGeneratorComponent.h
@@ -16,6 +16,8 @@ namespace UE::Chaos::ClothGenerator
         explicit FClothGeneratorProxy(const UChaosClothComponent& InClothComponent);
         ~FClothGeneratorProxy();
 
+		virtual void PostConstructor() override;
+
         using FClothSimulationProxy::Tick;
         using FClothSimulationProxy::FillSimulationContext;
         using FClothSimulationProxy::InitializeConfigs;
diff --git a/Engine/Plugins/ChaosClothAsset/Source/ChaosClothAssetEngine/Private/ChaosClothAsset/ClothComponent.cpp b/Engine/Plugins/ChaosClothAsset/Source/ChaosClothAssetEngine/Private/ChaosClothAsset/ClothComponent.cpp
--- a/Engine/Plugins/ChaosClothAsset/Source/ChaosClothAssetEngine/Private/ChaosClothAsset/ClothComponent.cpp
+++ b/Engine/Plugins/ChaosClothAsset/Source/ChaosClothAssetEngine/Private/ChaosClothAsset/ClothComponent.cpp
@@ -146,7 +146,7 @@ void UChaosClothComponent::RecreateClothSimulationProxy()
 			if (ClothSimulationModel && ClothSimulationModel->GetNumLods())
 			{
 				// Create the simulation proxy (note CreateClothSimulationProxy() can be overloaded)
-				ClothSimulationProxy = CreateClothSimulationProxy();
+				CreateClothSimulationProxyImpl();
 			}
 		}
 	}
@@ -539,6 +539,13 @@ TSharedPtr<UE::Chaos::ClothAsset::FClothSimulationProxy> UChaosClothComponent::C
 	return MakeShared<FClothSimulationProxy>(*this);
 }
 
+void UChaosClothComponent::CreateClothSimulationProxyImpl()
+{
+	// Note: CreateClothSimulationProxy and ClothSimulationProxy::PostConstructor are both virtual.
+	ClothSimulationProxy = CreateClothSimulationProxy();
+	ClothSimulationProxy->PostConstructor();
+}
+
 void UChaosClothComponent::AddCollisionSource(USkinnedMeshComponent* SourceComponent, const UPhysicsAsset* SourcePhysicsAsset, bool bUseSphylsOnly)
 {
 	CollisionSources->Add(SourceComponent, SourcePhysicsAsset, bUseSphylsOnly);
@@ -646,7 +653,8 @@ void UChaosClothComponent::PreProcessSimulation(const float DeltaTime)
 {
 	if (ClothSimulationProxy.IsValid())
 	{
-		ClothSimulationProxy->PreProcess_GameThread(DeltaTime);
+		constexpr bool bForceWaitForInitialization = true;
+		ClothSimulationProxy->PreProcess_GameThread(DeltaTime, bForceWaitForInitialization);
 	}
 }
 
diff --git a/Engine/Plugins/ChaosClothAsset/Source/ChaosClothAssetEngine/Private/ChaosClothAsset/ClothSimulationProxy.cpp b/Engine/Plugins/ChaosClothAsset/Source/ChaosClothAssetEngine/Private/ChaosClothAsset/ClothSimulationProxy.cpp
--- a/Engine/Plugins/ChaosClothAsset/Source/ChaosClothAssetEngine/Private/ChaosClothAsset/ClothSimulationProxy.cpp
+++ b/Engine/Plugins/ChaosClothAsset/Source/ChaosClothAssetEngine/Private/ChaosClothAsset/ClothSimulationProxy.cpp
@@ -44,6 +44,11 @@ namespace UE::Chaos::ClothAsset
 	float DeltaTimeDecay = 0.03;
 	FAutoConsoleVariableRef CVarDeltaTimeDecay(TEXT("p.ChaosClothAsset.DeltaTimeDecay"), DeltaTimeDecay, TEXT("Delta Time smoothing decay (1 = no smoothing)"));
 
+	bool bEnableAsyncClothInitialization = false;
+	FAutoConsoleVariableRef CVarEnableAsyncClothInitialization(TEXT("p.ChaosClothAsset.EnableAsyncClothInitialization"), bEnableAsyncClothInitialization, TEXT("Enable asynchronous cloth proxy initialization"));
+	bool bWaitForAsyncClothInitialization = true;
+	FAutoConsoleVariableRef CVarWaitForAsyncClothInitialization(TEXT("p.ChaosClothAsset.WaitForAsyncClothInitialization"), bWaitForAsyncClothInitialization, TEXT("When asynchronous cloth proxy initialization is enabled, wait for initialization to complete to start up cloth simulation. Otherwise, cloth simulation will be disabled until initialization has completed."));
+
 	static FAutoConsoleTaskPriority CPrio_ClothSimulationProxyParallelTask(
 		TEXT("TaskGraph.TaskPriorities.ClothSimulationProxyParallelTask"),
 		TEXT("Task and thread priority for the cloth simulation proxy."),
@@ -54,8 +59,15 @@ namespace UE::Chaos::ClothAsset
 	class FClothSimulationProxyParallelTask
 	{
 	public:
-		FClothSimulationProxyParallelTask(FClothSimulationProxy& InClothSimulationProxy)
+		enum struct EType : uint8
+		{
+			Tick = 0,
+			Initialization
+		};
+
+		FClothSimulationProxyParallelTask(FClothSimulationProxy& InClothSimulationProxy, EType InType = EType::Tick)
 			: ClothSimulationProxy(InClothSimulationProxy)
+			, Type(InType)
 		{
 		}
 
@@ -85,12 +97,21 @@ namespace UE::Chaos::ClothAsset
 			FScopeCycleCounterUObject ContextScope(ClothSimulationProxy.ClothComponent.GetSkinnedAsset());
 			SCOPE_CYCLE_COUNTER(STAT_ClothTotalTime);
 			CSV_SCOPED_TIMING_STAT(Animation, Cloth);
-
-			ClothSimulationProxy.Tick();
+			switch (Type)
+			{
+			case EType::Initialization:
+				ClothSimulationProxy.ExecuteInitialization();
+				break;
+			case EType::Tick:
+			default:
+				ClothSimulationProxy.Tick();
+				break;
+			}
 		}
 
 	private:
 		FClothSimulationProxy& ClothSimulationProxy;
+		EType Type = EType::Tick;
 	};
 
 	FClothSimulationProxy::FClothSimulationProxy(const UChaosClothComponent& InClothComponent)
@@ -100,6 +121,36 @@ namespace UE::Chaos::ClothAsset
 		, Solver(nullptr)
 		, Visualization(nullptr)
 		, MaxDeltaTime(UPhysicsSettings::Get()->MaxPhysicsDeltaTime)
+	{
+	}
+
+	FClothSimulationProxy::~FClothSimulationProxy()
+	{
+		WaitForParallelInitialization_GameThread();
+		CompleteParallelSimulation_GameThread();
+	}
+
+	void FClothSimulationProxy::PostConstructor()
+	{
+		PostConstructorInternal(bEnableAsyncClothInitialization);
+	}
+
+	void FClothSimulationProxy::PostConstructorInternal(bool bAsyncInitialization)
+	{
+		BeginInitialization_GameThread();
+		if (bAsyncInitialization)
+		{
+			ParallelInitializationTask = TGraphTask<FClothSimulationProxyParallelTask>::CreateTask(nullptr, ENamedThreads::GameThread).ConstructAndDispatchWhenReady(*this, FClothSimulationProxyParallelTask::EType::Initialization);
+			// Note: Complete Initialization will be handled when the parallel task completes.
+		}
+		else
+		{
+			ExecuteInitialization();
+			CompleteInitialization_GameThread();
+		}
+	}
+
+	void FClothSimulationProxy::BeginInitialization_GameThread()
 	{
 		using namespace ::Chaos;
 
@@ -160,25 +211,57 @@ PRAGMA_ENABLE_DEPRECATION_WARNINGS
 			Meshes[MeshIndex].Get(),
 			ColliderIndex != INDEX_NONE ? TArray<FClothingSimulationCollider*>({ Colliders[ColliderIndex].Get() }) : TArray<FClothingSimulationCollider*>(),
 			GroupId));
+	}
 
-		Solver->AddCloth(Cloths[ClothIndex].Get());
-		Cloths[ClothIndex]->Reset();
+	void FClothSimulationProxy::ExecuteInitialization()
+	{
+		int32 LocalNumCloths = 0;
+		int32 LocalNumKinematicParticles = 0;
+		int32 LocalNumDynamicParticles = 0;
+		for (TUniquePtr<::Chaos::FClothingSimulationCloth>& Cloth : Cloths)
+		{
+			Solver->AddCloth(Cloth.Get());
+			Cloth->Reset();
+			++LocalNumCloths;
+			LocalNumKinematicParticles += Cloth->GetNumActiveKinematicParticles();
+			LocalNumDynamicParticles += Cloth->GetNumActiveDynamicParticles();
+		}
 
 		// Update cloth stats
-		NumCloths = 1;
-		NumKinematicParticles = Cloths[ClothIndex]->GetNumActiveKinematicParticles();
-		NumDynamicParticles = Cloths[ClothIndex]->GetNumActiveDynamicParticles();
+		NumCloths = LocalNumCloths;
+		NumKinematicParticles = LocalNumKinematicParticles;
+		NumDynamicParticles = LocalNumDynamicParticles;
+	}
+
+	void FClothSimulationProxy::WaitForParallelInitialization_GameThread()
+	{
+		check(IsInGameThread());
+		using namespace ::Chaos;
+		if (IsValidRef(ParallelInitializationTask))
+		{
+			TRACE_CPUPROFILER_EVENT_SCOPE(STAT_ClothSimulationProxy_CompleteInitialization);
+			CSV_SCOPED_SET_WAIT_STAT(Cloth);
+
+			FTaskGraphInterface::Get().WaitUntilTaskCompletes(ParallelInitializationTask, ENamedThreads::GameThread);
+
+			// No longer need this task, it has completed
+			ParallelInitializationTask.SafeRelease();
+		}
+	}
+
+	void FClothSimulationProxy::CompleteInitialization_GameThread()
+	{
+		check(IsInGameThread());
+		using namespace ::Chaos;
+
+		WaitForParallelInitialization_GameThread();
 
 		// Set start pose (update the context, then the solver without advancing the simulation)
+		constexpr Softs::FSolverReal NoAdvanceDt = 0.f;
 		ClothSimulationContext->Fill(ClothComponent, NoAdvanceDt, MaxDeltaTime);
 		Solver->Update((Softs::FSolverReal)ClothSimulationContext->DeltaTime);
 	}
 
-	FClothSimulationProxy::~FClothSimulationProxy()
-	{
-		CompleteParallelSimulation_GameThread();
-	}
-	
 	bool FClothSimulationProxy::SetupSimulationData(float DeltaTime)
 	{
 		PreProcess_GameThread(DeltaTime);
@@ -187,12 +270,23 @@ PRAGMA_ENABLE_DEPRECATION_WARNINGS
 		return bIsSimulating;
 	}
 
-	void FClothSimulationProxy::PreProcess_GameThread(float DeltaTime)
+	void FClothSimulationProxy::PreProcess_GameThread(float DeltaTime, bool bForceWaitForInitialization)
 	{
 		SCOPE_CYCLE_COUNTER(STAT_ClothSimulationProxy_TickGame);
 
 		// Set bIsSimulating to its default value, this will be changed in PreSimulate_GameThread if the simulation ever runs
 		bIsSimulating = false;
+		bIsPreProcessed = false;
+		if (IsValidRef(ParallelInitializationTask))
+		{
+			if (!bForceWaitForInitialization && !bWaitForAsyncClothInitialization && !ParallelInitializationTask->IsComplete())
+			{
+				return;
+			}
+			CompleteInitialization_GameThread();
+		}
+
+		bIsPreProcessed = true;
 
 		// Fill a new context, note the context is also needed when the simulation is suspended or playing back the cache
 		constexpr bool bIsInitializationFalse = false;
@@ -214,6 +308,11 @@ PRAGMA_ENABLE_DEPRECATION_WARNINGS
 	{
 		SCOPE_CYCLE_COUNTER(STAT_ClothSimulationProxy_TickGame);
 
+		if (!bIsPreProcessed)
+		{
+			return false;
+		}
+
 		// Check if it is playing the cache back
 		const bool bIsSolverEnabled = Solver->GetEnableSolver();
 		if (!bIsSolverEnabled && ClothSimulationContext->CacheData.HasData())
@@ -270,6 +369,11 @@ PRAGMA_ENABLE_DEPRECATION_WARNINGS
 
 	void FClothSimulationProxy::PostProcess_GameThread()
 	{
+		if (!bIsPreProcessed)
+		{
+			CurrentSimulationData.Reset();
+			return;
+		}
 		if (!bIsSimulating)
 		{
 			// Take care of the LOD switching, as the simulation won't do it
@@ -315,6 +419,11 @@ PRAGMA_ENABLE_DEPRECATION_WARNINGS
 	{
 		using namespace ::Chaos;
 
+		if (!bIsPreProcessed)
+		{
+			return;
+		}
+
 		TRACE_CPUPROFILER_EVENT_SCOPE(FClothSimulationProxy_TickPhysics);
 		SCOPE_CYCLE_COUNTER(STAT_ClothSimulationProxy_TickPhysics);
 		const bool bUseCache = ClothSimulationContext->CacheData.HasData();
@@ -476,6 +585,8 @@ PRAGMA_ENABLE_DEPRECATION_WARNINGS
 	void FClothSimulationProxy::UpdateClothLODs()
 	{
 		using namespace ::Chaos;
+		check(!IsValidRef(ParallelInitializationTask));
+		check(bIsPreProcessed);
 
 		bool bAnyLODsChanged = false;
 		for (const TUniquePtr<FClothingSimulationCloth>& Cloth : Cloths)
@@ -508,6 +619,8 @@ PRAGMA_ENABLE_DEPRECATION_WARNINGS
 	void FClothSimulationProxy::WriteSimulationData()
 	{
 		using namespace ::Chaos;
+		check(!IsValidRef(ParallelInitializationTask));
+		check(bIsPreProcessed);
 
 		CSV_SCOPED_TIMING_STAT(Animation, Cloth);
 		TRACE_CPUPROFILER_EVENT_SCOPE(FClothSimulationProxy_WriteSimulationData);
@@ -634,7 +747,7 @@ PRAGMA_ENABLE_DEPRECATION_WARNINGS
 		// This is called during EndOfFrameUpdates, usually in a parallel-for loop. We need to be sure that
 		// the cloth task (if there is one) is complete, but it cannot be waited for here. See OnPreEndOfFrameUpdateSync
 		// which is called just before EOF updates and is where we would have waited for the cloth task.
-		if (!IsValidRef(ParallelTask) || ParallelTask->IsComplete())
+		if (bIsPreProcessed && (!IsValidRef(ParallelTask) || ParallelTask->IsComplete()))
 		{
 			return CurrentSimulationData;
 		}
@@ -647,7 +760,7 @@ PRAGMA_ENABLE_DEPRECATION_WARNINGS
 		SCOPE_CYCLE_COUNTER(STAT_ClothSimulationProxy_CalculateBounds);
 
 		check(Solver);
-		if (!IsValidRef(ParallelTask) || ParallelTask->IsComplete())
+		if (bIsPreProcessed && (!IsValidRef(ParallelTask) || ParallelTask->IsComplete()))
 		{
 			FBoxSphereBounds Bounds = Solver->CalculateBounds();
 
diff --git a/Engine/Plugins/ChaosClothAsset/Source/ChaosClothAssetEngine/Public/ChaosClothAsset/ClothComponent.h b/Engine/Plugins/ChaosClothAsset/Source/ChaosClothAssetEngine/Public/ChaosClothAsset/ClothComponent.h
--- a/Engine/Plugins/ChaosClothAsset/Source/ChaosClothAssetEngine/Public/ChaosClothAsset/ClothComponent.h
+++ b/Engine/Plugins/ChaosClothAsset/Source/ChaosClothAssetEngine/Public/ChaosClothAsset/ClothComponent.h
@@ -211,6 +211,7 @@ protected:
 	virtual TSharedPtr<UE::Chaos::ClothAsset::FClothSimulationProxy> CreateClothSimulationProxy();
 
 private:
+	void CreateClothSimulationProxyImpl();
 	void StartNewParallelSimulation(float DeltaTime);
 	void HandleExistingParallelSimulation();
 	bool ShouldWaitForParallelSimulationInTickComponent() const;
diff --git a/Engine/Plugins/ChaosClothAsset/Source/ChaosClothAssetEngine/Public/ChaosClothAsset/ClothSimulationProxy.h b/Engine/Plugins/ChaosClothAsset/Source/ChaosClothAssetEngine/Public/ChaosClothAsset/ClothSimulationProxy.h
--- a/Engine/Plugins/ChaosClothAsset/Source/ChaosClothAssetEngine/Public/ChaosClothAsset/ClothSimulationProxy.h
+++ b/Engine/Plugins/ChaosClothAsset/Source/ChaosClothAssetEngine/Public/ChaosClothAsset/ClothSimulationProxy.h
@@ -53,14 +53,17 @@ namespace UE::Chaos::ClothAsset
 		FClothSimulationProxy& operator=(const FClothSimulationProxy&) = delete;
 		FClothSimulationProxy& operator=(FClothSimulationProxy&&) = delete;
 
+		virtual void PostConstructor();
+
 		/** Start the simulation if it isn't paused or suspended and return true, or simply update the existing simulation data and return false otherwise. */
 		bool Tick_GameThread(float DeltaTime);
 
 		/** Wait for the parallel task to complete if one was running, and update the simulation data. */
 		void CompleteParallelSimulation_GameThread();
 
-		/** Prepare the proxy before the simulation is setup. For specialized usage when not calling Tick_GameThread. */
-		void PreProcess_GameThread(float DeltaTime);
+		/** Prepare the proxy before the simulation is setup. For specialized usage when not calling Tick_GameThread. 
+		 * @param bForceWaitForInitialization will ignore the value of p.ChaosClothAsset.WaitForAsyncClothInitialization and force waiting for any in flight initialization, guaranteeing successful preprocessing..*/
+		void PreProcess_GameThread(float DeltaTime, bool bForceWaitForInitialization = false);
 		/** Prepare the simulation data. For specialized usage when not calling Tick_GameThread. */
         bool PreSimulate_GameThread(float DeltaTime);
 		/** Write data back onto the game thread once the simulation is done. For specialized usage when not calling Tick_GameThread. */
@@ -96,6 +99,7 @@ PRAGMA_ENABLE_DEPRECATION_WARNINGS
 		bool HasCacheData() const { return CacheData.IsValid(); }
 
 	protected:
+		void PostConstructorInternal(bool bAsyncInitialization);
 		void Tick();
 		void WriteSimulationData();
 		UE_DEPRECATED(5.6, "Use PreProcess_GameThread, PreSimulate_GameThread, PostSimulate_GameThread, and PostProcess_GameThread instead.")
@@ -104,6 +108,16 @@ PRAGMA_ENABLE_DEPRECATION_WARNINGS
 		void FillSimulationContext(float DeltaTime, bool bIsInitialization = false);
 
 	private:
+
+		// Begin proxy initialization. Can't be done in parallel.
+		void BeginInitialization_GameThread();
+		// Thread-safe part of Initialization.
+		void ExecuteInitialization();
+		// Wait for any in-flight initialization.
+		void WaitForParallelInitialization_GameThread();
+		// Wait for any in-flight initialization and complete the initialization process. Can't be done in parallel.
+		void CompleteInitialization_GameThread();
+
 		bool ShouldEnableSolver(bool bSolverCurrentlyEnabled) const;
 		void UpdateClothLODs();
 
@@ -120,6 +134,9 @@ PRAGMA_ENABLE_DEPRECATION_WARNINGS
 		// Reference for the cloth parallel task, to detect whether or not a simulation is running
 		FGraphEventRef ParallelTask;
 
+		// Reference for the cloth async initialization parallel task, to detect whether or not initialization is not complete.
+		FGraphEventRef ParallelInitializationTask;
+
 		// Simulation data written back to the component after the simulation has taken place
 		TMap<int32, FClothSimulData> CurrentSimulationData;
 
@@ -176,6 +193,9 @@ PRAGMA_ENABLE_DEPRECATION_WARNINGS
 
 		// Whether the proxy has been running the simulation, for internal use only
 		bool bIsSimulating;
+
+		// Whether the proxy was initialized when Tick_GameThread/PreProcess_GameThread was called, for internal use only
+		bool bIsPreProcessed = false;
 	};
 
 #if !defined(CHAOS_TRANSFORM_CLOTH_SIMUL_DATA_ISPC_ENABLED_DEFAULT)
