Fix checks in ChaosClothGenerator that are occurring after adding support for async initialization.

I tested Cloth in Cloth and Dataflow Editor, and PIE, as well as ChaosClothGenerator, and they all seem to be working with this change.

diff --git a/Engine/Plugins/Animation/MLDeformer/ChaosClothGenerator/Source/ChaosClothGenerator/Private/ChaosClothGenerator.cpp b/Engine/Plugins/Animation/MLDeformer/ChaosClothGenerator/Source/ChaosClothGenerator/Private/ChaosClothGenerator.cpp
--- a/Engine/Plugins/Animation/MLDeformer/ChaosClothGenerator/Source/ChaosClothGenerator/Private/ChaosClothGenerator.cpp
+++ b/Engine/Plugins/Animation/MLDeformer/ChaosClothGenerator/Source/ChaosClothGenerator/Private/ChaosClothGenerator.cpp
@@ -870,7 +870,7 @@ namespace UE::Chaos::ClothGenerator
 		const TArray<FTransform> Transforms = GetBoneTransforms(TaskComponent, AnimFrame);
 		TaskComponent.Pose(Transforms);
 		TaskComponent.ForceNextUpdateTeleportAndReset();
-		DataGenerationProxy.FillSimulationContext(TimeStep);
+		DataGenerationProxy.PreProcess_Internal(TimeStep);
 		DataGenerationProxy.InitializeConfigs();
 		bool bCancelled = false;
 		for (int32 Step = 0; Step < NumSteps; ++Step)
@@ -888,7 +888,7 @@ namespace UE::Chaos::ClothGenerator
 				if (Step == 0 && NumSteps > 1)
 				{
 					TaskComponent.ResetTeleportMode();
-					DataGenerationProxy.FillSimulationContext(TimeStep);	
+					DataGenerationProxy.PreProcess_Internal(TimeStep);
 				}
 	
 				if (SaveType == ESaveType::EveryStep)
diff --git a/Engine/Plugins/Animation/MLDeformer/ChaosClothGenerator/Source/ChaosClothGenerator/Private/ClothGeneratorComponent.h b/Engine/Plugins/Animation/MLDeformer/ChaosClothGenerator/Source/ChaosClothGenerator/Private/ClothGeneratorComponent.h
--- a/Engine/Plugins/Animation/MLDeformer/ChaosClothGenerator/Source/ChaosClothGenerator/Private/ClothGeneratorComponent.h
+++ b/Engine/Plugins/Animation/MLDeformer/ChaosClothGenerator/Source/ChaosClothGenerator/Private/ClothGeneratorComponent.h
@@ -19,7 +19,7 @@ namespace UE::Chaos::ClothGenerator
 		virtual void PostConstructor() override;
 
         using FClothSimulationProxy::Tick;
-        using FClothSimulationProxy::FillSimulationContext;
+        using FClothSimulationProxy::PreProcess_Internal;
         using FClothSimulationProxy::InitializeConfigs;
         using FClothSimulationProxy::WriteSimulationData;
     };
diff --git a/Engine/Plugins/ChaosClothAsset/Source/ChaosClothAssetEngine/Private/ChaosClothAsset/ClothSimulationProxy.cpp b/Engine/Plugins/ChaosClothAsset/Source/ChaosClothAssetEngine/Private/ChaosClothAsset/ClothSimulationProxy.cpp
--- a/Engine/Plugins/ChaosClothAsset/Source/ChaosClothAssetEngine/Private/ChaosClothAsset/ClothSimulationProxy.cpp
+++ b/Engine/Plugins/ChaosClothAsset/Source/ChaosClothAssetEngine/Private/ChaosClothAsset/ClothSimulationProxy.cpp
@@ -260,6 +260,7 @@ PRAGMA_ENABLE_DEPRECATION_WARNINGS
 		constexpr Softs::FSolverReal NoAdvanceDt = 0.f;
 		ClothSimulationContext->Fill(ClothComponent, NoAdvanceDt, MaxDeltaTime);
 		Solver->Update((Softs::FSolverReal)ClothSimulationContext->DeltaTime);
+		bIsInitialized = true;
 	}
 
 	bool FClothSimulationProxy::SetupSimulationData(float DeltaTime)
@@ -286,11 +287,20 @@ PRAGMA_ENABLE_DEPRECATION_WARNINGS
 			CompleteInitialization_GameThread();
 		}
 
+		PreProcess_Internal(DeltaTime);
+	}
+
+	void FClothSimulationProxy::PreProcess_Internal(float DeltaTime)
+	{
+		check(bIsInitialized);
+
 		bIsPreProcessed = true;
 
 		// Fill a new context, note the context is also needed when the simulation is suspended or playing back the cache
 		constexpr bool bIsInitializationFalse = false;
+		PRAGMA_DISABLE_DEPRECATION_WARNINGS
 		FillSimulationContext(DeltaTime, bIsInitializationFalse);
+		PRAGMA_ENABLE_DEPRECATION_WARNINGS
 
 		// Check whether the solver should be enabled for caching purpose
 		const bool bShouldEnableSolver = ShouldEnableSolver(Solver->GetEnableSolver());  // Note: needs to be called after filling the context
diff --git a/Engine/Plugins/ChaosClothAsset/Source/ChaosClothAssetEngine/Public/ChaosClothAsset/ClothSimulationProxy.h b/Engine/Plugins/ChaosClothAsset/Source/ChaosClothAssetEngine/Public/ChaosClothAsset/ClothSimulationProxy.h
--- a/Engine/Plugins/ChaosClothAsset/Source/ChaosClothAssetEngine/Public/ChaosClothAsset/ClothSimulationProxy.h
+++ b/Engine/Plugins/ChaosClothAsset/Source/ChaosClothAssetEngine/Public/ChaosClothAsset/ClothSimulationProxy.h
@@ -105,8 +105,12 @@ PRAGMA_ENABLE_DEPRECATION_WARNINGS
 		UE_DEPRECATED(5.6, "Use PreProcess_GameThread, PreSimulate_GameThread, PostSimulate_GameThread, and PostProcess_GameThread instead.")
 		bool SetupSimulationData(float DeltaTime);
 		void InitializeConfigs();
+
+		UE_DEPRECATED(5.6, "This method will be made private. Use PreProcess_Internal instead.")
 		void FillSimulationContext(float DeltaTime, bool bIsInitialization = false);
 
+		void PreProcess_Internal(float DeltaTime);
+
 	private:
 
 		// Begin proxy initialization. Can't be done in parallel.
@@ -194,6 +198,9 @@ PRAGMA_ENABLE_DEPRECATION_WARNINGS
 		// Whether the proxy has been running the simulation, for internal use only
 		bool bIsSimulating;
 
+		// Whether the proxy has finished initialization, for internal use only
+		bool bIsInitialized = false;
+
 		// Whether the proxy was initialized when Tick_GameThread/PreProcess_GameThread was called, for internal use only
 		bool bIsPreProcessed = false;
 	};
