Add VertexColors support for Chaos Cloth Assets.

We already have the data in the LODModel, but we weren't setting the right flags to get them into the render data.

diff --git a/Engine/Plugins/ChaosClothAsset/Source/ChaosClothAssetEngine/Private/ChaosClothAsset/ClothAsset.cpp b/Engine/Plugins/ChaosClothAsset/Source/ChaosClothAssetEngine/Private/ChaosClothAsset/ClothAsset.cpp
--- a/Engine/Plugins/ChaosClothAsset/Source/ChaosClothAssetEngine/Private/ChaosClothAsset/ClothAsset.cpp
+++ b/Engine/Plugins/ChaosClothAsset/Source/ChaosClothAssetEngine/Private/ChaosClothAsset/ClothAsset.cpp
@@ -491,7 +491,7 @@ void UChaosClothAsset::InitResources()
 	if (SkeletalMeshRenderData.IsValid())
 	{
 		TArray<UMorphTarget*> DummyMorphTargets;  // Not even used by InitResources() at the moment
-		SkeletalMeshRenderData->InitResources(false, DummyMorphTargets, this);
+		SkeletalMeshRenderData->InitResources(GetHasVertexColors(), DummyMorphTargets, this);
 	}
 }
 
diff --git a/Engine/Plugins/ChaosClothAsset/Source/ChaosClothAssetEngine/Public/ChaosClothAsset/ClothAsset.h b/Engine/Plugins/ChaosClothAsset/Source/ChaosClothAssetEngine/Public/ChaosClothAsset/ClothAsset.h
--- a/Engine/Plugins/ChaosClothAsset/Source/ChaosClothAssetEngine/Public/ChaosClothAsset/ClothAsset.h
+++ b/Engine/Plugins/ChaosClothAsset/Source/ChaosClothAssetEngine/Public/ChaosClothAsset/ClothAsset.h
@@ -35,6 +35,7 @@ enum class EClothAssetAsyncProperties : uint64
 	ImportedModel = 1 << 2,
 	ClothCollection = 1 << 3,
 	RefSkeleton = 1 << 4,
+	HasVertexColors = 1 << 5,
 	All = MAX_uint64
 };
 ENUM_CLASS_FLAGS(EClothAssetAsyncProperties);
@@ -116,7 +117,13 @@ public:
 	virtual bool IsValidLODIndex(int32 Index) const override					{ return LODInfo.IsValidIndex(Index); }
 	virtual int32 GetMinLodIdx(bool bForceLowestLODIdx = false) const override;
 	virtual bool NeedCPUData(int32 LODIndex) const override						{ return false; }
-	virtual bool GetHasVertexColors() const override							{ return false; }
+	virtual bool GetHasVertexColors() const override
+	{
+		WaitUntilAsyncPropertyReleased(EClothAssetAsyncProperties::HasVertexColors, ESkinnedAssetAsyncPropertyLockType::ReadOnly);
+		PRAGMA_DISABLE_DEPRECATION_WARNINGS
+			return bHasVertexColors != 0;
+		PRAGMA_ENABLE_DEPRECATION_WARNINGS
+	}
 	virtual int32 GetPlatformMinLODIdx(const ITargetPlatform* TargetPlatform) const override;
 	virtual const FPerPlatformBool& GetDisableBelowMinLodStripping() const override { return DisableBelowMinLodStripping; }
 	virtual const FPerPlatformInt& GetMinLod() const override;
@@ -224,6 +231,14 @@ public:
 
 #endif
 
+	void SetHasVertexColors(bool InbHasVertexColors)
+	{
+		WaitUntilAsyncPropertyReleased(EClothAssetAsyncProperties::HasVertexColors);
+		PRAGMA_DISABLE_DEPRECATION_WARNINGS
+		bHasVertexColors = InbHasVertexColors;
+		PRAGMA_ENABLE_DEPRECATION_WARNINGS
+	}
+
 private:
 	/** Dataflow asset. */
 	UPROPERTY(EditAnywhere, Category = "Dataflow")
@@ -325,6 +340,10 @@ private:
 	UPROPERTY(EditAnywhere, Category = RayTracing)
 	uint8 bSupportRayTracing : 1;
 
+	UE_DEPRECATED(5.6, "This must be protected for async build, always use the accessors even internally.")
+	UPROPERTY()
+	uint8 bHasVertexColors : 1 = false;
+
 	/** Minimum raytracing LOD for this asset. */
 	UPROPERTY(EditAnywhere, Category = RayTracing)
 	int32 RayTracingMinLOD;
diff --git a/Engine/Plugins/ChaosClothAssetEditor/Source/ChaosClothAssetDataflowNodes/Private/ChaosClothAsset/TerminalNode.cpp b/Engine/Plugins/ChaosClothAssetEditor/Source/ChaosClothAssetDataflowNodes/Private/ChaosClothAsset/TerminalNode.cpp
--- a/Engine/Plugins/ChaosClothAssetEditor/Source/ChaosClothAssetDataflowNodes/Private/ChaosClothAsset/TerminalNode.cpp
+++ b/Engine/Plugins/ChaosClothAssetEditor/Source/ChaosClothAssetDataflowNodes/Private/ChaosClothAsset/TerminalNode.cpp
@@ -229,6 +229,7 @@ void FChaosClothAssetTerminalNode_v2::SetAssetValue(TObjectPtr<UObject> Asset, U
 			LoadObject<UPhysicsAsset>(ClothAsset, *PhysicsAssetPathName, nullptr, LOAD_None, nullptr) :
 			nullptr;
 		ClothAsset->SetPhysicsAsset(PhysicsAsset);
+		ClothAsset->SetHasVertexColors(true);
 
 		// Rebuild the asset static data
 		ClothAsset->Build(&LODTransitionDataCache);
