Cloth Geometry Tools: fix bug in adding seams when appending to a cloth collection
- also change a hard crash to a check if we end up with a DynamicMesh whose topology doesn't match the cloth collection

diff --git a/Engine/Plugins/ChaosClothAsset/Source/ChaosClothAsset/Private/ChaosClothAsset/ClothGeometryTools.cpp b/Engine/Plugins/ChaosClothAsset/Source/ChaosClothAsset/Private/ChaosClothAsset/ClothGeometryTools.cpp
--- a/Engine/Plugins/ChaosClothAsset/Source/ChaosClothAsset/Private/ChaosClothAsset/ClothGeometryTools.cpp
+++ b/Engine/Plugins/ChaosClothAsset/Source/ChaosClothAsset/Private/ChaosClothAsset/ClothGeometryTools.cpp
@@ -380,7 +380,7 @@ namespace UE::Chaos::ClothAsset
 		};
 
 		// Stitch together any vertices that were split, either via DynamicMesh NonManifoldMapping or UV Unwrap
-		static void BuildSeams(const TArray<FIsland>& Islands, const UE::Geometry::FDynamicMesh3& DynamicMesh, TArray<FSeam>& OutSeams)
+		static void BuildSeams(const TArray<FIsland>& Islands, const UE::Geometry::FDynamicMesh3& DynamicMesh, int32 PatternIndexOffset, TArray<FSeam>& OutSeams)
 		{
 			OutSeams.Reset();
 
@@ -404,7 +404,7 @@ namespace UE::Chaos::ClothAsset
 
 				// Find all internal seams
 				FSeam InternalSeam;
-				InternalSeam.Patterns = FIntVector2(IslandIndex);
+				InternalSeam.Patterns = FIntVector2(IslandIndex + PatternIndexOffset);
 				for (const TPair<int32, TArray<int32>>& Source : SourceIndexToPositions)
 				{
 					for (int32 FirstSourceArrayIdx = 0; FirstSourceArrayIdx < Source.Value.Num() - 1; ++FirstSourceArrayIdx)
@@ -426,7 +426,7 @@ namespace UE::Chaos::ClothAsset
 					const TMap<int32, TArray<int32>>& OtherSourceIndexToPositions = IslandSourceIndexToPositions[OtherIslandIndex];
 
 					FSeam Seam;
-					Seam.Patterns = FIntVector2(OtherIslandIndex, IslandIndex);
+					Seam.Patterns = FIntVector2(OtherIslandIndex + PatternIndexOffset, IslandIndex + PatternIndexOffset);
 					for (const TPair<int32, TArray<int32>>& FirstSource : SourceIndexToPositions)
 					{
 						if (const TArray<int32>* OtherSource = OtherSourceIndexToPositions.Find(FirstSource.Key))
@@ -911,6 +911,7 @@ namespace UE::Chaos::ClothAsset
 		}
 
 		FCollectionClothFacade Cloth(ClothCollection);
+		const int32 PatternIndexOffset = bAppend ? Cloth.GetNumSimPatterns() : 0;
 		for (FIsland& Island : Islands)
 		{
 			if (Island.Indices.Num() && Island.Positions2D.Num() && Island.Positions3D.Num())
@@ -956,7 +957,7 @@ namespace UE::Chaos::ClothAsset
 		}
 
 		TArray<FSeam> Seams;
-		BuildSeams(Islands, DynamicMesh, Seams);  // Build the seam information as to be able to re-weld the mesh for simulation
+		BuildSeams(Islands, DynamicMesh, PatternIndexOffset, Seams);  // Build the seam information as to be able to re-weld the mesh for simulation
 		for (const FSeam& Seam : Seams)
 		{
 			AddSeam(Cloth, Seam);
diff --git a/Engine/Plugins/ChaosClothAssetEditor/Source/ChaosClothAssetEditor/Private/ChaosClothAsset/ClothEditorMode.cpp b/Engine/Plugins/ChaosClothAssetEditor/Source/ChaosClothAssetEditor/Private/ChaosClothAsset/ClothEditorMode.cpp
--- a/Engine/Plugins/ChaosClothAssetEditor/Source/ChaosClothAssetEditor/Private/ChaosClothAsset/ClothEditorMode.cpp
+++ b/Engine/Plugins/ChaosClothAssetEditor/Source/ChaosClothAssetEditor/Private/ChaosClothAsset/ClothEditorMode.cpp
@@ -863,6 +863,8 @@ void UChaosClothAssetEditorMode::ReinitializeDynamicMeshComponents()
 
 	if (bPatternColors)
 	{
+		checkf(LodMesh.TriangleCount() == ClothFacade.GetNumSimFaces(), TEXT("Expected to have the same number of faces in LodMesh and ClothCollection sim mesh"));
+
 		LodMesh.Attributes()->EnablePrimaryColors();
 		LodMesh.Attributes()->PrimaryColors()->CreateFromPredicate([&ClothFacade](int ParentVID, int TriIDA, int TriIDB)
 			{
@@ -879,10 +881,13 @@ void UChaosClothAssetEditorMode::ReinitializeDynamicMeshComponents()
 			for (int32 TriID = 0; TriID < Pattern.GetNumSimFaces(); ++TriID)
 			{
 				const int32 GlobalTriID = Pattern.GetSimFacesOffset() + TriID;
-				const FIndex3i AttrTri = LodMesh.Attributes()->PrimaryColors()->GetTriangle(GlobalTriID);
-				ColorAttributeLayer->SetElement(AttrTri[0], (FVector4f)PatternColor);
-				ColorAttributeLayer->SetElement(AttrTri[1], (FVector4f)PatternColor);
-				ColorAttributeLayer->SetElement(AttrTri[2], (FVector4f)PatternColor);
+				FIndex3i AttrTri;
+				if (LodMesh.Attributes()->PrimaryColors()->GetTriangleIfValid(GlobalTriID, AttrTri))
+				{
+					ColorAttributeLayer->SetElement(AttrTri[0], (FVector4f)PatternColor);
+					ColorAttributeLayer->SetElement(AttrTri[1], (FVector4f)PatternColor);
+					ColorAttributeLayer->SetElement(AttrTri[2], (FVector4f)PatternColor);
+				}
 			}
 		}
 	}
