Add dataflow variable support to ChaosClothAsset

diff --git a/Engine/Plugins/ChaosClothAsset/Source/ChaosClothAssetEngine/Private/ChaosClothAsset/ClothAsset.cpp b/Engine/Plugins/ChaosClothAsset/Source/ChaosClothAssetEngine/Private/ChaosClothAsset/ClothAsset.cpp
--- a/Engine/Plugins/ChaosClothAsset/Source/ChaosClothAssetEngine/Private/ChaosClothAsset/ClothAsset.cpp
+++ b/Engine/Plugins/ChaosClothAsset/Source/ChaosClothAssetEngine/Private/ChaosClothAsset/ClothAsset.cpp
@@ -251,6 +251,17 @@ void UChaosClothAsset::Serialize(FArchive& Ar)
 void UChaosClothAsset::PostLoad()
 {
 	Super::PostLoad();
+
+#if WITH_EDITORONLY_DATA
+	if (DataflowAsset_DEPRECATED != nullptr)
+	{
+		DataflowInstance.SetDataflowAsset(DataflowAsset_DEPRECATED);
+		DataflowInstance.SetDataflowTerminal(FName(DataflowTerminal_DEPRECATED));
+
+		DataflowAsset_DEPRECATED = nullptr;
+		DataflowTerminal_DEPRECATED.Empty();
+	}
+#endif  // WITH_EDITORONLY_DATA
 }
 
 #if WITH_EDITOR
@@ -1062,8 +1073,8 @@ void UChaosClothAsset::WriteDataflowContent(const TObjectPtr<UDataflowBaseConten
 {
 	if(const TObjectPtr<UDataflowSkeletalContent> SkeletalContent = Cast<UDataflowSkeletalContent>(DataflowContent))
 	{
-		SkeletalContent->SetDataflowAsset(DataflowAsset);
-		SkeletalContent->SetDataflowTerminal(DataflowTerminal);
+		SkeletalContent->SetDataflowAsset(DataflowInstance.GetDataflowAsset());
+		SkeletalContent->SetDataflowTerminal(DataflowInstance.GetDataflowTerminal().ToString());
 
 #if WITH_EDITORONLY_DATA
 		SkeletalContent->SetAnimationAsset(GetPreviewSceneAnimation());
diff --git a/Engine/Plugins/ChaosClothAsset/Source/ChaosClothAssetEngine/Public/ChaosClothAsset/ClothAsset.h b/Engine/Plugins/ChaosClothAsset/Source/ChaosClothAssetEngine/Public/ChaosClothAsset/ClothAsset.h
--- a/Engine/Plugins/ChaosClothAsset/Source/ChaosClothAssetEngine/Public/ChaosClothAsset/ClothAsset.h
+++ b/Engine/Plugins/ChaosClothAsset/Source/ChaosClothAssetEngine/Public/ChaosClothAsset/ClothAsset.h
@@ -7,6 +7,7 @@
 #include "RenderCommandFence.h"
 #include "PerQualityLevelProperties.h"
 #include "Dataflow/DataflowContent.h"
+#include "Dataflow/DataflowInstance.h"
 #include "ClothAsset.generated.h"
 
 #if UE_ENABLE_INCLUDE_ORDER_DEPRECATED_IN_5_5
@@ -44,7 +45,10 @@ ENUM_CLASS_FLAGS(EClothAssetAsyncProperties);
  * Cloth asset for pattern based simulation.
  */
 UCLASS(hidecategories = Object, BlueprintType)
-class CHAOSCLOTHASSETENGINE_API UChaosClothAsset : public USkinnedAsset, public IDataflowContentOwner
+class CHAOSCLOTHASSETENGINE_API UChaosClothAsset :
+	public USkinnedAsset, 
+	public IDataflowContentOwner,
+	public IDataflowInstanceInterface
 {
 	GENERATED_BODY()
 public:
@@ -204,22 +208,32 @@ public:
 	 * Set the Dataflow graph asset for this cloth asset.
 	 * @param InDataflow The new dataflow asset.
 	 */
-	void SetDataflow(UDataflow* InDataflow) { DataflowAsset = InDataflow; }
+	void SetDataflow(UDataflow* InDataflow) { DataflowInstance.SetDataflowAsset(InDataflow); }
 
 	/** Return the Dataflow graph asset associated to this cloth asset if any. */
-	UDataflow* GetDataflow() { return DataflowAsset; }
+	UDataflow* GetDataflow() { return DataflowInstance.GetDataflowAsset(); }
 
 	/** Return the Dataflow graph asset associated to this cloth asset if any, const version. */
-	const UDataflow* GetDataflow() const { return DataflowAsset; }
+	const UDataflow* GetDataflow() const { return DataflowInstance.GetDataflowAsset(); }
 
 	/**
 	 * Set the name of the Dataflow terminal node for this cloth asset.
 	 * @param InDataflowTerminal The new name of the Dataflow terminal node.
 	 */
-	void SetDataflowTerminal(const FString& InDataflowTerminal) { DataflowTerminal = InDataflowTerminal; }
+	void SetDataflowTerminal(const FString& InDataflowTerminal) 
+	{ 
+		DataflowInstance.SetDataflowTerminal(FName(*InDataflowTerminal));
+	}
 
 	/** Return the Dataflow graph asset associated to this cloth asset if any. */
-	const FString& GetDataflowTerminal() const { return DataflowTerminal; }
+	FString GetDataflowTerminal() const 
+	{
+		return DataflowInstance.GetDataflowTerminal().ToString();
+	}
+
+	//~ Begin IDataflowGeneratedAssetInterface
+	virtual const FDataflowInstance& GetDataflowInstance() const { return DataflowInstance; };
+	//~ End IDataflowGeneratedAssetInterface
 
 #if WITH_EDITORONLY_DATA
 
@@ -240,13 +254,18 @@ public:
 	}
 
 private:
+#if WITH_EDITORONLY_DATA
 	/** Dataflow asset. */
-	UPROPERTY(EditAnywhere, Category = "Dataflow")
-	TObjectPtr<UDataflow> DataflowAsset;
+	UPROPERTY(meta = (DeprecatedProperty, DeprecationMessage = "Dataflow Asset is now stored in DataflowInstance"))
+	TObjectPtr<UDataflow> DataflowAsset_DEPRECATED;
 
 	/** Dataflow Asset terminal node. */
+	UPROPERTY(meta = (DeprecatedProperty, DeprecationMessage = "Dataflow terminal name is now stored in DataflowInstance"))
+	FString DataflowTerminal_DEPRECATED = "ClothAssetTerminal";
+#endif  // WITH_EDITORONLY_DATA
+
 	UPROPERTY(EditAnywhere, Category = "Dataflow")
-	FString DataflowTerminal = "ClothAssetTerminal";
+	FDataflowInstance DataflowInstance;
 
 	//~ Begin USkinnedAsset interface
 	/** Initial step for the Post Load process - Can't be done in parallel. */
