Add OverlayMaterial support to ChaosClothAsset

diff --git a/Engine/Plugins/ChaosClothAsset/Source/ChaosClothAssetEngine/Private/ChaosClothAsset/ClothAsset.cpp b/Engine/Plugins/ChaosClothAsset/Source/ChaosClothAssetEngine/Private/ChaosClothAsset/ClothAsset.cpp
--- a/Engine/Plugins/ChaosClothAsset/Source/ChaosClothAssetEngine/Private/ChaosClothAsset/ClothAsset.cpp
+++ b/Engine/Plugins/ChaosClothAsset/Source/ChaosClothAssetEngine/Private/ChaosClothAsset/ClothAsset.cpp
@@ -282,7 +282,10 @@ void UChaosClothAsset::PostLoad()
 #if WITH_EDITOR
 void UChaosClothAsset::PostEditChangeProperty(FPropertyChangedEvent& PropertyChangedEvent)
 {
-	if (PropertyChangedEvent.GetPropertyName() == GET_MEMBER_NAME_CHECKED(UChaosClothAsset, PhysicsAsset))
+	if (PropertyChangedEvent.GetPropertyName() == GET_MEMBER_NAME_CHECKED(UChaosClothAsset, PhysicsAsset) ||
+		PRAGMA_DISABLE_DEPRECATION_WARNINGS
+		PropertyChangedEvent.GetPropertyName() == GET_MEMBER_NAME_CHECKED(UChaosClothAsset, OverlayMaterial))
+		PRAGMA_ENABLE_DEPRECATION_WARNINGS
 	{
 		ReregisterComponents();
 	}
diff --git a/Engine/Plugins/ChaosClothAsset/Source/ChaosClothAssetEngine/Public/ChaosClothAsset/ClothAsset.h b/Engine/Plugins/ChaosClothAsset/Source/ChaosClothAssetEngine/Public/ChaosClothAsset/ClothAsset.h
--- a/Engine/Plugins/ChaosClothAsset/Source/ChaosClothAssetEngine/Public/ChaosClothAsset/ClothAsset.h
+++ b/Engine/Plugins/ChaosClothAsset/Source/ChaosClothAssetEngine/Public/ChaosClothAsset/ClothAsset.h
@@ -37,6 +37,8 @@ enum class EClothAssetAsyncProperties : uint64
 	ClothCollection = 1 << 3,
 	RefSkeleton = 1 << 4,
 	HasVertexColors = 1 << 5,
+	OverlayMaterial = 1 << 6,
+	OverlayMaterialMaxDrawDistance = 1 << 7,
 	All = MAX_uint64
 };
 ENUM_CLASS_FLAGS(EClothAssetAsyncProperties);
@@ -116,8 +118,27 @@ public:
 	virtual const USkeleton* GetSkeleton() const override						{ return Skeleton; }
 	virtual void SetSkeleton(USkeleton* InSkeleton) override					{ Skeleton = InSkeleton; }
 	virtual UMeshDeformer* GetDefaultMeshDeformer() const override				{ return nullptr; }
-	virtual class UMaterialInterface* GetOverlayMaterial() const override		{ return nullptr; }
-	virtual float GetOverlayMaterialMaxDrawDistance() const override			{ return 0.f; }
+
+	/** Get the default overlay material used by this mesh */
+	UFUNCTION(BlueprintCallable, Category = "Rendering|Material")
+	virtual class UMaterialInterface* GetOverlayMaterial() const override		
+	{
+		WaitUntilAsyncPropertyReleased(EClothAssetAsyncProperties::OverlayMaterial, ESkinnedAssetAsyncPropertyLockType::ReadOnly);
+		PRAGMA_DISABLE_DEPRECATION_WARNINGS
+		return OverlayMaterial;
+		PRAGMA_ENABLE_DEPRECATION_WARNINGS
+	}
+
+	/** Get the default overlay material max draw distance used by this mesh */
+	UFUNCTION(BlueprintCallable, Category = "Rendering|Material")
+	virtual float GetOverlayMaterialMaxDrawDistance() const override
+	{
+		WaitUntilAsyncPropertyReleased(EClothAssetAsyncProperties::OverlayMaterialMaxDrawDistance, ESkinnedAssetAsyncPropertyLockType::ReadOnly);
+		PRAGMA_DISABLE_DEPRECATION_WARNINGS
+		return OverlayMaterialMaxDrawDistance;
+		PRAGMA_ENABLE_DEPRECATION_WARNINGS
+	}
+
 	virtual bool IsValidLODIndex(int32 Index) const override					{ return LODInfo.IsValidIndex(Index); }
 	virtual int32 GetMinLodIdx(bool bForceLowestLODIdx = false) const override;
 	virtual bool NeedCPUData(int32 LODIndex) const override						{ return false; }
@@ -125,7 +146,7 @@ public:
 	{
 		WaitUntilAsyncPropertyReleased(EClothAssetAsyncProperties::HasVertexColors, ESkinnedAssetAsyncPropertyLockType::ReadOnly);
 		PRAGMA_DISABLE_DEPRECATION_WARNINGS
-			return bHasVertexColors != 0;
+		return bHasVertexColors != 0;
 		PRAGMA_ENABLE_DEPRECATION_WARNINGS
 	}
 	virtual int32 GetPlatformMinLODIdx(const ITargetPlatform* TargetPlatform) const override;
@@ -152,6 +173,26 @@ public:
 #endif
 	//~ End USkinnedAsset interface
 
+	/** Change the default overlay material used by this mesh */
+	UFUNCTION(BlueprintCallable, Category = "Rendering|Material")
+	void SetOverlayMaterial(class UMaterialInterface* NewOverlayMaterial)
+	{
+		WaitUntilAsyncPropertyReleased(EClothAssetAsyncProperties::OverlayMaterial);
+		PRAGMA_DISABLE_DEPRECATION_WARNINGS
+		OverlayMaterial = NewOverlayMaterial;
+		PRAGMA_ENABLE_DEPRECATION_WARNINGS
+	}
+
+	/** Change the default overlay material max draw distance used by this mesh */
+	UFUNCTION(BlueprintCallable, Category = "Rendering|Material")
+	void SetOverlayMaterialMaxDrawDistance(float InMaxDrawDistance)
+	{
+		WaitUntilAsyncPropertyReleased(EClothAssetAsyncProperties::OverlayMaterialMaxDrawDistance);
+		PRAGMA_DISABLE_DEPRECATION_WARNINGS
+		OverlayMaterialMaxDrawDistance = InMaxDrawDistance;
+		PRAGMA_ENABLE_DEPRECATION_WARNINGS
+	}
+
 	/** Return the enclosed Cloth Collection object. */
 	TArray<TSharedRef<FManagedArrayCollection>>& GetClothCollections() 
 	{
@@ -393,6 +434,16 @@ private:
 	UPROPERTY()
 	FGuid AssetGuid;
 
+	/** Default translucent material to blend on top of this mesh. Mesh will be rendered twice - once with a base material and once with overlay material */
+	UE_DEPRECATED(5.6, "This must be protected for async build, always use the accessors even internally.")
+	UPROPERTY(EditAnywhere, Category = Rendering)
+	TObjectPtr<class UMaterialInterface> OverlayMaterial;
+
+	/** Default max draw distance for overlay material. A distance of 0 indicates that overlay will be culled using primitive max distance. */
+	UE_DEPRECATED(5.6, "This must be protected for async build, always use the accessors even internally.")
+	UPROPERTY(EditAnywhere, Category = Rendering)
+	float OverlayMaterialMaxDrawDistance;
+
 	/** Bounds for this asset. */
 	FBoxSphereBounds Bounds;
 
