Add SetSimulateInEditor bp function to ChaosClothComponent.

diff --git a/Engine/Plugins/ChaosClothAsset/Source/ChaosClothAssetEngine/Private/ChaosClothAsset/ClothComponent.cpp b/Engine/Plugins/ChaosClothAsset/Source/ChaosClothAssetEngine/Private/ChaosClothAsset/ClothComponent.cpp
--- a/Engine/Plugins/ChaosClothAsset/Source/ChaosClothAssetEngine/Private/ChaosClothAsset/ClothComponent.cpp
+++ b/Engine/Plugins/ChaosClothAsset/Source/ChaosClothAssetEngine/Private/ChaosClothAsset/ClothComponent.cpp
@@ -580,6 +580,13 @@ void UChaosClothComponent::SetCollideWithEnvironment(bool bCollide)
 	CollisionSources->SetCollideWithEnvironment(bCollideWithEnvironment);
 }
 
+void UChaosClothComponent::SetSimulateInEditor(const bool bNewSimulateState)
+{
+#if WITH_EDITOR
+	bSimulateInEditor = bNewSimulateState;
+#endif
+}
+
 void UChaosClothComponent::StartNewParallelSimulation(float DeltaTime)
 {
 	if (ClothSimulationProxy.IsValid())
diff --git a/Engine/Plugins/ChaosClothAsset/Source/ChaosClothAssetEngine/Public/ChaosClothAsset/ClothComponent.h b/Engine/Plugins/ChaosClothAsset/Source/ChaosClothAssetEngine/Public/ChaosClothAsset/ClothComponent.h
--- a/Engine/Plugins/ChaosClothAsset/Source/ChaosClothAssetEngine/Public/ChaosClothAsset/ClothComponent.h
+++ b/Engine/Plugins/ChaosClothAsset/Source/ChaosClothAssetEngine/Public/ChaosClothAsset/ClothComponent.h
@@ -137,6 +137,12 @@ public:
 	UFUNCTION(BlueprintCallable, Category = "ClothComponent", Meta = (Keywords = "Chaos Cloth Collide Environment"))
 	bool GetCollideWithEnvironment() const { return bCollideWithEnvironment; }
 
+	/**
+	* Sets whether or not to simulate cloth in the editor.
+	* This is supported only in the editor
+	*/
+	UFUNCTION(BlueprintCallable, Category = "ClothComponent", meta = (DevelopmentOnly, UnsafeDuringActorConstruction = "true"))
+	void SetSimulateInEditor(const bool bNewSimulateState);
 
 	/**
 	 * Return the property collections holding the runtime properties for this cloth component (one per LOD).
