Chaos Cloth Asset - Made the cloth an inherited class of a new UChaosClothAssetBase class to prepare for the new OutfitAsset.

#rnx

diff --git a/Engine/Config/BaseEngine.ini b/Engine/Config/BaseEngine.ini
--- a/Engine/Config/BaseEngine.ini
+++ b/Engine/Config/BaseEngine.ini
@@ -1339,6 +1339,8 @@ ClassName=/Script/OnlineSubsystemUtils.OnlineEngineInterfaceImpl
 +FunctionRedirects=(OldName="SkeletalMeshComponent.SetAnimClass",NewName="SetAnimInstanceClass")
 +FunctionRedirects=(OldName="SkeletalMeshComponent.K2_SetAnimInstanceClass",NewName="SetAnimInstanceClass")
 
++FunctionRedirects=(OldName="ChaosClothComponent.SetClothAsset", NewName="SetAsset")
++PropertyRedirects=(OldName="ChaosClothComponent.ClothAsset", NewName="ChaosClothComponent.Asset")
 
 ; End of releases redirects
 
diff --git a/Engine/Plugins/Animation/MLDeformer/ChaosClothGenerator/Source/ChaosClothGenerator/Private/ChaosClothGenerator.cpp b/Engine/Plugins/Animation/MLDeformer/ChaosClothGenerator/Source/ChaosClothGenerator/Private/ChaosClothGenerator.cpp
--- a/Engine/Plugins/Animation/MLDeformer/ChaosClothGenerator/Source/ChaosClothGenerator/Private/ChaosClothGenerator.cpp
+++ b/Engine/Plugins/Animation/MLDeformer/ChaosClothGenerator/Source/ChaosClothGenerator/Private/ChaosClothGenerator.cpp
@@ -740,7 +740,7 @@ namespace UE::Chaos::ClothGenerator
 		for(int32 Index = 0; Index < Num; ++Index)
 		{
 			UClothGeneratorComponent* const CopyComponent = NewObject<UClothGeneratorComponent>();
-			CopyComponent->SetClothAsset(&Asset);
+			CopyComponent->SetAsset(&Asset);
 			CopyComponent->RegisterComponentWithWorld(World);
 	
 			USkinnedMeshComponent* const PoseComponent = CopyComponent->LeaderPoseComponent.Get() ? CopyComponent->LeaderPoseComponent.Get() : CopyComponent;
@@ -933,7 +933,7 @@ namespace UE::Chaos::ClothGenerator
 		const double Time = FMath::Clamp(AnimationSequence->GetSamplingFrameRate().AsSeconds(Frame), 0., (double)AnimationSequence->GetPlayLength());
 		FAnimExtractContext ExtractionContext(Time);
 	
-		UChaosClothAsset* const ClothAsset = InClothComponent.GetClothAsset();
+		UChaosClothAsset* const ClothAsset = Cast<UChaosClothAsset>(InClothComponent.GetAsset());
 		const FReferenceSkeleton* const ReferenceSkeleton = ClothAsset ? &ClothAsset->GetRefSkeleton() : nullptr;
 		USkeleton* const Skeleton = ClothAsset ? ClothAsset->GetSkeleton() : nullptr;
 		const int32 NumBones = ReferenceSkeleton ? ReferenceSkeleton->GetNum() : 0;
diff --git a/Engine/Plugins/ChaosCloth/Source/ChaosCloth/Private/ChaosCloth/ChaosClothingSimulationCloth.cpp b/Engine/Plugins/ChaosCloth/Source/ChaosCloth/Private/ChaosCloth/ChaosClothingSimulationCloth.cpp
--- a/Engine/Plugins/ChaosCloth/Source/ChaosCloth/Private/ChaosCloth/ChaosClothingSimulationCloth.cpp
+++ b/Engine/Plugins/ChaosCloth/Source/ChaosCloth/Private/ChaosCloth/ChaosClothingSimulationCloth.cpp
@@ -526,7 +526,9 @@ void FClothingSimulationCloth::SetConfig(FClothingSimulationConfig* InConfig)
 	{
 		// Create a default empty config object for coherence
 		PropertyCollection = MakeShared<FManagedArrayCollection>();
+PRAGMA_DISABLE_DEPRECATION_WARNINGS  // TODO Remove these pragmas when the method is removed and the TArray can be correctly deducted again
 		Config = new FClothingSimulationConfig({ PropertyCollection });
+PRAGMA_ENABLE_DEPRECATION_WARNINGS
 	}
 }
 
diff --git a/Engine/Plugins/ChaosCloth/Source/ChaosCloth/Private/ChaosCloth/ChaosClothingSimulationConfig.cpp b/Engine/Plugins/ChaosCloth/Source/ChaosCloth/Private/ChaosCloth/ChaosClothingSimulationConfig.cpp
--- a/Engine/Plugins/ChaosCloth/Source/ChaosCloth/Private/ChaosCloth/ChaosClothingSimulationConfig.cpp
+++ b/Engine/Plugins/ChaosCloth/Source/ChaosCloth/Private/ChaosCloth/ChaosClothingSimulationConfig.cpp
@@ -12,9 +12,23 @@ namespace Chaos
 	{
 	}
 
+	FClothingSimulationConfig::FClothingSimulationConfig(const TArray<TSharedPtr<FManagedArrayCollection>>& InPropertyCollections)
+	{
+		Initialize(InPropertyCollections);
+	}
+
+	FClothingSimulationConfig::FClothingSimulationConfig(std::initializer_list<TSharedPtr<FManagedArrayCollection>> InitList)
+	{
+PRAGMA_DISABLE_DEPRECATION_WARNINGS
+		Initialize(InitList);
+PRAGMA_ENABLE_DEPRECATION_WARNINGS
+	}
+
 	FClothingSimulationConfig::FClothingSimulationConfig(const TArray<TSharedPtr<const FManagedArrayCollection>>& InPropertyCollections)
 	{
+PRAGMA_DISABLE_DEPRECATION_WARNINGS
 		Initialize(InPropertyCollections);
+PRAGMA_ENABLE_DEPRECATION_WARNINGS
 	}
 
 	FClothingSimulationConfig::~FClothingSimulationConfig() = default;
@@ -240,6 +254,34 @@ namespace Chaos
 		Property->AddValue(TEXT("UseLegacyConfig"), bUseLegacyConfig, NonAnimatablePropertyFlags);
 	}
 
+	void FClothingSimulationConfig::Initialize(const TArray<TSharedPtr<FManagedArrayCollection>>& InPropertyCollections)
+	{
+		PropertyCollections.Reset(InPropertyCollections.Num());
+		Properties.Reset(InPropertyCollections.Num());
+
+		for (const TSharedPtr<FManagedArrayCollection>& InPropertyCollection : InPropertyCollections)
+		{
+			TSharedPtr<FManagedArrayCollection>& PropertyCollection = PropertyCollections.Add_GetRef(MakeShared<FManagedArrayCollection>());
+			TUniquePtr<Softs::FCollectionPropertyMutableFacade>& Property = Properties.Add_GetRef(MakeUnique<Softs::FCollectionPropertyMutableFacade>(PropertyCollection));
+			Property->Copy(*InPropertyCollection);
+		}
+	}
+
+	// UE_DEPRECATED(5.6, "Use Initialize(const TArray<TSharedPtr<FManagedArrayCollection>>&).")
+	CHAOSCLOTH_API void FClothingSimulationConfig::Initialize(std::initializer_list<TSharedPtr<FManagedArrayCollection>> InitList)
+	{
+		PropertyCollections.Reset((int32)InitList.size());
+		Properties.Reset((int32)InitList.size());
+
+		for (const TSharedPtr<FManagedArrayCollection>* It = InitList.begin(); It != InitList.end(); ++It)
+		{
+			TSharedPtr<FManagedArrayCollection>& PropertyCollection = PropertyCollections.Add_GetRef(MakeShared<FManagedArrayCollection>());
+			TUniquePtr<Softs::FCollectionPropertyMutableFacade>& Property = Properties.Add_GetRef(MakeUnique<Softs::FCollectionPropertyMutableFacade>(PropertyCollection));
+			Property->Copy(**It);
+		}
+	}
+
+	// UE_DEPRECATED(5.6, "Use Initialize(const TArray<TSharedPtr<FManagedArrayCollection>>&).")
 	void FClothingSimulationConfig::Initialize(const TArray<TSharedPtr<const FManagedArrayCollection>>& InPropertyCollections)
 	{
 		PropertyCollections.Reset(InPropertyCollections.Num());
diff --git a/Engine/Plugins/ChaosCloth/Source/ChaosCloth/Private/ChaosCloth/ChaosClothingSimulationSolver.cpp b/Engine/Plugins/ChaosCloth/Source/ChaosCloth/Private/ChaosCloth/ChaosClothingSimulationSolver.cpp
--- a/Engine/Plugins/ChaosCloth/Source/ChaosCloth/Private/ChaosCloth/ChaosClothingSimulationSolver.cpp
+++ b/Engine/Plugins/ChaosCloth/Source/ChaosCloth/Private/ChaosCloth/ChaosClothingSimulationSolver.cpp
@@ -515,7 +515,9 @@ void FClothingSimulationSolver::SetConfig(FClothingSimulationConfig* InConfig)
 	{
 		// Create a default empty config object for coherence
 		PropertyCollection = MakeShared<FManagedArrayCollection>();
+PRAGMA_DISABLE_DEPRECATION_WARNINGS  // TODO Remove these pragmas when the method is removed and the TArray can be correctly deducted again
 		Config = new FClothingSimulationConfig({ PropertyCollection });
+PRAGMA_ENABLE_DEPRECATION_WARNINGS
 	}
 }
 
diff --git a/Engine/Plugins/ChaosCloth/Source/ChaosCloth/Public/ChaosCloth/ChaosClothingSimulationConfig.h b/Engine/Plugins/ChaosCloth/Source/ChaosCloth/Public/ChaosCloth/ChaosClothingSimulationConfig.h
--- a/Engine/Plugins/ChaosCloth/Source/ChaosCloth/Public/ChaosCloth/ChaosClothingSimulationConfig.h
+++ b/Engine/Plugins/ChaosCloth/Source/ChaosCloth/Public/ChaosCloth/ChaosClothingSimulationConfig.h
@@ -22,6 +22,10 @@ namespace Chaos
 	{
 	public:
 		CHAOSCLOTH_API FClothingSimulationConfig();
+		explicit CHAOSCLOTH_API FClothingSimulationConfig(const TArray<TSharedPtr<FManagedArrayCollection>>& InPropertyCollections);
+		UE_DEPRECATED(5.6, "Use FClothingSimulationConfig(const TArray<TSharedPtr<FManagedArrayCollection>>&) instead.")
+		explicit CHAOSCLOTH_API FClothingSimulationConfig(std::initializer_list<TSharedPtr<FManagedArrayCollection>> InitList);
+		UE_DEPRECATED(5.6, "Use FClothingSimulationConfig(const TArray<TSharedPtr<FManagedArrayCollection>>&) instead.")
 		CHAOSCLOTH_API FClothingSimulationConfig(const TArray<TSharedPtr<const FManagedArrayCollection>>& InPropertyCollections);
 
 		CHAOSCLOTH_API ~FClothingSimulationConfig();
@@ -40,6 +44,10 @@ namespace Chaos
 		CHAOSCLOTH_API void Initialize(const UChaosClothConfig* ClothConfig, const UChaosClothSharedSimConfig* ClothSharedConfig, bool bUseLegacyConfig = false);
 
 		/** Initialize config from an array of property collections (one per LOD). */
+		CHAOSCLOTH_API void Initialize(const TArray<TSharedPtr<FManagedArrayCollection>>& InPropertyCollections);
+		UE_DEPRECATED(5.6, "Use Initialize(const TArray<TSharedPtr<FManagedArrayCollection>>&) instead.")
+		CHAOSCLOTH_API void Initialize(std::initializer_list<TSharedPtr<FManagedArrayCollection>> InitList);
+		UE_DEPRECATED(5.6, "Use Initialize(const TArray<TSharedPtr<FManagedArrayCollection>>&) instead.")
 		CHAOSCLOTH_API void Initialize(const TArray<TSharedPtr<const FManagedArrayCollection>>& InPropertyCollections);
 
 		/** Return a property collection facade for reading properties from this configuration. */
diff --git a/Engine/Plugins/ChaosClothAsset/Source/ChaosClothAssetEngine/Private/ChaosClothAsset/ClothAsset.cpp b/Engine/Plugins/ChaosClothAsset/Source/ChaosClothAssetEngine/Private/ChaosClothAsset/ClothAsset.cpp
--- a/Engine/Plugins/ChaosClothAsset/Source/ChaosClothAssetEngine/Private/ChaosClothAsset/ClothAsset.cpp
+++ b/Engine/Plugins/ChaosClothAsset/Source/ChaosClothAssetEngine/Private/ChaosClothAsset/ClothAsset.cpp
@@ -4,7 +4,6 @@
 #include "ChaosClothAsset/ClothAssetBuilder.h"
 #include "ChaosClothAsset/ClothAssetPrivate.h"
 #include "ChaosClothAsset/ClothCollectionGroup.h"
-#include "ChaosClothAsset/ClothComponent.h"
 #include "ChaosClothAsset/ClothGeometryTools.h"
 #include "ChaosClothAsset/ClothSimulationModel.h"
 #include "ChaosClothAsset/CollectionClothFacade.h"
@@ -17,14 +16,12 @@
 #include "Engine/RendererSettings.h"
 #include "Engine/SkeletalMesh.h"
 #include "Engine/SkinnedAssetAsyncCompileUtils.h"
-#include "Features/IModularFeatures.h"
 #include "GeometryCollection/ManagedArrayCollection.h"
 #include "PhysicsEngine/PhysicsAsset.h"
 #include "Rendering/SkeletalMeshModel.h"
 #include "Rendering/SkeletalMeshRenderData.h"
 #include "UObject/Package.h"
 #include "EngineUtils.h"
-#include "Engine/Engine.h"
 #if WITH_EDITOR
 #include "IMeshBuilderModule.h"
 #include "DerivedDataCacheInterface.h"
@@ -45,17 +42,6 @@ FAutoConsoleVariableRef CVarClothCollectionOnlyCookRequiredFacades(
 	bClothCollectionOnlyCookRequiredFacades,
 	TEXT("Default setting for culling managed arrays on the cloth collection during the cook. Default[true]"));
 
-
-const TCHAR* MinLodQualityLevelCVarName = TEXT("p.ClothAsset.MinLodQualityLevel");
-const TCHAR* MinLodQualityLevelScalabilitySection = TEXT("ViewDistanceQuality");
-int32 MinLodQualityLevel = -1;
-FAutoConsoleVariableRef CVarClothAssetMinLodQualityLevel(
-	MinLodQualityLevelCVarName,
-	MinLodQualityLevel,
-	TEXT("The quality level for the Min stripping LOD. \n"),
-	FConsoleVariableDelegate::CreateStatic(&UChaosClothAsset::OnLodStrippingQualityLevelChanged),
-	ECVF_Scalability);
-	
 ::Chaos::FChaosArchive& Serialize(::Chaos::FChaosArchive& Ar, TArray<TSharedRef<FManagedArrayCollection>>& ClothCollections)
 {
 	Ar.UsingCustomVersion(FUE5MainStreamObjectVersion::GUID);
@@ -164,28 +150,16 @@ TArray<TSharedRef<FManagedArrayCollection>> TrimOnCook(const FString InAssetName
 
 UChaosClothAsset::UChaosClothAsset(const FObjectInitializer& ObjectInitializer)
 	: Super(ObjectInitializer)
-	, DataflowInstance(this, nullptr, FName(TEXT("ClothAssetTerminal")))
-	, MinQualityLevelLOD(0)
-	, DisableBelowMinLodStripping(FPerPlatformBool(false))
-	, MinLod(0)
 #if WITH_EDITORONLY_DATA
 	, MeshModel(MakeShareable(new FSkeletalMeshModel()))
 #endif
 {
+	DataflowInstance.SetDataflowTerminal(TEXT("ClothAssetTerminal"));
+
 	// Setup a single LOD's Cloth Collection
 	TSharedRef<FManagedArrayCollection>& ClothCollection = GetClothCollections().Emplace_GetRef(MakeShared<FManagedArrayCollection>());
 	UE::Chaos::ClothAsset::FCollectionClothFacade ClothFacade(ClothCollection);
 	ClothFacade.DefineSchema();
-
-	// Add the LODInfo for the default LOD 0
-	LODInfo.SetNum(1);
-
-	// Set default skeleton (must be done after having added the LOD)
-	constexpr bool bRebuildModels = false;
-	constexpr bool bRebindMeshes = false;
-	SetReferenceSkeleton(nullptr, bRebuildModels, bRebindMeshes);
-
-	MinQualityLevelLOD.SetQualityLevelCVarForCooking(UE::Chaos::ClothAsset::Private::MinLodQualityLevelCVarName, UE::Chaos::ClothAsset::Private::MinLodQualityLevelScalabilitySection);
 }
 
 UChaosClothAsset::UChaosClothAsset(FVTableHelper& Helper)
@@ -197,16 +171,6 @@ PRAGMA_DISABLE_DEPRECATION_WARNINGS
 UChaosClothAsset::~UChaosClothAsset() = default;
 PRAGMA_ENABLE_DEPRECATION_WARNINGS
 
-FSkeletalMeshLODInfo* UChaosClothAsset::GetLODInfo(int32 Index)
-{
-	return LODInfo.IsValidIndex(Index) ? &LODInfo[Index] : nullptr;
-}
-
-const FSkeletalMeshLODInfo* UChaosClothAsset::GetLODInfo(int32 Index) const
-{
-	return LODInfo.IsValidIndex(Index) ? &LODInfo[Index] : nullptr;
-}
-
 FMatrix UChaosClothAsset::GetComposedRefPoseMatrix(FName InBoneName) const
 {
 	FMatrix LocalPose(FMatrix::Identity);
@@ -487,60 +451,6 @@ void UChaosClothAsset::FinishPostLoadInternal(FSkinnedAssetPostLoadContext& Cont
 #endif
 }
 
-void UChaosClothAsset::BeginDestroy()
-{
-	check(IsInGameThread());
-
-	Super::BeginDestroy();
-
-	// Release the mesh's render resources now
-	ReleaseResources();
-}
-
-bool UChaosClothAsset::IsReadyForFinishDestroy()
-{
-	if (!Super::IsReadyForFinishDestroy())
-	{
-		return false;
-	}
-
-	ReleaseResources();
-
-	// see if we have hit the resource flush fence
-	return ReleaseResourcesFence.IsFenceComplete();
-}
-
-void UChaosClothAsset::InitResources()
-{
-	LLM_SCOPE_BYNAME(TEXT("ClothAsset/InitResources"));
-
-	// Build the material channel data used by the texture streamer
-	UpdateUVChannelData(false);
-
-	if (SkeletalMeshRenderData.IsValid())
-	{
-		TArray<UMorphTarget*> DummyMorphTargets;  // Not even used by InitResources() at the moment
-		SkeletalMeshRenderData->InitResources(GetHasVertexColors(), DummyMorphTargets, this);
-	}
-}
-
-void UChaosClothAsset::ReleaseResources()
-{
-	if (SkeletalMeshRenderData && SkeletalMeshRenderData->IsInitialized())
-	{
-		if (GIsEditor && !GIsPlayInEditorWorld)
-		{
-			// Flush the rendering command to be sure there is no command left that can create/modify a rendering resource
-			FlushRenderingCommands();
-		}
-
-		SkeletalMeshRenderData->ReleaseResources();
-
-		// Insert a fence to signal when these commands completed
-		ReleaseResourcesFence.BeginFence();
-	}
-}
-
 void UChaosClothAsset::CalculateInvRefMatrices()
 {
 	auto GetRefPoseMatrix = [this](int32 BoneIndex)->FMatrix
@@ -666,51 +576,6 @@ void UChaosClothAsset::Build(TArray<FChaosClothAssetLodTransitionDataCache>* InO
 	ReregisterComponents();
 }
 
-#if WITH_EDITOR
-void UChaosClothAsset::ExecuteBuildInternal(FSkinnedAssetBuildContext& Context)
-{
-	TRACE_CPUPROFILER_EVENT_SCOPE(UChaosClothAsset::ExecuteBuildInternal);
-
-	// This scope allows us to use any locked properties without causing stalls
-	FSkinnedAssetAsyncBuildScope AsyncBuildScope(this);
-
-	// rebuild render data from imported model
-	CacheDerivedData(&Context);
-
-	// Build the material channel data used by the texture streamer
-	UpdateUVChannelData(true);
-}
-
-void UChaosClothAsset::BeginBuildInternal(FSkinnedAssetBuildContext& Context)
-{
-	TRACE_CPUPROFILER_EVENT_SCOPE(UChaosClothAsset::BeginBuildInternal);
-
-	SetInternalFlags(EInternalObjectFlags::Async);
-	
-	// Unregister all instances of this component
-	Context.RecreateRenderStateContext = MakeUnique<FSkinnedMeshComponentRecreateRenderStateContext>(this, false);
-
-	// Release the render data resources.
-	ReleaseResources();
-
-	// Flush the resource release commands to the rendering thread to ensure that the build doesn't occur while a resource is still
-	// allocated, and potentially accessing the UChaosClothAsset.
-	ReleaseResourcesFence.Wait();
-
-	// Lock all properties that should not be modified/accessed during async post-load
-	AcquireAsyncProperty();
-}
-
-void UChaosClothAsset::FinishBuildInternal(FSkinnedAssetBuildContext& Context)
-{
-	TRACE_CPUPROFILER_EVENT_SCOPE(UChaosClothAsset::FinishBuildInternal);
-
-	ClearInternalFlags(EInternalObjectFlags::Async);
-
-	ReleaseAsyncProperty();
-}
-#endif // #if WITH_EDITOR
-
 #if WITH_EDITOR
 void UChaosClothAsset::BuildMeshModel()
 {
@@ -753,84 +618,17 @@ const FMeshUVChannelInfo* UChaosClothAsset::GetUVChannelData(int32 MaterialIndex
 	return nullptr;
 }
 
-FSkeletalMeshRenderData* UChaosClothAsset::GetResourceForRendering() const
-{
-	WaitUntilAsyncPropertyReleased(EClothAssetAsyncProperties::RenderData);
-	return SkeletalMeshRenderData.Get();
-}
-
-void UChaosClothAsset::SetResourceForRendering(TUniquePtr<FSkeletalMeshRenderData>&& InSkeletalMeshRenderData)
-{
-	WaitUntilAsyncPropertyReleased(EClothAssetAsyncProperties::RenderData);
-	SkeletalMeshRenderData = MoveTemp(InSkeletalMeshRenderData);
-}
-
-void UChaosClothAsset::WaitUntilAsyncPropertyReleased(EClothAssetAsyncProperties AsyncProperties, ESkinnedAssetAsyncPropertyLockType LockType) const
-{
-	// Cast strongly typed enum to uint64
-	WaitUntilAsyncPropertyReleasedInternal((uint64)AsyncProperties, LockType);
-}
-
 FString UChaosClothAsset::GetAsyncPropertyName(uint64 Property) const
 {
 	return StaticEnum<EClothAssetAsyncProperties>()->GetNameByValue(Property).ToString();
 }
 
-bool UChaosClothAsset::IsMinLodQualityLevelEnable() const
-{
-	return (GEngine && GEngine->UseClothAssetMinLODPerQualityLevels);
-}
-
-void UChaosClothAsset::OnLodStrippingQualityLevelChanged(IConsoleVariable* Variable) 
-{
-#if WITH_EDITOR || PLATFORM_DESKTOP
-	if (GEngine && GEngine->UseClothAssetMinLODPerQualityLevels)
-	{
-		for (TObjectIterator<UChaosClothAsset> It; It; ++It)
-		{
-			UChaosClothAsset* ClothAsset = *It;
-			if (ClothAsset && ClothAsset->GetQualityLevelMinLod().PerQuality.Num() > 0)
-			{
-				FSkinnedMeshComponentRecreateRenderStateContext Context(ClothAsset, false);
-			}
-		}
-	}
-#endif
-}
-
-int32 UChaosClothAsset::GetMinLodIdx(bool bForceLowestLODIndex) const
-{
-	if (IsMinLodQualityLevelEnable())
-	{
-		return bForceLowestLODIndex ? GetQualityLevelMinLod().GetLowestValue() : GetQualityLevelMinLod().GetValue(UE::Chaos::ClothAsset::Private::MinLodQualityLevel);
-	}
-	else
-	{
-		return GetMinLod().GetValue();
-	}
-}
-
-int32 UChaosClothAsset::GetPlatformMinLODIdx(const ITargetPlatform* InTargetPlatform) const
+const TArray<TSharedRef<FManagedArrayCollection>>& UChaosClothAsset::GetCollections(int32) const
 {
-#if WITH_EDITOR
-	check(InTargetPlatform);
-	if (IsMinLodQualityLevelEnable())
-	{
-		// get all supported quality level from scalability + engine ini files
-		return GetQualityLevelMinLod().GetValueForPlatform(InTargetPlatform);
-	}
-	else
-	{
-		return GetMinLod().GetValueForPlatform(*InTargetPlatform->IniPlatformName());
-	}
-#else
-	return 0;
-#endif
-}
-
-const FPerPlatformInt& UChaosClothAsset::GetMinLod() const
-{
-	return MinLod;
+	WaitUntilAsyncPropertyReleased(EClothAssetAsyncProperties::ClothCollection);
+PRAGMA_DISABLE_DEPRECATION_WARNINGS
+	return ClothCollections;
+PRAGMA_ENABLE_DEPRECATION_WARNINGS
 }
 
 #if WITH_EDITOR
@@ -932,6 +730,11 @@ bool UChaosClothAsset::IsInitialBuildDone() const
 }
 #endif // WITH_EDITOR
 
+bool UChaosClothAsset::HasValidClothSimulationModels() const
+{
+	return ClothSimulationModel && ClothSimulationModel->GetNumLods();
+}
+
 void UChaosClothAsset::SetPhysicsAsset(UPhysicsAsset* InPhysicsAsset)
 {
 	using namespace UE::Chaos::ClothAsset;
@@ -948,28 +751,13 @@ void UChaosClothAsset::SetReferenceSkeleton(const FReferenceSkeleton* ReferenceS
 {
 	using namespace UE::Chaos::ClothAsset;
 
-	// Update the reference skeleton
-	if (ReferenceSkeleton)
+	if (!ReferenceSkeleton)
 	{
-		GetRefSkeleton() = *ReferenceSkeleton;
+		bRebindMeshes = true;  // Force the binding when a default reference skeleton is being created 
 	}
-	else
-	{
-		// Create a default reference skeleton
-		GetRefSkeleton().Empty(1);
-		FReferenceSkeletonModifier ReferenceSkeletonModifier(GetRefSkeleton(), nullptr);
 
-		FMeshBoneInfo MeshBoneInfo;
-		constexpr const TCHAR* RootName = TEXT("Root");
-		MeshBoneInfo.ParentIndex = INDEX_NONE;
-#if WITH_EDITORONLY_DATA
-		MeshBoneInfo.ExportName = RootName;
-#endif
-		MeshBoneInfo.Name = FName(RootName);
-		ReferenceSkeletonModifier.Add(MeshBoneInfo, FTransform::Identity);
-
-		bRebindMeshes = true; // Force the binding when a default reference skeleton is being created 
-	}
+	// Update the reference skeleton
+	Super::SetReferenceSkeleton(ReferenceSkeleton);
 
 	// Rebind the meshes
 	if (bRebindMeshes)
@@ -1033,83 +821,3 @@ void UChaosClothAsset::CopySimMeshToRenderMesh(UMaterialInterface* Material)
 		}
 	}
 }
-
-void UChaosClothAsset::ReregisterComponents()
-{
-	// Recreate the simulation proxies with the updated physics asset
-	for (TObjectIterator<UChaosClothComponent> ObjectIterator; ObjectIterator; ++ObjectIterator)
-	{
-		if (UChaosClothComponent* const Component = *ObjectIterator)
-		{
-			if (Component->GetClothAsset() == this)
-			{
-				const FComponentReregisterContext Context(Component);  // Context goes out of scope, causing the Component to be re-registered
-			}
-		}
-	}
-}
-
-#if WITH_EDITORONLY_DATA
-
-void UChaosClothAsset::SetPreviewSceneSkeletalMesh(USkeletalMesh* Mesh)
-{
-	PreviewSceneSkeletalMesh = Mesh;
-}
-
-USkeletalMesh* UChaosClothAsset::GetPreviewSceneSkeletalMesh() const
-{
-	// Load the SkeletalMesh asset if it's not already loaded
-	return PreviewSceneSkeletalMesh.LoadSynchronous();
-}
-
-void UChaosClothAsset::SetPreviewSceneAnimation(UAnimationAsset* Animation)
-{
-	PreviewSceneAnimation = Animation;
-}
-
-UAnimationAsset* UChaosClothAsset::GetPreviewSceneAnimation() const
-{
-	// Load the animation asset if it's not already loaded
-	return PreviewSceneAnimation.LoadSynchronous();
-}
-
-#endif
-
-TObjectPtr<UDataflowBaseContent> UChaosClothAsset::CreateDataflowContent()
-{
-	TObjectPtr<UDataflowSkeletalContent> SkeletalContent = UE::DataflowContextHelpers::CreateNewDataflowContent<UDataflowSkeletalContent>(this);
-
-	SkeletalContent->SetDataflowOwner(this);
-	SkeletalContent->SetTerminalAsset(this);
-
-	WriteDataflowContent(SkeletalContent);
-	
-	return SkeletalContent;
-}
-
-void UChaosClothAsset::WriteDataflowContent(const TObjectPtr<UDataflowBaseContent>& DataflowContent) const
-{
-	if(const TObjectPtr<UDataflowSkeletalContent> SkeletalContent = Cast<UDataflowSkeletalContent>(DataflowContent))
-	{
-		SkeletalContent->SetDataflowAsset(DataflowInstance.GetDataflowAsset());
-		SkeletalContent->SetDataflowTerminal(DataflowInstance.GetDataflowTerminal().ToString());
-
-#if WITH_EDITORONLY_DATA
-		SkeletalContent->SetAnimationAsset(GetPreviewSceneAnimation());
-		SkeletalContent->SetSkeletalMesh(GetPreviewSceneSkeletalMesh());
-#endif
-	}
-}
-
-void UChaosClothAsset::ReadDataflowContent(const TObjectPtr<UDataflowBaseContent>& DataflowContent)
-{
-	if(const TObjectPtr<UDataflowSkeletalContent> SkeletalContent = Cast<UDataflowSkeletalContent>(DataflowContent))
-	{
-#if WITH_EDITORONLY_DATA
-		PreviewSceneAnimation = SkeletalContent->GetAnimationAsset();
-		PreviewSceneSkeletalMesh = SkeletalContent->GetSkeletalMesh();
-#endif
-	}
-}
-
-
diff --git a/Engine/Plugins/ChaosClothAsset/Source/ChaosClothAssetEngine/Private/ChaosClothAsset/ClothAssetBase.cpp b/Engine/Plugins/ChaosClothAsset/Source/ChaosClothAssetEngine/Private/ChaosClothAsset/ClothAssetBase.cpp
new file mode 100644
--- /dev/null
+++ b/Engine/Plugins/ChaosClothAsset/Source/ChaosClothAssetEngine/Private/ChaosClothAsset/ClothAssetBase.cpp
@@ -0,0 +1,392 @@
+// Copyright Epic Games, Inc. All Rights Reserved.
+
+#include "ChaosClothAsset/ClothAssetBase.h"
+#if WITH_EDITORONLY_DATA
+#include "Animation/AnimationAsset.h"
+#endif
+#include "ChaosClothAsset/ClothComponent.h"
+#include "Engine/Engine.h"
+#include "Engine/SkeletalMesh.h"
+#include "Engine/SkinnedAssetAsyncCompileUtils.h"
+#include "Rendering/SkeletalMeshRenderData.h"
+#if WITH_EDITOR
+#include "Interfaces/ITargetPlatform.h"
+#endif
+
+#include UE_INLINE_GENERATED_CPP_BY_NAME(ClothAssetBase)
+namespace UE::Chaos::ClothAsset::Private
+{
+
+	const TCHAR* MinLodQualityLevelCVarName = TEXT("p.ClothAsset.MinLodQualityLevel");
+	const TCHAR* MinLodQualityLevelScalabilitySection = TEXT("ViewDistanceQuality");
+	int32 MinLodQualityLevel = -1;
+	FAutoConsoleVariableRef CVarClothAssetMinLodQualityLevel(
+		MinLodQualityLevelCVarName,
+		MinLodQualityLevel,
+		TEXT("The quality level for the Min stripping LOD. \n"),
+		FConsoleVariableDelegate::CreateLambda([](IConsoleVariable* /*Variable*/)
+			{
+#if WITH_EDITOR || PLATFORM_DESKTOP
+				if (GEngine && GEngine->UseClothAssetMinLODPerQualityLevels)
+				{
+					for (TObjectIterator<UChaosClothAssetBase> It; It; ++It)
+					{
+						UChaosClothAssetBase* ClothAsset = *It;
+						if (ClothAsset && ClothAsset->GetQualityLevelMinLod().PerQuality.Num() > 0)
+						{
+							FSkinnedMeshComponentRecreateRenderStateContext Context(ClothAsset, false);
+						}
+					}
+				}
+#endif
+			}),
+		ECVF_Scalability);
+}
+
+/** Used for locking resources during async building. */
+enum class UChaosClothAssetBase::EAsyncProperties : uint32
+{
+	None = 0,
+	RenderData = 1 << 0,
+	RefSkeleton = 1 << 1,
+	HasVertexColors = 1 << 2,
+	OverlayMaterial = 1 << 3,
+	OverlayMaterialMaxDrawDistance = 1 << 4,
+	All = MAX_uint32
+};
+
+UChaosClothAssetBase::UChaosClothAssetBase(const FObjectInitializer& ObjectInitializer)
+	: Super(ObjectInitializer)
+	, DataflowInstance(this)
+	, MinQualityLevelLOD(0)
+	, DisableBelowMinLodStripping(FPerPlatformBool(false))
+	, MinLod(0)
+{
+
+	// Add the LODInfo for the default LOD 0
+	LODInfo.SetNum(1);
+
+	// Set default skeleton (must be done after having added the LOD)
+	SetReferenceSkeleton(nullptr);
+
+	MinQualityLevelLOD.SetQualityLevelCVarForCooking(UE::Chaos::ClothAsset::Private::MinLodQualityLevelCVarName, UE::Chaos::ClothAsset::Private::MinLodQualityLevelScalabilitySection);
+}
+
+void UChaosClothAssetBase::SetReferenceSkeleton(const FReferenceSkeleton* ReferenceSkeleton)
+{
+	// Update the reference skeleton
+	if (ReferenceSkeleton)
+	{
+		GetRefSkeleton() = *ReferenceSkeleton;
+	}
+	else
+	{
+		// Create a default reference skeleton
+		GetRefSkeleton().Empty(1);
+		FReferenceSkeletonModifier ReferenceSkeletonModifier(GetRefSkeleton(), nullptr);
+
+		FMeshBoneInfo MeshBoneInfo;
+		constexpr const TCHAR* RootName = TEXT("Root");
+		MeshBoneInfo.ParentIndex = INDEX_NONE;
+#if WITH_EDITORONLY_DATA
+		MeshBoneInfo.ExportName = RootName;
+#endif
+		MeshBoneInfo.Name = FName(RootName);
+		ReferenceSkeletonModifier.Add(MeshBoneInfo, FTransform::Identity);
+	}
+}
+
+TObjectPtr<UDataflowBaseContent> UChaosClothAssetBase::CreateDataflowContent()
+{
+	TObjectPtr<UDataflowSkeletalContent> SkeletalContent = UE::DataflowContextHelpers::CreateNewDataflowContent<UDataflowSkeletalContent>(this);
+
+	SkeletalContent->SetDataflowOwner(this);
+	SkeletalContent->SetTerminalAsset(this);
+
+	WriteDataflowContent(SkeletalContent);
+
+	return SkeletalContent;
+}
+
+void UChaosClothAssetBase::WriteDataflowContent(const TObjectPtr<UDataflowBaseContent>& DataflowContent) const
+{
+	if (const TObjectPtr<UDataflowSkeletalContent> SkeletalContent = Cast<UDataflowSkeletalContent>(DataflowContent))
+	{
+		SkeletalContent->SetDataflowAsset(DataflowInstance.GetDataflowAsset());
+		SkeletalContent->SetDataflowTerminal(DataflowInstance.GetDataflowTerminal().ToString());
+
+#if WITH_EDITORONLY_DATA
+		SkeletalContent->SetAnimationAsset(GetPreviewSceneAnimation());
+		SkeletalContent->SetSkeletalMesh(GetPreviewSceneSkeletalMesh());
+#endif
+	}
+}
+
+void UChaosClothAssetBase::ReadDataflowContent(const TObjectPtr<UDataflowBaseContent>& DataflowContent)
+{
+	if (const TObjectPtr<UDataflowSkeletalContent> SkeletalContent = Cast<UDataflowSkeletalContent>(DataflowContent))
+	{
+#if WITH_EDITORONLY_DATA
+		PreviewSceneAnimation = SkeletalContent->GetAnimationAsset();
+		PreviewSceneSkeletalMesh = SkeletalContent->GetSkeletalMesh();
+#endif
+	}
+}
+
+void UChaosClothAssetBase::BeginDestroy()
+{
+	check(IsInGameThread());
+
+	Super::BeginDestroy();
+
+	// Release the mesh's render resources now
+	ReleaseResources();
+}
+
+bool UChaosClothAssetBase::IsReadyForFinishDestroy()
+{
+	if (!Super::IsReadyForFinishDestroy())
+	{
+		return false;
+	}
+
+	ReleaseResources();
+
+	// see if we have hit the resource flush fence
+	return ReleaseResourcesFence.IsFenceComplete();
+}
+
+FSkeletalMeshLODInfo* UChaosClothAssetBase::GetLODInfo(int32 Index)
+{
+	return LODInfo.IsValidIndex(Index) ? &LODInfo[Index] : nullptr;
+}
+
+const FSkeletalMeshLODInfo* UChaosClothAssetBase::GetLODInfo(int32 Index) const
+{
+	return LODInfo.IsValidIndex(Index) ? &LODInfo[Index] : nullptr;
+}
+
+FReferenceSkeleton& UChaosClothAssetBase::GetRefSkeleton()
+{
+	WaitUntilAsyncPropertyReleased(EAsyncProperties::RefSkeleton);
+PRAGMA_DISABLE_DEPRECATION_WARNINGS
+	return RefSkeleton;
+PRAGMA_ENABLE_DEPRECATION_WARNINGS
+}
+const FReferenceSkeleton& UChaosClothAssetBase::GetRefSkeleton() const
+{
+	WaitUntilAsyncPropertyReleased(EAsyncProperties::RefSkeleton, ESkinnedAssetAsyncPropertyLockType::ReadOnly);
+PRAGMA_DISABLE_DEPRECATION_WARNINGS
+	return RefSkeleton;
+PRAGMA_ENABLE_DEPRECATION_WARNINGS
+}
+
+FSkeletalMeshRenderData* UChaosClothAssetBase::GetResourceForRendering() const
+{
+	WaitUntilAsyncPropertyReleased(EAsyncProperties::RenderData);
+	return SkeletalMeshRenderData.Get();
+}
+
+UMaterialInterface* UChaosClothAssetBase::GetOverlayMaterial() const
+{
+	WaitUntilAsyncPropertyReleased(EAsyncProperties::OverlayMaterial, ESkinnedAssetAsyncPropertyLockType::ReadOnly);
+PRAGMA_DISABLE_DEPRECATION_WARNINGS
+	return OverlayMaterial;
+PRAGMA_ENABLE_DEPRECATION_WARNINGS
+}
+
+float UChaosClothAssetBase::GetOverlayMaterialMaxDrawDistance() const
+{
+	WaitUntilAsyncPropertyReleased(EAsyncProperties::OverlayMaterialMaxDrawDistance, ESkinnedAssetAsyncPropertyLockType::ReadOnly);
+PRAGMA_DISABLE_DEPRECATION_WARNINGS
+	return OverlayMaterialMaxDrawDistance;
+PRAGMA_ENABLE_DEPRECATION_WARNINGS
+}
+
+int32 UChaosClothAssetBase::GetMinLodIdx(bool bForceLowestLODIdx) const
+{
+	if (IsMinLodQualityLevelEnable())
+	{
+		return bForceLowestLODIdx ? GetQualityLevelMinLod().GetLowestValue() : GetQualityLevelMinLod().GetValue(UE::Chaos::ClothAsset::Private::MinLodQualityLevel);
+	}
+	else
+	{
+		return GetMinLod().GetValue();
+	}
+}
+
+bool UChaosClothAssetBase::GetHasVertexColors() const
+{
+	WaitUntilAsyncPropertyReleased(EAsyncProperties::HasVertexColors, ESkinnedAssetAsyncPropertyLockType::ReadOnly);
+PRAGMA_DISABLE_DEPRECATION_WARNINGS
+	return bHasVertexColors != 0;
+PRAGMA_ENABLE_DEPRECATION_WARNINGS
+}
+
+int32 UChaosClothAssetBase::GetPlatformMinLODIdx(const ITargetPlatform* TargetPlatform) const
+{
+#if WITH_EDITOR
+	check(TargetPlatform);
+	if (IsMinLodQualityLevelEnable())
+	{
+		// get all supported quality level from scalability + engine ini files
+		return GetQualityLevelMinLod().GetValueForPlatform(TargetPlatform);
+	}
+	else
+	{
+		return GetMinLod().GetValueForPlatform(*TargetPlatform->IniPlatformName());
+	}
+#else
+	return 0;
+#endif
+}
+
+bool UChaosClothAssetBase::IsMinLodQualityLevelEnable() const
+{
+	return (GEngine && GEngine->UseClothAssetMinLODPerQualityLevels);
+}
+
+void UChaosClothAssetBase::SetOverlayMaterial(UMaterialInterface* NewOverlayMaterial)
+{
+	WaitUntilAsyncPropertyReleased(EAsyncProperties::OverlayMaterial);
+PRAGMA_DISABLE_DEPRECATION_WARNINGS
+	OverlayMaterial = NewOverlayMaterial;
+PRAGMA_ENABLE_DEPRECATION_WARNINGS
+}
+
+void UChaosClothAssetBase::SetOverlayMaterialMaxDrawDistance(float InMaxDrawDistance)
+{
+	WaitUntilAsyncPropertyReleased(EAsyncProperties::OverlayMaterialMaxDrawDistance);
+PRAGMA_DISABLE_DEPRECATION_WARNINGS
+	OverlayMaterialMaxDrawDistance = InMaxDrawDistance;
+PRAGMA_ENABLE_DEPRECATION_WARNINGS
+}
+
+void UChaosClothAssetBase::InitResources()
+{
+	LLM_SCOPE_BYNAME(TEXT("ClothAsset/InitResources"));
+
+	// Build the material channel data used by the texture streamer
+	UpdateUVChannelData(false);
+
+	if (SkeletalMeshRenderData.IsValid())
+	{
+		TArray<UMorphTarget*> DummyMorphTargets;  // Not even used by InitResources() at the moment
+		SkeletalMeshRenderData->InitResources(GetHasVertexColors(), DummyMorphTargets, this);
+	}
+}
+
+void UChaosClothAssetBase::ReleaseResources()
+{
+	if (SkeletalMeshRenderData && SkeletalMeshRenderData->IsInitialized())
+	{
+		if (GIsEditor && !GIsPlayInEditorWorld)
+		{
+			// Flush the rendering command to be sure there is no command left that can create/modify a rendering resource
+			FlushRenderingCommands();
+		}
+
+		SkeletalMeshRenderData->ReleaseResources();
+
+		// Insert a fence to signal when these commands completed
+		ReleaseResourcesFence.BeginFence();
+	}
+}
+
+void UChaosClothAssetBase::SetResourceForRendering(TUniquePtr<FSkeletalMeshRenderData>&& InSkeletalMeshRenderData)
+{
+	WaitUntilAsyncPropertyReleased(EAsyncProperties::RenderData);
+	SkeletalMeshRenderData = MoveTemp(InSkeletalMeshRenderData);
+}
+
+#if WITH_EDITOR
+void UChaosClothAssetBase::ExecuteBuildInternal(FSkinnedAssetBuildContext& Context)
+{
+	TRACE_CPUPROFILER_EVENT_SCOPE(UChaosClothAsset::ExecuteBuildInternal);
+
+	// This scope allows us to use any locked properties without causing stalls
+	FSkinnedAssetAsyncBuildScope AsyncBuildScope(this);
+
+	// rebuild render data from imported model
+	CacheDerivedData(&Context);
+
+	// Build the material channel data used by the texture streamer
+	UpdateUVChannelData(true);
+}
+
+void UChaosClothAssetBase::BeginBuildInternal(FSkinnedAssetBuildContext& Context)
+{
+	TRACE_CPUPROFILER_EVENT_SCOPE(UChaosClothAsset::BeginBuildInternal);
+
+	SetInternalFlags(EInternalObjectFlags::Async);
+
+	// Unregister all instances of this component
+	Context.RecreateRenderStateContext = MakeUnique<FSkinnedMeshComponentRecreateRenderStateContext>(this, false);
+
+	// Release the render data resources.
+	ReleaseResources();
+
+	// Flush the resource release commands to the rendering thread to ensure that the build doesn't occur while a resource is still
+	// allocated, and potentially accessing the UChaosClothAsset.
+	ReleaseResourcesFence.Wait();
+
+	// Lock all properties that should not be modified/accessed during async post-load
+	AcquireAsyncProperty();
+}
+
+void UChaosClothAssetBase::FinishBuildInternal(FSkinnedAssetBuildContext& Context)
+{
+	TRACE_CPUPROFILER_EVENT_SCOPE(UChaosClothAsset::FinishBuildInternal);
+
+	ClearInternalFlags(EInternalObjectFlags::Async);
+
+	ReleaseAsyncProperty();
+}
+#endif // #if WITH_EDITOR
+
+#if WITH_EDITORONLY_DATA
+void UChaosClothAssetBase::SetPreviewSceneSkeletalMesh(USkeletalMesh* Mesh)
+{
+	PreviewSceneSkeletalMesh = Mesh;
+}
+
+USkeletalMesh* UChaosClothAssetBase::GetPreviewSceneSkeletalMesh() const
+{
+	// Load the SkeletalMesh asset if it's not already loaded
+	return PreviewSceneSkeletalMesh.LoadSynchronous();
+}
+
+void UChaosClothAssetBase::SetPreviewSceneAnimation(UAnimationAsset* Animation)
+{
+	PreviewSceneAnimation = Animation;
+}
+
+UAnimationAsset* UChaosClothAssetBase::GetPreviewSceneAnimation() const
+{
+	// Load the animation asset if it's not already loaded
+	return PreviewSceneAnimation.LoadSynchronous();
+}
+#endif  // #if WITH_EDITORONLY_DATA
+
+void UChaosClothAssetBase::SetHasVertexColors(bool InbHasVertexColors)
+{
+	WaitUntilAsyncPropertyReleased(EAsyncProperties::HasVertexColors);
+PRAGMA_DISABLE_DEPRECATION_WARNINGS
+	bHasVertexColors = InbHasVertexColors;
+PRAGMA_ENABLE_DEPRECATION_WARNINGS
+}
+
+void UChaosClothAssetBase::ReregisterComponents()
+{
+	// Recreate the simulation proxies with the updated physics asset
+	for (TObjectIterator<UChaosClothComponent> ObjectIterator; ObjectIterator; ++ObjectIterator)
+	{
+		if (UChaosClothComponent* const Component = *ObjectIterator)
+		{
+			if (Component->GetAsset() == this)
+			{
+				const FComponentReregisterContext Context(Component);  // Context goes out of scope, causing the Component to be re-registered
+			}
+		}
+	}
+}
diff --git a/Engine/Plugins/ChaosClothAsset/Source/ChaosClothAssetEngine/Private/ChaosClothAsset/ClothComponent.cpp b/Engine/Plugins/ChaosClothAsset/Source/ChaosClothAssetEngine/Private/ChaosClothAsset/ClothComponent.cpp
--- a/Engine/Plugins/ChaosClothAsset/Source/ChaosClothAssetEngine/Private/ChaosClothAsset/ClothComponent.cpp
+++ b/Engine/Plugins/ChaosClothAsset/Source/ChaosClothAssetEngine/Private/ChaosClothAsset/ClothComponent.cpp
@@ -42,20 +42,30 @@ UChaosClothComponent::UChaosClothComponent(FVTableHelper& Helper)
 
 UChaosClothComponent::~UChaosClothComponent() = default;
 
-void UChaosClothComponent::SetClothAsset(UChaosClothAsset* InClothAsset)
+void UChaosClothComponent::SetAsset(UChaosClothAssetBase* InAsset)
 {
-	SetSkinnedAssetAndUpdate(InClothAsset);
+	SetSkinnedAssetAndUpdate(InAsset);
 
-	PRAGMA_DISABLE_DEPRECATION_WARNINGS
 #if WITH_EDITORONLY_DATA
-		ClothAsset = InClothAsset;
+PRAGMA_DISABLE_DEPRECATION_WARNINGS
+	Asset = InAsset;
+PRAGMA_ENABLE_DEPRECATION_WARNINGS
 #endif
-	PRAGMA_ENABLE_DEPRECATION_WARNINGS
+}
+
+UChaosClothAssetBase* UChaosClothComponent::GetAsset() const
+{
+	return Cast<UChaosClothAssetBase>(GetSkinnedAsset());
+}
+
+void UChaosClothComponent::SetClothAsset(UChaosClothAsset* InClothAsset)
+{
+	SetAsset(InClothAsset);
 }
 
 UChaosClothAsset* UChaosClothComponent::GetClothAsset() const
 {
-	return Cast<UChaosClothAsset>(GetSkinnedAsset());
+	return Cast<UChaosClothAsset>(GetAsset());
 }
 
 bool UChaosClothComponent::IsSimulationSuspended() const
@@ -74,9 +84,9 @@ void UChaosClothComponent::ResetConfigProperties()
 {
 	if (IsRegistered())
 	{
-		if (GetClothAsset())
+		if (UChaosClothAsset* const ClothAsset = Cast<UChaosClothAsset>(GetAsset()))
 		{
-			const TArray<TSharedRef<FManagedArrayCollection>>& ClothCollections = GetClothAsset()->GetClothCollections();
+			const TArray<TSharedRef<FManagedArrayCollection>>& ClothCollections = ClothAsset->GetClothCollections();
 			PropertyCollections.Reset();
 			PropertyCollections.Reserve(ClothCollections.Num());
 			CollectionPropertyFacades.Reset();
@@ -112,9 +122,9 @@ void UChaosClothComponent::UpdateConfigProperties()
 {
 	if (IsRegistered())
 	{
-		if (GetClothAsset())
+		if (UChaosClothAsset* const ClothAsset = Cast<UChaosClothAsset>(GetAsset()))
 		{
-			const TArray<TSharedRef<FManagedArrayCollection>>& ClothCollections = GetClothAsset()->GetClothCollections();
+			const TArray<TSharedRef<FManagedArrayCollection>>& ClothCollections = ClothAsset->GetClothCollections();
 			if (ClothCollections.Num() == PropertyCollections.Num())
 			{
 				check(CollectionPropertyFacades.Num() == ClothCollections.Num());
@@ -141,9 +151,9 @@ void UChaosClothComponent::RecreateClothSimulationProxy()
 	{
 		ClothSimulationProxy.Reset();
 
-		if (GetClothAsset())
+		if (UChaosClothAsset* const ClothAsset = Cast<UChaosClothAsset>(GetAsset()))
 		{
-			const TSharedPtr<const FChaosClothSimulationModel> ClothSimulationModel = GetClothAsset()->GetClothSimulationModel();
+			const TSharedPtr<const FChaosClothSimulationModel> ClothSimulationModel = ClothAsset->GetClothSimulationModel();
 			if (ClothSimulationModel && ClothSimulationModel->GetNumLods())
 			{
 				// Create the simulation proxy (note CreateClothSimulationProxy() can be overloaded)
@@ -162,9 +172,9 @@ void UChaosClothComponent::PostLoad()
 	Super::PostLoad();
 
 #if WITH_EDITORONLY_DATA
-	PRAGMA_DISABLE_DEPRECATION_WARNINGS
-		ClothAsset = GetClothAsset();
-	PRAGMA_ENABLE_DEPRECATION_WARNINGS
+PRAGMA_DISABLE_DEPRECATION_WARNINGS
+	Asset = GetAsset();
+PRAGMA_ENABLE_DEPRECATION_WARNINGS
 #endif
 }
 
@@ -175,9 +185,9 @@ void UChaosClothComponent::PostEditChangeProperty(FPropertyChangedEvent& Propert
 	if (const FProperty* const Property = PropertyChangedEvent.Property)
 	{
 PRAGMA_DISABLE_DEPRECATION_WARNINGS
-		if (Property->GetFName() == GET_MEMBER_NAME_CHECKED(UChaosClothComponent, ClothAsset))
+		if (Property->GetFName() == GET_MEMBER_NAME_CHECKED(UChaosClothComponent, Asset))
 		{
-			SetClothAsset(ClothAsset);
+			SetAsset(Asset);
 		}
 PRAGMA_ENABLE_DEPRECATION_WARNINGS
 		if (Property->GetFName() == GET_MEMBER_NAME_CHECKED(UChaosClothComponent, bSimulateInEditor))
@@ -509,7 +519,7 @@ void UChaosClothComponent::GetAdditionalRequiredBonesForLeader(int32 LeaderLODIn
 	}
 
 	// Merge the physics asset bones (the leader's physics asset can be different to this component's cloth asset)
-	if (const UPhysicsAsset* const PhysicsAsset = GetClothAsset() ? GetClothAsset()->GetPhysicsAsset() : nullptr)
+	if (const UPhysicsAsset* const PhysicsAsset = GetAsset() ? GetAsset()->GetPhysicsAsset() : nullptr)
 	{
 		if (const USkinnedAsset* const LeaderSkinnedAsset = ensure(LeaderPoseComponent.IsValid()) ? LeaderPoseComponent->GetSkinnedAsset() : nullptr)  // Needs the leader SkinnedAsset for the correct RefSkeleton
 		{
@@ -629,10 +639,10 @@ void UChaosClothComponent::UpdateComponentSpaceTransforms()
 {
 	check(IsRegistered());
 
-	if (!LeaderPoseComponent.IsValid() && GetClothAsset() && GetClothAsset()->GetResourceForRendering())
+	if (!LeaderPoseComponent.IsValid() && GetAsset() && GetAsset()->GetResourceForRendering())
 	{
-		FSkeletalMeshLODRenderData& LODData = GetClothAsset()->GetResourceForRendering()->LODRenderData[GetPredictedLODLevel()];
-		GetClothAsset()->FillComponentSpaceTransforms(GetClothAsset()->GetRefSkeleton().GetRefBonePose(), LODData.RequiredBones, GetEditableComponentSpaceTransforms());
+		FSkeletalMeshLODRenderData& LODData = GetAsset()->GetResourceForRendering()->LODRenderData[GetPredictedLODLevel()];
+		GetAsset()->FillComponentSpaceTransforms(GetAsset()->GetRefSkeleton().GetRefBonePose(), LODData.RequiredBones, GetEditableComponentSpaceTransforms());
 
 		bNeedToFlipSpaceBaseBuffers = true; // Have updated space bases so need to flip
 		FlipEditableSpaceBases();
@@ -642,9 +652,9 @@ void UChaosClothComponent::UpdateComponentSpaceTransforms()
 
 void UChaosClothComponent::UpdateVisibility()
 {
-	if (GetClothAsset() && GetClothAsset()->GetResourceForRendering())
+	if (GetAsset() && GetAsset()->GetResourceForRendering())
 	{
-		const FSkeletalMeshRenderData* const SkeletalMeshRenderData = GetClothAsset()->GetResourceForRendering();
+		const FSkeletalMeshRenderData* const SkeletalMeshRenderData = GetAsset()->GetResourceForRendering();
 		const int32 FirstValidLODIdx = SkeletalMeshRenderData ? SkeletalMeshRenderData->GetFirstValidLODIdx(0) : INDEX_NONE;
 		SetVisibility(FirstValidLODIdx != INDEX_NONE);
 	}
diff --git a/Engine/Plugins/ChaosClothAsset/Source/ChaosClothAssetEngine/Private/ChaosClothAsset/ClothSimulationContext.cpp b/Engine/Plugins/ChaosClothAsset/Source/ChaosClothAssetEngine/Private/ChaosClothAsset/ClothSimulationContext.cpp
--- a/Engine/Plugins/ChaosClothAsset/Source/ChaosClothAssetEngine/Private/ChaosClothAsset/ClothSimulationContext.cpp
+++ b/Engine/Plugins/ChaosClothAsset/Source/ChaosClothAssetEngine/Private/ChaosClothAsset/ClothSimulationContext.cpp
@@ -38,7 +38,7 @@ namespace UE::Chaos::ClothAsset
 		SolverGeometryScale = ClothComponent.GetClothGeometryScale();;
 
 		// Update bone transforms
-		const UChaosClothAsset* const ClothAsset = ClothComponent.GetClothAsset();
+		const UChaosClothAsset* const ClothAsset = Cast<const UChaosClothAsset>(ClothComponent.GetAsset());
 		const FReferenceSkeleton* const ReferenceSkeleton = ClothAsset ? &ClothAsset->GetRefSkeleton() : nullptr;
 		int32 NumBones = ReferenceSkeleton ? ReferenceSkeleton->GetNum() : 0;
 
diff --git a/Engine/Plugins/ChaosClothAsset/Source/ChaosClothAssetEngine/Private/ChaosClothAsset/ClothSimulationMesh.cpp b/Engine/Plugins/ChaosClothAsset/Source/ChaosClothAssetEngine/Private/ChaosClothAsset/ClothSimulationMesh.cpp
--- a/Engine/Plugins/ChaosClothAsset/Source/ChaosClothAssetEngine/Private/ChaosClothAsset/ClothSimulationMesh.cpp
+++ b/Engine/Plugins/ChaosClothAsset/Source/ChaosClothAssetEngine/Private/ChaosClothAsset/ClothSimulationMesh.cpp
@@ -9,8 +9,11 @@
 
 namespace UE::Chaos::ClothAsset
 {
-	FClothSimulationMesh::FClothSimulationMesh(const FChaosClothSimulationModel& InClothSimulationModel, const FClothSimulationContext& InClothSimulationContext,
-		const TArray<TSharedRef<const FManagedArrayCollection>>& InManagedArrayCollections, const FString& DebugName)
+	FClothSimulationMesh::FClothSimulationMesh(
+		const FChaosClothSimulationModel& InClothSimulationModel,
+		const FClothSimulationContext& InClothSimulationContext,
+		const TArray<TSharedRef<FManagedArrayCollection>>& InManagedArrayCollections,
+		const FString& DebugName)
 		: ::Chaos::FClothingSimulationMesh(DebugName)
 		, ClothSimulationModel(InClothSimulationModel)
 		, ClothSimulationContext(InClothSimulationContext)
diff --git a/Engine/Plugins/ChaosClothAsset/Source/ChaosClothAssetEngine/Private/ChaosClothAsset/ClothSimulationMesh.h b/Engine/Plugins/ChaosClothAsset/Source/ChaosClothAssetEngine/Private/ChaosClothAsset/ClothSimulationMesh.h
--- a/Engine/Plugins/ChaosClothAsset/Source/ChaosClothAssetEngine/Private/ChaosClothAsset/ClothSimulationMesh.h
+++ b/Engine/Plugins/ChaosClothAsset/Source/ChaosClothAssetEngine/Private/ChaosClothAsset/ClothSimulationMesh.h
@@ -13,7 +13,11 @@ namespace UE::Chaos::ClothAsset
 	class FClothSimulationMesh : public ::Chaos::FClothingSimulationMesh
 	{
 	public:
-		FClothSimulationMesh(const FChaosClothSimulationModel& InClothSimulationModel, const FClothSimulationContext& InClothSimulationContext, const TArray<TSharedRef<const FManagedArrayCollection>>& InManagedArrayCollections, const FString& DebugName);
+		FClothSimulationMesh(
+			const FChaosClothSimulationModel& InClothSimulationModel,
+			const FClothSimulationContext& InClothSimulationContext,
+			const TArray<TSharedRef<FManagedArrayCollection>>& InManagedArrayCollections,
+			const FString& DebugName);
 
 		virtual ~FClothSimulationMesh() override = default;
 
@@ -58,6 +62,6 @@ namespace UE::Chaos::ClothAsset
 	private:
 		const FChaosClothSimulationModel& ClothSimulationModel;
 		const FClothSimulationContext& ClothSimulationContext;
-		const TArray<TSharedRef<const FManagedArrayCollection>>& ManagedArrayCollections;
+		const TArray<TSharedRef<FManagedArrayCollection>>& ManagedArrayCollections;
 	};
 }
diff --git a/Engine/Plugins/ChaosClothAsset/Source/ChaosClothAssetEngine/Private/ChaosClothAsset/ClothSimulationProxy.cpp b/Engine/Plugins/ChaosClothAsset/Source/ChaosClothAssetEngine/Private/ChaosClothAsset/ClothSimulationProxy.cpp
--- a/Engine/Plugins/ChaosClothAsset/Source/ChaosClothAssetEngine/Private/ChaosClothAsset/ClothSimulationProxy.cpp
+++ b/Engine/Plugins/ChaosClothAsset/Source/ChaosClothAssetEngine/Private/ChaosClothAsset/ClothSimulationProxy.cpp
@@ -157,7 +157,7 @@ namespace UE::Chaos::ClothAsset
 		check(IsInGameThread());
 
 		// Create solver config simulation thread object first. Need to know which solver type we're creating.
-		const int32 SolverConfigIndex = Configs.Emplace(MakeUnique<FClothingSimulationConfig>(ClothComponent.GetPropertyCollections()));  // TODO: Use a separate solver config for outfits
+		const int32 SolverConfigIndex = Configs.Emplace(MakeUnique<FClothingSimulationConfig>(ClothComponent.GetPropertyCollections(0)));  // TODO: Use a separate solver config for outfits
 
 		// Use new SoftsEvolution, not PBDEvolution.
 		constexpr bool bUseLegacySolver = false;
@@ -179,7 +179,7 @@ PRAGMA_ENABLE_DEPRECATION_WARNINGS
 		Solver->SetLocalSpaceRotation((FQuat)ClothSimulationContext->ComponentTransform.GetRotation());
 
 		// Create mesh simulation thread object
-		const UChaosClothAsset* const ClothAsset = ClothComponent.GetClothAsset();
+		UChaosClothAsset* const ClothAsset = Cast<UChaosClothAsset>(ClothComponent.GetAsset());
 		if (!ClothAsset)
 		{
 			return;
@@ -201,7 +201,7 @@ PRAGMA_ENABLE_DEPRECATION_WARNINGS
 		Colliders[ColliderIndex]->SetCollisionData(&CollisionSourcesProxy->GetCollisionData());
 
 		// Create cloth config simulation thread object
-		const int32 ClothConfigIndex = Configs.Emplace(MakeUnique<FClothingSimulationConfig>(ClothComponent.GetPropertyCollections()));
+		const int32 ClothConfigIndex = Configs.Emplace(MakeUnique<FClothingSimulationConfig>(ClothComponent.GetPropertyCollections(0)));
 
 		// Create cloth simulation thread object
 		constexpr uint32 GroupId = 0;
@@ -332,7 +332,7 @@ PRAGMA_ENABLE_DEPRECATION_WARNINGS
 
 		// Check that the render mesh current LOD isn't just fully skinned
 		const int32 LodIndex = ClothSimulationContext->LodIndex;
-		const UChaosClothAsset* const ClothAsset = ClothComponent.GetClothAsset();
+		const UChaosClothAsset* const ClothAsset = Cast<const UChaosClothAsset>(ClothComponent.GetAsset());
 		if (!ClothAsset ||
 			!ClothAsset->GetResourceForRendering() ||
 			!ClothAsset->GetResourceForRendering()->LODRenderData.IsValidIndex(LodIndex) ||
@@ -801,7 +801,7 @@ PRAGMA_ENABLE_DEPRECATION_WARNINGS
 		// Replace physics thread's configs with the game thread's configs
 		for (const TUniquePtr<::Chaos::FClothingSimulationConfig>& Config : Configs)
 		{
-			Config->Initialize(ClothComponent.GetPropertyCollections());  // TODO: Outfit and multi-cloths/solver config update
+			Config->Initialize(ClothComponent.GetPropertyCollections(0));  // TODO: Outfit and multi-cloths/solver config update
 		}
 		// TODO: separate solver config lod from cloth component
 		Solver->SetSolverLOD(ClothSimulationContext->LodIndex);
diff --git a/Engine/Plugins/ChaosClothAsset/Source/ChaosClothAssetEngine/Public/ChaosClothAsset/ClothAsset.h b/Engine/Plugins/ChaosClothAsset/Source/ChaosClothAssetEngine/Public/ChaosClothAsset/ClothAsset.h
--- a/Engine/Plugins/ChaosClothAsset/Source/ChaosClothAssetEngine/Public/ChaosClothAsset/ClothAsset.h
+++ b/Engine/Plugins/ChaosClothAsset/Source/ChaosClothAssetEngine/Public/ChaosClothAsset/ClothAsset.h
@@ -2,12 +2,10 @@
 
 #pragma once
 
-#include "Engine/SkinnedAsset.h"
+#include "ChaosClothAsset/ClothAssetBase.h"
 #include "ReferenceSkeleton.h"
 #include "RenderCommandFence.h"
 #include "PerQualityLevelProperties.h"
-#include "Dataflow/DataflowContent.h"
-#include "Dataflow/DataflowInstance.h"
 #include "ClothAsset.generated.h"
 
 #if UE_ENABLE_INCLUDE_ORDER_DEPRECATED_IN_5_5
@@ -16,14 +14,8 @@ namespace Dataflow = UE::Dataflow;
 namespace UE_DEPRECATED(5.5, "Use UE::Dataflow instead.") Dataflow {}
 #endif
 
-class FSkeletalMeshRenderData;
 class FSkeletalMeshModel;
-class FSkinnedAssetCompilationContext;
-class USkeletalMesh;
-class UDataflow;
-struct FChaosClothSimulationModel;
 struct FSkeletalMeshLODInfo;
-struct FManagedArrayCollection;
 struct FChaosClothAssetLodTransitionDataCache;
 class UAnimationAsset;
 
@@ -31,15 +23,12 @@ UENUM()
 enum class EClothAssetAsyncProperties : uint64
 {
 	None = 0,
-	RenderData = 1 << 0,
+	RenderData UE_DEPRECATED(5.6, "Use EChaosClothAssetBaseAsyncProperties::RenderData instead.") = 1 << 0,
 	ThumbnailInfo = 1 << 1,
 	ImportedModel = 1 << 2,
 	ClothCollection = 1 << 3,
-	RefSkeleton = 1 << 4,
-	HasVertexColors = 1 << 5,
-	OverlayMaterial = 1 << 6,
-	OverlayMaterialMaxDrawDistance = 1 << 7,
-	All = MAX_uint64
+	RefSkeleton UE_DEPRECATED(5.6, "Use EChaosClothAssetBaseAsyncProperties::RefSkeleton instead.") = 1 << 4,
+	All = MAX_uint32  // Max is uint32 as we need some space for the UChaosClothAssetBase ones
 };
 ENUM_CLASS_FLAGS(EClothAssetAsyncProperties);
 
@@ -47,10 +36,7 @@ ENUM_CLASS_FLAGS(EClothAssetAsyncProperties);
  * Cloth asset for pattern based simulation.
  */
 UCLASS(hidecategories = Object, BlueprintType)
-class CHAOSCLOTHASSETENGINE_API UChaosClothAsset :
-	public USkinnedAsset, 
-	public IDataflowContentOwner,
-	public IDataflowInstanceInterface
+class CHAOSCLOTHASSETENGINE_API UChaosClothAsset : public UChaosClothAssetBase
 {
 	GENERATED_BODY()
 public:
@@ -60,8 +46,6 @@ public:
 
 	//~ Begin UObject interface
 	virtual void Serialize(FArchive& Ar) override;
-	virtual void BeginDestroy() override;
-	virtual bool IsReadyForFinishDestroy() override;
 	virtual void PostLoad() override;
 #if WITH_EDITOR
 	virtual void PostEditChangeProperty(FPropertyChangedEvent& PropertyChangedEvent) override;
@@ -69,29 +53,7 @@ public:
 	virtual void GetResourceSizeEx(FResourceSizeEx& CumulativeResourceSize) override;
 	//~ End UObject interface
 
-	//~ Begin IDataflowContentOwner interface 
-	virtual TObjectPtr<UDataflowBaseContent> CreateDataflowContent() override;
-	virtual void WriteDataflowContent(const TObjectPtr<UDataflowBaseContent>& DataflowContent) const override;
-	virtual void ReadDataflowContent(const TObjectPtr<UDataflowBaseContent>& DataflowContent) override;
-	//~ End IDataflowContentOwner interface 
-
 	//~ Begin USkinnedAsset interface
-	virtual FReferenceSkeleton& GetRefSkeleton()								
-	{
-		WaitUntilAsyncPropertyReleased(EClothAssetAsyncProperties::RefSkeleton);
-		PRAGMA_DISABLE_DEPRECATION_WARNINGS
-		return RefSkeleton;
-		PRAGMA_ENABLE_DEPRECATION_WARNINGS
-	}
-	virtual const FReferenceSkeleton& GetRefSkeleton() const					
-	{
-		WaitUntilAsyncPropertyReleased(EClothAssetAsyncProperties::RefSkeleton, ESkinnedAssetAsyncPropertyLockType::ReadOnly);
-		PRAGMA_DISABLE_DEPRECATION_WARNINGS
-		return RefSkeleton;
-		PRAGMA_ENABLE_DEPRECATION_WARNINGS
-	}
-	virtual FSkeletalMeshLODInfo* GetLODInfo(int32 Index);
-	virtual const FSkeletalMeshLODInfo* GetLODInfo(int32 Index) const;
 	UFUNCTION(BlueprintGetter)
 	virtual class UPhysicsAsset* GetShadowPhysicsAsset() const override			{ return ShadowPhysicsAsset; }
 	virtual FMatrix GetComposedRefPoseMatrix(FName InBoneName) const override;
@@ -101,58 +63,14 @@ public:
 	virtual int32 GetRayTracingMinLOD() const override							{ return RayTracingMinLOD; }
 	virtual TArray<FMatrix44f>& GetRefBasesInvMatrix() override					{ return RefBasesInvMatrix; }
 	virtual const TArray<FMatrix44f>& GetRefBasesInvMatrix() const override		{ return RefBasesInvMatrix; }
-	virtual TArray<FSkeletalMeshLODInfo>& GetLODInfoArray() override			{ return LODInfo; }
-	virtual const TArray<FSkeletalMeshLODInfo>& GetLODInfoArray() const override { return LODInfo; }
-	virtual FSkeletalMeshRenderData* GetResourceForRendering() const override;
-	virtual int32 GetDefaultMinLod() const										{ return 0; }
 	virtual UPhysicsAsset* GetPhysicsAsset() const								{ return PhysicsAsset; }
 	virtual TArray<FSkeletalMaterial>& GetMaterials() override					{ return Materials; }
 	virtual const TArray<FSkeletalMaterial>& GetMaterials() const override		{ return Materials; }
-	virtual int32 GetLODNum() const override									{ return LODInfo.Num(); }
 	virtual bool IsMaterialUsed(int32 MaterialIndex) const override				{ return true; }
-	virtual FBoxSphereBounds GetBounds() const override							{ return Bounds; }
-	virtual TArray<class USkeletalMeshSocket*> GetActiveSocketList() const override { static TArray<class USkeletalMeshSocket*> Dummy; return Dummy; }
-	virtual USkeletalMeshSocket* FindSocket(FName InSocketName) const override	{ return nullptr; }
-	virtual USkeletalMeshSocket* FindSocketInfo(FName InSocketName, FTransform& OutTransform, int32& OutBoneIndex, int32& OutIndex) const override { return nullptr; }
 	virtual USkeleton* GetSkeleton() override									{ return Skeleton; }  // Note: The USkeleton isn't a reliable source of reference skeleton
 	virtual const USkeleton* GetSkeleton() const override						{ return Skeleton; }
 	virtual void SetSkeleton(USkeleton* InSkeleton) override					{ Skeleton = InSkeleton; }
-	virtual UMeshDeformer* GetDefaultMeshDeformer() const override				{ return nullptr; }
 
-	/** Get the default overlay material used by this mesh */
-	UFUNCTION(BlueprintCallable, Category = "Rendering|Material")
-	virtual class UMaterialInterface* GetOverlayMaterial() const override		
-	{
-		WaitUntilAsyncPropertyReleased(EClothAssetAsyncProperties::OverlayMaterial, ESkinnedAssetAsyncPropertyLockType::ReadOnly);
-		PRAGMA_DISABLE_DEPRECATION_WARNINGS
-		return OverlayMaterial;
-		PRAGMA_ENABLE_DEPRECATION_WARNINGS
-	}
-
-	/** Get the default overlay material max draw distance used by this mesh */
-	UFUNCTION(BlueprintCallable, Category = "Rendering|Material")
-	virtual float GetOverlayMaterialMaxDrawDistance() const override
-	{
-		WaitUntilAsyncPropertyReleased(EClothAssetAsyncProperties::OverlayMaterialMaxDrawDistance, ESkinnedAssetAsyncPropertyLockType::ReadOnly);
-		PRAGMA_DISABLE_DEPRECATION_WARNINGS
-		return OverlayMaterialMaxDrawDistance;
-		PRAGMA_ENABLE_DEPRECATION_WARNINGS
-	}
-
-	virtual bool IsValidLODIndex(int32 Index) const override					{ return LODInfo.IsValidIndex(Index); }
-	virtual int32 GetMinLodIdx(bool bForceLowestLODIdx = false) const override;
-	virtual bool NeedCPUData(int32 LODIndex) const override						{ return false; }
-	virtual bool GetHasVertexColors() const override
-	{
-		WaitUntilAsyncPropertyReleased(EClothAssetAsyncProperties::HasVertexColors, ESkinnedAssetAsyncPropertyLockType::ReadOnly);
-		PRAGMA_DISABLE_DEPRECATION_WARNINGS
-		return bHasVertexColors != 0;
-		PRAGMA_ENABLE_DEPRECATION_WARNINGS
-	}
-	virtual int32 GetPlatformMinLODIdx(const ITargetPlatform* TargetPlatform) const override;
-	virtual const FPerPlatformBool& GetDisableBelowMinLodStripping() const override { return DisableBelowMinLodStripping; }
-	virtual const FPerPlatformInt& GetMinLod() const override;
-	virtual bool IsMinLodQualityLevelEnable() const override;
 #if WITH_EDITOR
 	/* Build a LOD model for the targeted platform. */
 	virtual void BuildLODModel(const ITargetPlatform* TargetPlatform, int32 LODIndex) override;
@@ -173,26 +91,6 @@ public:
 #endif
 	//~ End USkinnedAsset interface
 
-	/** Change the default overlay material used by this mesh */
-	UFUNCTION(BlueprintCallable, Category = "Rendering|Material")
-	void SetOverlayMaterial(class UMaterialInterface* NewOverlayMaterial)
-	{
-		WaitUntilAsyncPropertyReleased(EClothAssetAsyncProperties::OverlayMaterial);
-		PRAGMA_DISABLE_DEPRECATION_WARNINGS
-		OverlayMaterial = NewOverlayMaterial;
-		PRAGMA_ENABLE_DEPRECATION_WARNINGS
-	}
-
-	/** Change the default overlay material max draw distance used by this mesh */
-	UFUNCTION(BlueprintCallable, Category = "Rendering|Material")
-	void SetOverlayMaterialMaxDrawDistance(float InMaxDrawDistance)
-	{
-		WaitUntilAsyncPropertyReleased(EClothAssetAsyncProperties::OverlayMaterialMaxDrawDistance);
-		PRAGMA_DISABLE_DEPRECATION_WARNINGS
-		OverlayMaterialMaxDrawDistance = InMaxDrawDistance;
-		PRAGMA_ENABLE_DEPRECATION_WARNINGS
-	}
-
 	/** Return the enclosed Cloth Collection object. */
 	TArray<TSharedRef<FManagedArrayCollection>>& GetClothCollections() 
 	{
@@ -211,9 +109,6 @@ public:
 		PRAGMA_ENABLE_DEPRECATION_WARNINGS
 	}
 
-	/** Return the cloth simulation ready LOD model data. */
-	TSharedPtr<const FChaosClothSimulationModel> GetClothSimulationModel() const { return ClothSimulationModel; }
-
 	/** Build this asset static render and simulation data. This needs to be done every time the asset has changed. */
 	void Build(TArray<FChaosClothAssetLodTransitionDataCache>* InOutTransitionCache = nullptr);
 
@@ -239,61 +134,60 @@ public:
 	 * @param bRebuildModels Whether to rebuild the mesh static models. This is always required, but could be skipped if done soon after.
 	 * @param bRebindMeshes Bindings could be invalid after a change of reference skeleton, use this option to clear the existing binding and bind the sim and render meshes to the root bone.
 	 */
-	void SetReferenceSkeleton(const FReferenceSkeleton* ReferenceSkeleton, bool bRebuildModels = true, bool bRebindMeshes = true);
-
-	const FPerQualityLevelInt& GetQualityLevelMinLod() const { return MinQualityLevelLOD; }
-	void SetQualityLevelMinLod(FPerQualityLevelInt InMinLod) { MinQualityLevelLOD = MoveTemp(InMinLod); }
-	static void OnLodStrippingQualityLevelChanged(IConsoleVariable* Variable);
-
-	/**
-	 * Set the Dataflow graph asset for this cloth asset.
-	 * @param InDataflow The new dataflow asset.
-	 */
-	void SetDataflow(UDataflow* InDataflow) { DataflowInstance.SetDataflowAsset(InDataflow); }
-
-	/** Return the Dataflow graph asset associated to this cloth asset if any. */
-	UDataflow* GetDataflow() { return DataflowInstance.GetDataflowAsset(); }
+	void SetReferenceSkeleton(const FReferenceSkeleton* ReferenceSkeleton, bool bRebuildModels, bool bRebindMeshes = true);
 
-	/** Return the Dataflow graph asset associated to this cloth asset if any, const version. */
-	const UDataflow* GetDataflow() const { return DataflowInstance.GetDataflowAsset(); }
+	UE_DEPRECATED(5.6, "No longer used.")
+	static void OnLodStrippingQualityLevelChanged(IConsoleVariable* /*Variable*/) {}
 
 	/**
 	 * Set the name of the Dataflow terminal node for this cloth asset.
 	 * @param InDataflowTerminal The new name of the Dataflow terminal node.
 	 */
+	UE_DEPRECATED(5.6, "Function will be removed as the DataflowTerminal shouldn't be set by name/string.")
 	void SetDataflowTerminal(const FString& InDataflowTerminal) 
 	{ 
 		DataflowInstance.SetDataflowTerminal(FName(*InDataflowTerminal));
 	}
 
 	/** Return the Dataflow graph asset associated to this cloth asset if any. */
+	UE_DEPRECATED(5.6, "Use GetDataflowInstance() instead")
 	FString GetDataflowTerminal() const 
 	{
 		return DataflowInstance.GetDataflowTerminal().ToString();
 	}
 
-	//~ Begin IDataflowGeneratedAssetInterface
-	virtual const FDataflowInstance& GetDataflowInstance() const { return DataflowInstance; };
-	//~ End IDataflowGeneratedAssetInterface
-
-#if WITH_EDITORONLY_DATA
-
-	void SetPreviewSceneSkeletalMesh(USkeletalMesh* Mesh);
-	USkeletalMesh* GetPreviewSceneSkeletalMesh() const;
-
-	void SetPreviewSceneAnimation(UAnimationAsset* Animation);
-	UAnimationAsset* GetPreviewSceneAnimation() const;
-
+	//~ Begin UChaosClothAssetBase interface
+	/** Return the cloth simulation ready LOD model data. */
+	virtual TSharedPtr<const FChaosClothSimulationModel> GetClothSimulationModel(int32 /*ModelIndex*/ = 0) const override
+	{
+		return ClothSimulationModel;
+	}
+	virtual bool HasValidClothSimulationModels() const override;
+	virtual int32 GetNumClothSimulationModels() const override
+	{
+		return 1;
+	}
+	virtual FName GetClothSimulationModelName(int32 /*ModelIndex*/) const override
+	{
+		return GetFName();
+	}
+	virtual const TArray<TSharedRef<FManagedArrayCollection>>& GetCollections(int32 /*ModelIndex*/) const override;
+	virtual const UPhysicsAsset* GetPhysicsAssetForModel(int32 /*ModelIndex*/) const override
+	{
+		return PhysicsAsset;
+	}
+#if WITH_EDITOR
+	/** Load render data from DDC if the data is cached, otherwise generate render data and save into DDC */
+	virtual void CacheDerivedData(FSkinnedAssetCompilationContext* Context) override;
 #endif
 
-	void SetHasVertexColors(bool InbHasVertexColors)
+	virtual void SetReferenceSkeleton(const FReferenceSkeleton* ReferenceSkeleton) override
 	{
-		WaitUntilAsyncPropertyReleased(EClothAssetAsyncProperties::HasVertexColors);
-		PRAGMA_DISABLE_DEPRECATION_WARNINGS
-		bHasVertexColors = InbHasVertexColors;
-		PRAGMA_ENABLE_DEPRECATION_WARNINGS
+		SetReferenceSkeleton(ReferenceSkeleton, true, true);
 	}
 
+	//~ End UChaosClothAssetBase interface
+
 private:
 #if WITH_EDITORONLY_DATA
 	/** Dataflow asset. */
@@ -305,9 +199,6 @@ private:
 	FString DataflowTerminal_DEPRECATED = "ClothAssetTerminal";
 #endif  // WITH_EDITORONLY_DATA
 
-	UPROPERTY(EditAnywhere, Category = "Dataflow")
-	FDataflowInstance DataflowInstance;
-
 	//~ Begin USkinnedAsset interface
 	/** Initial step for the Post Load process - Can't be done in parallel. */
 	virtual void BeginPostLoadInternal(FSkinnedAssetPostLoadContext& Context) override;
@@ -319,12 +210,6 @@ private:
 	virtual FString GetAsyncPropertyName(uint64 Property) const override;
 	//~ End USkinnedAsset interface
 
-	/**
-	 * Wait for the asset to finish compilation to protect internal skinned asset data from race conditions during async build.
-	 * This should be called before accessing all async accessible properties.
-	 */
-	void WaitUntilAsyncPropertyReleased(EClothAssetAsyncProperties AsyncProperties, ESkinnedAssetAsyncPropertyLockType LockType = ESkinnedAssetAsyncPropertyLockType::ReadWrite) const;
-
 	/** Pre-calculate refpose-to-local transforms. */
 	void CalculateInvRefMatrices();
 
@@ -339,32 +224,6 @@ private:
 	/** Build the clothing simulation meshes from the Cloth Collection. */
 	void BuildClothSimulationModel(TArray<FChaosClothAssetLodTransitionDataCache>* InOutTransitionCache = nullptr);
 
-	/** Initialize all render resources. */
-	void InitResources();
-
-	/** Safely release the render data. */
-	void ReleaseResources();
-
-	/** Set render data. */
-	void SetResourceForRendering(TUniquePtr<FSkeletalMeshRenderData>&& InSkeletalMeshRenderData);
-
-#if WITH_EDITOR
-	/** Load render data from DDC if the data is cached, otherwise generate render data and save into DDC */
-	void CacheDerivedData(FSkinnedAssetCompilationContext* Context);
-
-	/** Initial step for the building process - Can't be done in parallel. USkinnedAsset Interface. */
-	virtual void BeginBuildInternal(FSkinnedAssetBuildContext& Context) override;
-
-	/** Thread-safe part. USkinnedAsset Interface. */
-	virtual void ExecuteBuildInternal(FSkinnedAssetBuildContext& Context) override;
-
-	/** Complete the building process - Can't be done in parallel. USkinnedAsset Interface. */
-	virtual void FinishBuildInternal(FSkinnedAssetBuildContext& Context) override;
-#endif
-
-	/** Reregister all components using this asset to reset the simulation in case anything has changed. */
-	void ReregisterComponents();
-
 	/** List of materials for this cloth asset. Set by the Dataflow evaluation. */
 	UPROPERTY(EditAnywhere, Category = Materials, Meta = (EditCondition = "DataflowAsset == nullptr"))
 	TArray<FSkeletalMaterial> Materials;
@@ -381,33 +240,6 @@ private:
 	UPROPERTY(EditAnywhere, Category = Collision, Meta = (EditCondition = "DataflowAsset == nullptr"))
 	TObjectPtr<UPhysicsAsset> PhysicsAsset;
 
-	/** Struct containing information for each LOD level, such as materials to use, and when use the LOD. Not currently editable or customizable through the Dataflow. */
-	UPROPERTY(VisibleAnywhere, EditFixedSize, Category = LevelOfDetails)
-	TArray<FSkeletalMeshLODInfo> LODInfo;
-	
-	/** Set the Minimum LOD by Quality Level. This property is used when "Use Cloth Asset Min LOD Per Quality Levels" is set at the Project level. Otherwise, the (per platform) Minimum LOD value is used.*/
-	UPROPERTY(EditAnywhere, Category = LODSettings, meta = (DisplayName = "Quality Level Minimum LOD"))
-	FPerQualityLevelInt MinQualityLevelLOD;
-
-	UPROPERTY(EditAnywhere, Category = LODSettings)
-	FPerPlatformBool DisableBelowMinLodStripping;
-
-	/** Set the Minimum LOD by platform. This property is overriden by "Quality Level Minimum LOD" when "Use Cloth Asset Min LOD Per Quality Levels" is set at the Project level.*/
-	UPROPERTY(EditAnywhere, Category = LODSettings, Meta = (DisplayName = "Minimum LOD"))
-	FPerPlatformInt MinLod;
-
-	/** Enable raytracing for this asset. */
-	UPROPERTY(EditAnywhere, Category = RayTracing)
-	uint8 bSupportRayTracing : 1;
-
-	UE_DEPRECATED(5.6, "This must be protected for async build, always use the accessors even internally.")
-	UPROPERTY()
-	uint8 bHasVertexColors : 1 = false;
-
-	/** Minimum raytracing LOD for this asset. */
-	UPROPERTY(EditAnywhere, Category = RayTracing)
-	int32 RayTracingMinLOD;
-
 	/** Whether to blend positions between the skinned/simulated transitions of the cloth render mesh. Only used when not overriden in Dataflow by the ProxyDefomer node. */
 	UE_DEPRECATED(5.4, "Superseded by the ProxyDefomer node.")
 	UPROPERTY()
@@ -434,19 +266,6 @@ private:
 	UPROPERTY()
 	FGuid AssetGuid;
 
-	/** Default translucent material to blend on top of this mesh. Mesh will be rendered twice - once with a base material and once with overlay material */
-	UE_DEPRECATED(5.6, "This must be protected for async build, always use the accessors even internally.")
-	UPROPERTY(EditAnywhere, Category = Rendering)
-	TObjectPtr<class UMaterialInterface> OverlayMaterial;
-
-	/** Default max draw distance for overlay material. A distance of 0 indicates that overlay will be culled using primitive max distance. */
-	UE_DEPRECATED(5.6, "This must be protected for async build, always use the accessors even internally.")
-	UPROPERTY(EditAnywhere, Category = Rendering)
-	float OverlayMaterialMaxDrawDistance;
-
-	/** Bounds for this asset. */
-	FBoxSphereBounds Bounds;
-
 	/** Mesh-space ref pose, where parent matrices are applied to ref pose matrices. */
 	TArray<FMatrix> CachedComposedRefPoseMatrices;
 
@@ -454,19 +273,9 @@ private:
 	UE_DEPRECATED(5.4, "This must be protected for async build, always use the accessors even internally.")
 	TArray<TSharedRef<FManagedArrayCollection>> ClothCollections;
 
-	/** Reference skeleton created from the provided skeleton asset. */
-	UE_DEPRECATED(5.4, "This must be protected for async build, always use the accessors even internally.")
-	FReferenceSkeleton RefSkeleton;
-
 	/** Reference skeleton precomputed bases. */
 	TArray<FMatrix44f> RefBasesInvMatrix;
 
-	/** Rendering data. */
-	TUniquePtr<FSkeletalMeshRenderData> SkeletalMeshRenderData;
-
-	/** A fence which is used to keep track of the rendering thread releasing the static mesh resources. */
-	FRenderCommandFence ReleaseResourcesFence;
-
 #if WITH_EDITORONLY_DATA
 	/** Source mesh geometry information (not used at runtime). */
 	UE_DEPRECATED(5.4, "This must be protected for async build, always use the accessors even internally.")
@@ -476,27 +285,6 @@ private:
 	/** Simulation mesh Lods as fed to the solver for constraints creation. Ownership gets transferred to the proxy when it is changed during a simulation. */
 	TSharedPtr<FChaosClothSimulationModel> ClothSimulationModel;
 
-
-#if WITH_EDITORONLY_DATA
-
-	/*
-	* The following PreviewScene properties are modeled after PreviewSkeletalMesh in USkeleton
-	*	- they are inside WITH_EDITORONLY_DATA because they are not used at game runtime
-	*	- TSoftObjectPtrs since that will make it possible to avoid loading these assets until the PreviewScene asks for them
-	*	- DuplicateTransient so that if you copy a ClothAsset it won't copy these preview properties
-	*	- AssetRegistrySearchable makes it so that if the user searches the name of a PreviewScene asset in the Asset Browser, it will return any ClothAssets that use it
-	*/
-
-	/** Optional Skeletal Mesh that the cloth asset is attached to in the Preview Scene in the Cloth Editor */
-	UPROPERTY(DuplicateTransient, AssetRegistrySearchable)
-	TSoftObjectPtr<USkeletalMesh> PreviewSceneSkeletalMesh;
-
-	/** Optional animation attached to PreviewSceneSkeletalMesh in the Preview Scene in the Cloth Editor */
-	UPROPERTY(DuplicateTransient, AssetRegistrySearchable)
-	TSoftObjectPtr<UAnimationAsset> PreviewSceneAnimation;
-
-#endif
-
 #if WITH_EDITOR
 	struct FBuilder;
 #endif
diff --git a/Engine/Plugins/ChaosClothAsset/Source/ChaosClothAssetEngine/Public/ChaosClothAsset/ClothAssetBase.h b/Engine/Plugins/ChaosClothAsset/Source/ChaosClothAssetEngine/Public/ChaosClothAsset/ClothAssetBase.h
new file mode 100644
--- /dev/null
+++ b/Engine/Plugins/ChaosClothAsset/Source/ChaosClothAssetEngine/Public/ChaosClothAsset/ClothAssetBase.h
@@ -0,0 +1,312 @@
+// Copyright Epic Games, Inc. All Rights Reserved.
+
+#pragma once
+
+#include "Engine/SkinnedAsset.h"
+#include "Dataflow/DataflowContent.h"
+#include "Dataflow/DataflowInstance.h"
+#include "ClothAssetBase.generated.h"
+
+class FSkeletalMeshRenderData;
+class FSkinnedAssetCompilationContext;
+class USkeletalMesh;
+class UDataflow;
+struct FChaosClothSimulationModel;
+struct FManagedArrayCollection;
+
+/**
+ * Base cloth asset class.
+ */
+UCLASS(ABSTRACT)
+class CHAOSCLOTHASSETENGINE_API UChaosClothAssetBase
+	: public USkinnedAsset
+	, public IDataflowContentOwner
+	, public IDataflowInstanceInterface
+{
+	GENERATED_BODY()
+
+public:
+	UChaosClothAssetBase(const FObjectInitializer& ObjectInitializer);
+
+	//~ Begin UChaosClothAssetBase interface
+	/** Return whether the asset has any valid cloth simulation models and is simulation enabled. */
+	virtual bool HasValidClothSimulationModels() const
+	PURE_VIRTUAL(UChaosClothAssetBase::HasValidClothSimulationModels, return false;);
+
+	/** Return the number of cloth simulation models in this asset. */
+	virtual int32 GetNumClothSimulationModels() const
+	PURE_VIRTUAL(UChaosClothAssetBase::GetNumClothSimulationModels, return 0;);
+
+	/** Return the simulation model physical data. */
+	virtual TSharedPtr<const FChaosClothSimulationModel> GetClothSimulationModel(int32 ModelIndex) const
+	PURE_VIRTUAL(UChaosClothAssetBase::GetClothSimulationModel, return nullptr;);
+
+	/** Return the name of the specified cloth simulation models in this asset. */
+	virtual FName GetClothSimulationModelName(int32 ModelIndex) const
+	PURE_VIRTUAL(UChaosClothAssetBase::GetClothSimulationModelName, return NAME_None;);
+
+	/** Return the collections that use the property facade for this asset model (one per LOD). */
+	virtual const TArray<TSharedRef<FManagedArrayCollection>>& GetCollections(int32 ModelIndex) const
+	PURE_VIRTUAL(UChaosClothAssetBase::GetCollections, static const TArray<TSharedRef<FManagedArrayCollection>> EmptyArray; return EmptyArray;);
+
+	/** Return the physics asset used for the simulation by this asset model. */
+	virtual const UPhysicsAsset* GetPhysicsAssetForModel(int32 ModelIndex) const
+	PURE_VIRTUAL(UChaosClothAssetBase::GetPhysicsAssetForModel, return nullptr;);
+
+#if WITH_EDITOR
+	/** Load the platform render data from the DDC if cached, otherwise generate the data and cache it. */
+	virtual void CacheDerivedData(FSkinnedAssetCompilationContext* /*Context*/) {}
+#endif
+
+	/** Set the specified reference skeleton or a default reference skeleton with a simple root if ReferenceSkeleton is nullptr. */
+	virtual void SetReferenceSkeleton(const FReferenceSkeleton* ReferenceSkeleton);
+	//~ End UChaosClothAssetBase interface
+
+	//~ Begin IDataflowContentOwner interface 
+	virtual TObjectPtr<UDataflowBaseContent> CreateDataflowContent() override;
+	virtual void WriteDataflowContent(const TObjectPtr<UDataflowBaseContent>& DataflowContent) const override;
+	virtual void ReadDataflowContent(const TObjectPtr<UDataflowBaseContent>& DataflowContent) override;
+	//~ End IDataflowContentOwner interface 
+
+	//~ Begin IDataflowInstanceInterface interface
+	virtual const FDataflowInstance& GetDataflowInstance() const { return DataflowInstance; };
+	//~ End IDataflowInstanceInterface interface
+
+	//~ Begin UObject interface
+	virtual void BeginDestroy() override;
+	virtual bool IsReadyForFinishDestroy() override;
+	//~ End UObject interface
+
+	//~ Begin USkinnedAsset interface
+	virtual FSkeletalMeshLODInfo* GetLODInfo(int32 Index);
+	virtual const FSkeletalMeshLODInfo* GetLODInfo(int32 Index) const;
+	virtual FReferenceSkeleton& GetRefSkeleton() override;
+	virtual const FReferenceSkeleton& GetRefSkeleton() const override;
+	virtual FSkeletalMeshRenderData* GetResourceForRendering() const override;
+	virtual TArray<FSkeletalMeshLODInfo>& GetLODInfoArray() override
+	{
+		return LODInfo;
+	}
+	virtual const TArray<FSkeletalMeshLODInfo>& GetLODInfoArray() const override
+	{
+		return LODInfo;
+	}
+	virtual int32 GetDefaultMinLod() const
+	{
+		return 0;
+	}
+	virtual const FPerPlatformInt& GetMinLod() const override
+	{
+		return MinLod;
+	}
+	virtual int32 GetLODNum() const override
+	{
+		return LODInfo.Num();
+	}
+	virtual FBoxSphereBounds GetBounds() const override
+	{
+		return Bounds;
+	}
+	virtual TArray<class USkeletalMeshSocket*> GetActiveSocketList() const override
+	{
+		return TArray<class USkeletalMeshSocket*>();
+	}
+	virtual USkeletalMeshSocket* FindSocket(FName InSocketName) const override
+	{
+		return nullptr;
+	}
+	virtual USkeletalMeshSocket* FindSocketInfo(FName InSocketName, FTransform& OutTransform, int32& OutBoneIndex, int32& OutIndex) const override
+	{
+		return nullptr;
+	}
+	virtual UMeshDeformer* GetDefaultMeshDeformer() const override
+	{
+		return nullptr;
+	}
+	/** Get the default overlay material used by this mesh */
+	UFUNCTION(BlueprintCallable, Category = "Rendering|Material")
+	virtual class UMaterialInterface* GetOverlayMaterial() const override;
+	/** Get the default overlay material max draw distance used by this mesh */
+	UFUNCTION(BlueprintCallable, Category = "Rendering|Material")
+	virtual float GetOverlayMaterialMaxDrawDistance() const override;
+	virtual bool IsValidLODIndex(int32 Index) const override
+	{
+		return LODInfo.IsValidIndex(Index);
+	}
+	virtual int32 GetMinLodIdx(bool bForceLowestLODIdx = false) const override;
+	virtual bool NeedCPUData(int32 LODIndex) const override
+	{
+		return false;
+	}
+	virtual bool GetHasVertexColors() const override;
+	virtual int32 GetPlatformMinLODIdx(const ITargetPlatform* TargetPlatform) const override;
+	virtual const FPerPlatformBool& GetDisableBelowMinLodStripping() const override { return DisableBelowMinLodStripping; }
+	virtual bool IsMinLodQualityLevelEnable() const override;
+	//~ End USkinnedAsset interface
+
+	const FPerQualityLevelInt& GetQualityLevelMinLod() const
+	{
+		return MinQualityLevelLOD;
+	}
+	void SetQualityLevelMinLod(FPerQualityLevelInt InMinLod)
+	{
+		MinQualityLevelLOD = MoveTemp(InMinLod);
+	}
+
+	/** Change the default overlay material used by this mesh */
+	UFUNCTION(BlueprintCallable, Category = "Rendering|Material")
+	void SetOverlayMaterial(class UMaterialInterface* NewOverlayMaterial);
+	/** Change the default overlay material max draw distance used by this mesh */
+	UFUNCTION(BlueprintCallable, Category = "Rendering|Material")
+	void SetOverlayMaterialMaxDrawDistance(float InMaxDrawDistance);
+
+	/**
+	 * Set the Dataflow graph asset for this asset.
+	 * @param InDataflow The new dataflow asset.
+	 */
+	void SetDataflow(UDataflow* InDataflow)
+	{
+		DataflowInstance.SetDataflowAsset(InDataflow);
+	}
+
+	/** Return the Dataflow graph asset associated to this asset if any. */
+	UDataflow* GetDataflow()
+	{
+		return DataflowInstance.GetDataflowAsset();
+	}
+
+	/** Return the Dataflow graph asset associated to this asset if any, const version. */
+	const UDataflow* GetDataflow() const
+	{
+		return DataflowInstance.GetDataflowAsset();
+	}
+
+#if WITH_EDITORONLY_DATA
+	void SetPreviewSceneSkeletalMesh(USkeletalMesh* Mesh);
+	USkeletalMesh* GetPreviewSceneSkeletalMesh() const;
+
+	void SetPreviewSceneAnimation(UAnimationAsset* Animation);
+	UAnimationAsset* GetPreviewSceneAnimation() const;
+#endif
+
+	void SetHasVertexColors(bool InbHasVertexColors);
+
+protected:
+	/** Reregister all components using this asset to reset the simulation in case anything has changed. */
+	void ReregisterComponents();
+
+	/**
+	 * Wait for the asset to finish compilation to protect internal skinned asset data from race conditions during async build.
+	 * This should be called before accessing all async accessible properties.
+	 */
+	template <typename EnumType UE_REQUIRES(TIsEnum<EnumType>::Value)>
+	void WaitUntilAsyncPropertyReleased(const EnumType AsyncProperties, ESkinnedAssetAsyncPropertyLockType LockType = ESkinnedAssetAsyncPropertyLockType::ReadWrite) const
+	{
+		checkf((uint64)AsyncProperties <= (uint64)TNumericLimits<uint32>::Max(), TEXT("Cannot have more than 32 async properties per derived class."));
+		WaitUntilAsyncPropertyReleasedInternal((uint64)AsyncProperties << 32ull, LockType);
+	}
+
+	/** Set render data. */
+	void SetResourceForRendering(TUniquePtr<FSkeletalMeshRenderData>&& InSkeletalMeshRenderData);
+
+	/** Initialize all render resources. */
+	void InitResources();
+
+	/** Safely release the render data. */
+	void ReleaseResources();
+
+#if WITH_EDITOR
+	/** Initial step for the building process - Can't be done in parallel. USkinnedAsset Interface. */
+	virtual void BeginBuildInternal(FSkinnedAssetBuildContext& Context) override;
+
+	/** Thread-safe part. USkinnedAsset Interface. */
+	virtual void ExecuteBuildInternal(FSkinnedAssetBuildContext& Context) override;
+
+	/** Complete the building process - Can't be done in parallel. USkinnedAsset Interface. */
+	virtual void FinishBuildInternal(FSkinnedAssetBuildContext& Context) override;
+#endif
+
+	UE_DEPRECATED(5.6, "This must be protected for async build, always use the accessors even internally.")
+	UPROPERTY()
+	uint8 bHasVertexColors : 1 = false;
+
+	UPROPERTY(EditAnywhere, Category = "Dataflow")
+	FDataflowInstance DataflowInstance;
+
+	/** Reference skeleton created from the provided skeleton asset. */
+	UE_DEPRECATED(5.4, "This must be protected for async build, always use the accessors even internally.")
+	FReferenceSkeleton RefSkeleton;
+
+	/** Struct containing information for each LOD level, such as materials to use, and when use the LOD. Not currently editable or customizable through the Dataflow. */
+	UPROPERTY(VisibleAnywhere, EditFixedSize, Category = LevelOfDetails)
+	TArray<FSkeletalMeshLODInfo> LODInfo;
+
+	/** Set the Minimum LOD by Quality Level. This property is used when "Use Cloth Asset Min LOD Per Quality Levels" is set at the Project level. Otherwise, the (per platform) Minimum LOD value is used.*/
+	UPROPERTY(EditAnywhere, Category = LODSettings, meta = (DisplayName = "Quality Level Minimum LOD"))
+	FPerQualityLevelInt MinQualityLevelLOD = 0;
+
+	UPROPERTY(EditAnywhere, Category = LODSettings)
+	FPerPlatformBool DisableBelowMinLodStripping;
+
+	/** Set the Minimum LOD by platform. This property is overriden by "Quality Level Minimum LOD" when "Use Cloth Asset Min LOD Per Quality Levels" is set at the Project level.*/
+	UPROPERTY(EditAnywhere, Category = LODSettings, Meta = (DisplayName = "Minimum LOD"))
+	FPerPlatformInt MinLod = 0;
+
+	/** Enable raytracing for this asset. */
+	UPROPERTY(EditAnywhere, Category = RayTracing)
+	uint8 bSupportRayTracing : 1;
+
+	/** Minimum raytracing LOD for this asset. */
+	UPROPERTY(EditAnywhere, Category = RayTracing)
+	int32 RayTracingMinLOD;
+
+	/** Default translucent material to blend on top of this mesh. Mesh will be rendered twice - once with a base material and once with overlay material */
+	UE_DEPRECATED(5.6, "This must be protected for async build, always use the accessors even internally.")
+	UPROPERTY(EditAnywhere, Category = Rendering)
+	TObjectPtr<class UMaterialInterface> OverlayMaterial;
+
+	/** Default max draw distance for overlay material. A distance of 0 indicates that overlay will be culled using primitive max distance. */
+	UE_DEPRECATED(5.6, "This must be protected for async build, always use the accessors even internally.")
+	UPROPERTY(EditAnywhere, Category = Rendering)
+	float OverlayMaterialMaxDrawDistance;
+
+
+	/** Bounds for this asset. */
+	FBoxSphereBounds Bounds;
+
+private:
+	enum class EAsyncProperties : uint32;
+	FRIEND_ENUM_CLASS_FLAGS(EAsyncProperties);
+
+	/**
+	 * Wait for the asset to finish compilation to protect internal skinned asset data from race conditions during async build.
+	 * This should be called before accessing all async accessible properties.
+	 */
+	template <>
+	void WaitUntilAsyncPropertyReleased<EAsyncProperties>(const EAsyncProperties AsyncProperties, ESkinnedAssetAsyncPropertyLockType LockType) const
+	{
+		WaitUntilAsyncPropertyReleasedInternal((uint64)AsyncProperties, LockType);
+	}
+
+#if WITH_EDITORONLY_DATA
+	// The following PreviewScene properties are modeled after PreviewSkeletalMesh in USkeleton
+	//	- they are inside WITH_EDITORONLY_DATA because they are not used at game runtime
+	//	- TSoftObjectPtrs since that will make it possible to avoid loading these assets until the PreviewScene asks for them
+	//	- DuplicateTransient so that if you copy a ClothAsset it won't copy these preview properties
+	//	- AssetRegistrySearchable makes it so that if the user searches the name of a PreviewScene asset in the Asset Browser, it will return any ClothAssets that use it
+
+	/** Optional Skeletal Mesh that the cloth asset is attached to in the Preview Scene in the Cloth Editor */
+	UPROPERTY(DuplicateTransient, AssetRegistrySearchable)
+	TSoftObjectPtr<USkeletalMesh> PreviewSceneSkeletalMesh;
+
+	/** Optional animation attached to PreviewSceneSkeletalMesh in the Preview Scene in the Cloth Editor */
+	UPROPERTY(DuplicateTransient, AssetRegistrySearchable)
+	TSoftObjectPtr<UAnimationAsset> PreviewSceneAnimation;
+#endif
+
+	/** Rendering data. */
+	TUniquePtr<FSkeletalMeshRenderData> SkeletalMeshRenderData;
+
+	/** A fence which is used to keep track of the rendering thread releasing the static mesh resources. */
+	FRenderCommandFence ReleaseResourcesFence;
+};
diff --git a/Engine/Plugins/ChaosClothAsset/Source/ChaosClothAssetEngine/Public/ChaosClothAsset/ClothComponent.h b/Engine/Plugins/ChaosClothAsset/Source/ChaosClothAssetEngine/Public/ChaosClothAsset/ClothComponent.h
--- a/Engine/Plugins/ChaosClothAsset/Source/ChaosClothAssetEngine/Public/ChaosClothAsset/ClothComponent.h
+++ b/Engine/Plugins/ChaosClothAsset/Source/ChaosClothAssetEngine/Public/ChaosClothAsset/ClothComponent.h
@@ -12,6 +12,7 @@ namespace Dataflow = UE::Dataflow;
 namespace UE_DEPRECATED(5.5, "Use UE::Dataflow instead.") Dataflow {}
 #endif
 
+class UChaosClothAssetBase;
 class UChaosClothAsset;
 class UChaosClothComponent;
 class UChaosClothAssetInteractor;
@@ -39,18 +40,28 @@ UCLASS(
 	HideCategories = (Object, "Mesh|SkeletalAsset", Constraints, Advanced, Cooking, Collision, Navigation))
 class CHAOSCLOTHASSETENGINE_API UChaosClothComponent : public USkinnedMeshComponent, public IDataflowPhysicsSolverInterface
 {
-	GENERATED_BODY()	
+	GENERATED_BODY()
 public:
 	UChaosClothComponent(const FObjectInitializer& ObjectInitializer);
 	UChaosClothComponent(FVTableHelper& Helper);
 	~UChaosClothComponent();
 
+	/** Set the cloth object used by this component, could be a cloth asset, an outfit asset, or any other type of asset inheriting from UChaosClothAssetBase. */
+	UFUNCTION(BlueprintCallable, Category = "ClothComponent", Meta = (Keywords = "Chaos Cloth Outfit Asset"))
+	void SetAsset(UChaosClothAssetBase* InAsset);
+
+	/** Get the cloth object used by this component, could be a cloth asset, an outfit asset, or any other type of asset inheriting from UChaosClothAssetBase. */
+	UFUNCTION(BlueprintPure, Category = "ClothComponent", Meta = (Keywords = "Chaos Cloth Outfit Asset"))
+	UChaosClothAssetBase* GetAsset() const;
+
 	/** Set the cloth asset used by this component. */
-	UFUNCTION(BlueprintCallable, Category = "ClothComponent", Meta = (Keywords = "Chaos Cloth Asset"))
+	UFUNCTION(BlueprintCallable, Category = "ClothComponent", Meta = (DeprecatedFunction, Keywords = "Chaos Cloth Asset"))
+	UE_DEPRECATED(5.6, "Use SetAsset instead.")
 	void SetClothAsset(UChaosClothAsset* InClothAsset);
 
 	/** Get the cloth asset used by this component. */
-	UFUNCTION(BlueprintPure, Category = "ClothComponent", Meta = (Keywords = "Chaos Cloth Asset"))
+	UFUNCTION(BlueprintPure, Category = "ClothComponent", Meta = (DeprecatedFunction, Keywords = "Chaos Cloth Asset"))
+	UE_DEPRECATED(5.6, "Use GetAsset instead.")
 	UChaosClothAsset* GetClothAsset() const;
 
 	/** Reset the teleport mode. */
@@ -145,11 +156,16 @@ public:
 	void SetSimulateInEditor(const bool bNewSimulateState);
 
 	/**
-	 * Return the property collections holding the runtime properties for this cloth component (one per LOD).
+	 * Return the property collections holding the runtime properties for this cloth component model (one per LOD).
 	 * This might be different from the cloth asset's since the component's properties can be modified in code or in blueprints.
 	 * This could also be different from the cloth simulation object until the cloth simulation thread synchronise the properties.
 	 */
-	const TArray<TSharedPtr<const FManagedArrayCollection>>& GetPropertyCollections() const 
+	const TArray<TSharedPtr<FManagedArrayCollection>>& GetPropertyCollections(int32 /*ModelIndex*/) const
+	{
+		return PropertyCollections;
+	}
+	UE_DEPRECATED(5.6, "Use GetPropertyCollections(int32 ModelIndex) instead.")
+	const TArray<TSharedPtr<const FManagedArrayCollection>>& GetPropertyCollections() const
 	{
 		return reinterpret_cast<const TArray<TSharedPtr<const FManagedArrayCollection>>&>(PropertyCollections);
 	}
@@ -286,10 +302,10 @@ private:
 	UPROPERTY(EditInstanceOnly, Transient, Category = ClothComponent)
 	uint8 bSimulateInEditor : 1;
 
-	/** Cloth asset used by this component. */
-	UE_DEPRECATED(5.1, "This property isn't deprecated, but getter and setter must be used at all times to preserve correct operations.")
-	UPROPERTY(EditAnywhere, Transient, Setter = SetClothAsset, BlueprintSetter = SetClothAsset, Getter = GetClothAsset, BlueprintGetter = GetClothAsset, Category = ClothComponent)
-	TObjectPtr<UChaosClothAsset> ClothAsset;
+	/** Asset used by this component. Can be either a cloth asset or an outfit asset. */
+	UE_DEPRECATED(5.6, "This property isn't deprecated, but getter and setter must be used at all times to preserve correct operations.")
+	UPROPERTY(EditAnywhere, Transient, Setter = SetAsset, BlueprintSetter = SetAsset, Getter = GetAsset, BlueprintGetter = GetAsset, Category = ClothComponent)
+	TObjectPtr<UChaosClothAssetBase> Asset;
 #endif
 
 	UPROPERTY(Transient)
diff --git a/Engine/Plugins/ChaosClothAssetEditor/Source/ChaosClothAssetEditor/Private/ChaosClothAsset/ChaosClothAssetEditorModule.cpp b/Engine/Plugins/ChaosClothAssetEditor/Source/ChaosClothAssetEditor/Private/ChaosClothAsset/ChaosClothAssetEditorModule.cpp
--- a/Engine/Plugins/ChaosClothAssetEditor/Source/ChaosClothAssetEditor/Private/ChaosClothAsset/ChaosClothAssetEditorModule.cpp
+++ b/Engine/Plugins/ChaosClothAssetEditor/Source/ChaosClothAssetEditor/Private/ChaosClothAsset/ChaosClothAssetEditorModule.cpp
@@ -38,7 +38,7 @@ namespace UE::Chaos::ClothAsset
 			{
 				if (UChaosClothAsset* ClothAsset = Cast<UChaosClothAsset>(InAsset))
 				{
-					ClothComponent->SetClothAsset(ClothAsset);
+					ClothComponent->SetAsset(ClothAsset);
 					return true;
 				}
 			}
@@ -49,7 +49,7 @@ namespace UE::Chaos::ClothAsset
 		{
 			if (UChaosClothComponent* ClothComponent = Cast<UChaosClothComponent>(InComponent))
 			{
-				return ClothComponent->GetClothAsset();
+				return Cast<UChaosClothAsset>(ClothComponent->GetAsset());
 			}
 			return nullptr;
 		}
diff --git a/Engine/Plugins/ChaosClothAssetEditor/Source/ChaosClothAssetEditor/Private/ChaosClothAsset/ClothEditorMode.cpp b/Engine/Plugins/ChaosClothAssetEditor/Source/ChaosClothAssetEditor/Private/ChaosClothAsset/ClothEditorMode.cpp
--- a/Engine/Plugins/ChaosClothAssetEditor/Source/ChaosClothAssetEditor/Private/ChaosClothAsset/ClothEditorMode.cpp
+++ b/Engine/Plugins/ChaosClothAssetEditor/Source/ChaosClothAssetEditor/Private/ChaosClothAsset/ClothEditorMode.cpp
@@ -1316,7 +1316,7 @@ FBox UChaosClothAssetEditorMode::PreviewBoundingBox() const
 
 	if (const UChaosClothComponent* const Cloth = PreviewScene->GetClothComponent())
 	{
-		if (Cloth->GetClothAsset())
+		if (Cast<UChaosClothAsset>(Cloth->GetAsset()))
 		{
 			Bounds += Cloth->Bounds.GetBox();
 		}
diff --git a/Engine/Plugins/ChaosClothAssetEditor/Source/ChaosClothAssetEditor/Private/ChaosClothAsset/ClothEditorPreviewScene.cpp b/Engine/Plugins/ChaosClothAssetEditor/Source/ChaosClothAssetEditor/Private/ChaosClothAsset/ClothEditorPreviewScene.cpp
--- a/Engine/Plugins/ChaosClothAssetEditor/Source/ChaosClothAssetEditor/Private/ChaosClothAsset/ClothEditorPreviewScene.cpp
+++ b/Engine/Plugins/ChaosClothAssetEditor/Source/ChaosClothAssetEditor/Private/ChaosClothAsset/ClothEditorPreviewScene.cpp
@@ -27,7 +27,7 @@ namespace UE::Chaos::ClothAsset::Private
 {
 	void ValidateClothComponentAttachmentBones(const UChaosClothComponent& ClothComponent)
 	{
-		const UChaosClothAsset* const ClothAsset = ClothComponent.GetClothAsset();
+		const UChaosClothAsset* const ClothAsset = Cast<const UChaosClothAsset>(ClothComponent.GetAsset());
 		const FSkeletalMeshRenderData* const ClothRenderData = ClothComponent.GetSkeletalMeshRenderData();
 		const USkeletalMeshComponent* const SkeletalMeshComponent = Cast<USkeletalMeshComponent>(ClothComponent.LeaderPoseComponent.Get());
 		if (ClothAsset && ClothRenderData && SkeletalMeshComponent)
@@ -258,7 +258,7 @@ void FChaosClothPreviewScene::SceneDescriptionPropertyChanged(const FName& Prope
 
 		UE::Chaos::ClothAsset::Private::ValidateClothComponentAttachmentBones(*ClothComponent);
 
-		if (UChaosClothAsset* const ClothAsset = ClothComponent->GetClothAsset())
+		if (UChaosClothAsset* const ClothAsset = Cast<UChaosClothAsset>(ClothComponent->GetAsset()))
 		{
 			ClothAsset->SetPreviewSceneSkeletalMesh(PreviewSceneDescription->SkeletalMeshAsset);
 		}
@@ -278,7 +278,7 @@ void FChaosClothPreviewScene::SceneDescriptionPropertyChanged(const FName& Prope
 	{
 		UpdateSkeletalMeshAnimation();
 
-		if (UChaosClothAsset* const ClothAsset = ClothComponent->GetClothAsset())
+		if (UChaosClothAsset* const ClothAsset = Cast<UChaosClothAsset>(ClothComponent->GetAsset()))
 		{
 			ClothAsset->SetPreviewSceneAnimation(PreviewSceneDescription->AnimationAsset);
 		}
@@ -341,14 +341,14 @@ void FChaosClothPreviewScene::SetClothAsset(UChaosClothAsset* Asset)
 	check(SceneActor);
 	check(ClothComponent);
 
-	if (ClothComponent->GetClothAsset() == Asset)
+	if (Cast<UChaosClothAsset>(ClothComponent->GetAsset()) == Asset)
 	{
 		// Update the config properties on the component from the asset.
 		ClothComponent->UpdateConfigProperties();
 	}
 	else
 	{
-		ClothComponent->SetClothAsset(Asset);
+		ClothComponent->SetAsset(Asset);
 	}
 
 	UpdateClothComponentAttachment();
diff --git a/Engine/Plugins/ChaosClothAssetEditor/Source/ChaosClothAssetEditor/Private/ChaosClothAsset/ClothEditorToolkit.cpp b/Engine/Plugins/ChaosClothAssetEditor/Source/ChaosClothAssetEditor/Private/ChaosClothAsset/ClothEditorToolkit.cpp
--- a/Engine/Plugins/ChaosClothAssetEditor/Source/ChaosClothAssetEditor/Private/ChaosClothAsset/ClothEditorToolkit.cpp
+++ b/Engine/Plugins/ChaosClothAssetEditor/Source/ChaosClothAssetEditor/Private/ChaosClothAsset/ClothEditorToolkit.cpp
@@ -919,13 +919,13 @@ void FChaosClothAssetEditorToolkit::EvaluateNode(const FDataflowNode* Node, cons
 			Dataflow,
 			Node,
 			Output,
-			ClothAsset->GetDataflowTerminal(),
+			ClothAsset->GetDataflowInstance().GetDataflowTerminal().ToString(),
 			ClothAsset);
 
 		// Refresh editor
 		if (Node && OldTimestamp < LastDataflowNodeTimestamp)
 		{
-			if (Node->GetName() == ClothAsset->GetDataflowTerminal())
+			if (Node->GetName() == ClothAsset->GetDataflowInstance().GetDataflowTerminal())
 			{
 				OnClothAssetChanged();
 			}
@@ -1377,7 +1377,7 @@ void FChaosClothAssetEditorToolkit::OnClothAssetChanged()
 
 	if (UChaosClothAsset* const ClothAsset = Cast<UChaosClothAsset>(ObjectsToEdit[0]))
 	{
-		const bool bHadClothAsset = (ClothPreviewScene->GetClothComponent()->GetClothAsset() != nullptr);
+		const bool bHadClothAsset = (ClothPreviewScene->GetClothComponent()->GetAsset() != nullptr);
 
 		ClothPreviewScene->SetClothAsset(ClothAsset);
 
diff --git a/Engine/Plugins/Experimental/ChaosOutfitAsset/Source/ChaosOutfitAssetEngine/Private/ChaosClothAsset/OutfitAsset.cpp b/Engine/Plugins/Experimental/ChaosOutfitAsset/Source/ChaosOutfitAssetEngine/Private/ChaosClothAsset/OutfitAsset.cpp
--- a/Engine/Plugins/Experimental/ChaosOutfitAsset/Source/ChaosOutfitAssetEngine/Private/ChaosClothAsset/OutfitAsset.cpp
+++ b/Engine/Plugins/Experimental/ChaosOutfitAsset/Source/ChaosOutfitAssetEngine/Private/ChaosClothAsset/OutfitAsset.cpp
@@ -7,7 +7,7 @@
 UChaosOutfitAsset::UChaosOutfitAsset(const FObjectInitializer& ObjectInitializer)
 	: Super(ObjectInitializer)
 {
-	SetDataflowTerminal(TEXT("OutfitAssetTerminal"));
+	DataflowInstance.SetDataflowTerminal(TEXT("OutfitAssetTerminal"));
 }
 
 UChaosOutfitAsset::UChaosOutfitAsset(FVTableHelper& Helper)
