Chaos Outfit Asset - Added first pass Outfit Asset with working simulation only.

#rnx

diff --git a/Engine/Plugins/ChaosClothAsset/Source/ChaosClothAssetEngine/Private/ChaosClothAsset/ClothAsset.cpp b/Engine/Plugins/ChaosClothAsset/Source/ChaosClothAssetEngine/Private/ChaosClothAsset/ClothAsset.cpp
--- a/Engine/Plugins/ChaosClothAsset/Source/ChaosClothAssetEngine/Private/ChaosClothAsset/ClothAsset.cpp
+++ b/Engine/Plugins/ChaosClothAsset/Source/ChaosClothAssetEngine/Private/ChaosClothAsset/ClothAsset.cpp
@@ -21,6 +21,7 @@
 #include "Rendering/SkeletalMeshModel.h"
 #include "Rendering/SkeletalMeshRenderData.h"
 #include "UObject/Package.h"
+#include "UObject/FortniteMainBranchObjectVersion.h"
 #include "EngineUtils.h"
 #if WITH_EDITOR
 #include "IMeshBuilderModule.h"
@@ -171,27 +172,12 @@ PRAGMA_DISABLE_DEPRECATION_WARNINGS
 UChaosClothAsset::~UChaosClothAsset() = default;
 PRAGMA_ENABLE_DEPRECATION_WARNINGS
 
-FMatrix UChaosClothAsset::GetComposedRefPoseMatrix(FName InBoneName) const
-{
-	FMatrix LocalPose(FMatrix::Identity);
-
-	if (InBoneName != NAME_None)
-	{
-		const int32 BoneIndex = GetRefSkeleton().FindBoneIndex(InBoneName);
-		if (BoneIndex != INDEX_NONE)
-		{
-			return GetComposedRefPoseMatrix(BoneIndex);
-		}
-		// TODO: Might need to add sockets like on the SkeletalMesh
-	}
-
-	return LocalPose;
-}
-
 void UChaosClothAsset::Serialize(FArchive& Ar)
 {
 	Super::Serialize(Ar);
 
+	Ar.UsingCustomVersion(FFortniteMainBranchObjectVersion::GUID);
+
 	bool bCooked = Ar.IsCooking();
 	Ar << bCooked;
 
@@ -208,7 +194,10 @@ void UChaosClothAsset::Serialize(FArchive& Ar)
 		UE::Chaos::ClothAsset::Private::Serialize(ChaosArchive, GetClothCollections());
 	}
 
-	Ar << GetRefSkeleton();
+	if (Ar.CustomVer(FFortniteMainBranchObjectVersion::GUID) < FFortniteMainBranchObjectVersion::AddClothAssetBase)
+	{
+		Ar << GetRefSkeleton(); // Moved to Cloth Asset Base serialization
+	}
 
 	if (bCooked && !IsTemplate() && !Ar.IsCountingMemory())  // Counting of these resources are done in GetResourceSizeEx, so skip these when counting memory
 	{
@@ -451,58 +440,6 @@ void UChaosClothAsset::FinishPostLoadInternal(FSkinnedAssetPostLoadContext& Cont
 #endif
 }
 
-void UChaosClothAsset::CalculateInvRefMatrices()
-{
-	auto GetRefPoseMatrix = [this](int32 BoneIndex)->FMatrix
-	{
-		check(BoneIndex >= 0 && BoneIndex < GetRefSkeleton().GetRawBoneNum());
-		FTransform BoneTransform = GetRefSkeleton().GetRawRefBonePose()[BoneIndex];
-		BoneTransform.NormalizeRotation();  // Make sure quaternion is normalized!
-		return BoneTransform.ToMatrixWithScale();
-	};
-
-	const int32 NumRealBones = GetRefSkeleton().GetRawBoneNum();
-
-	RefBasesInvMatrix.Empty(NumRealBones);
-	RefBasesInvMatrix.AddUninitialized(NumRealBones);
-
-	// Reset cached mesh-space ref pose
-	TArray<FMatrix> ComposedRefPoseMatrices;
-	ComposedRefPoseMatrices.SetNumUninitialized(NumRealBones);
-
-	// Precompute the Mesh.RefBasesInverse
-	for (int32 BoneIndex = 0; BoneIndex < NumRealBones; ++BoneIndex)
-	{
-		// Render the default pose
-		ComposedRefPoseMatrices[BoneIndex] = GetRefPoseMatrix(BoneIndex);
-
-		// Construct mesh-space skeletal hierarchy
-		if (BoneIndex > 0)
-		{
-			int32 Parent = GetRefSkeleton().GetRawParentIndex(BoneIndex);
-			ComposedRefPoseMatrices[BoneIndex] = ComposedRefPoseMatrices[BoneIndex] * ComposedRefPoseMatrices[Parent];
-		}
-
-		FVector XAxis, YAxis, ZAxis;
-		ComposedRefPoseMatrices[BoneIndex].GetScaledAxes(XAxis, YAxis, ZAxis);
-		if (XAxis.IsNearlyZero(UE_SMALL_NUMBER) &&
-			YAxis.IsNearlyZero(UE_SMALL_NUMBER) &&
-			ZAxis.IsNearlyZero(UE_SMALL_NUMBER))
-		{
-			// This is not allowed, warn them
-			UE_LOG(
-				LogChaosClothAsset,
-				Warning,
-				TEXT("Reference Pose for asset %s for joint (%s) includes NIL matrix. Zero scale isn't allowed on ref pose."),
-				*GetPathName(),
-				*GetRefSkeleton().GetBoneName(BoneIndex).ToString());
-		}
-
-		// Precompute inverse so we can use from-refpose-skin vertices
-		RefBasesInvMatrix[BoneIndex] = FMatrix44f(ComposedRefPoseMatrices[BoneIndex].Inverse());
-	}
-}
-
 void UChaosClothAsset::CalculateBounds()
 {
 	using namespace UE::Chaos::ClothAsset;
@@ -606,18 +543,6 @@ void UChaosClothAsset::BuildClothSimulationModel(TArray<FChaosClothAssetLodTrans
 		GetRefSkeleton(), InOutTransitionCache);
 }
 
-const FMeshUVChannelInfo* UChaosClothAsset::GetUVChannelData(int32 MaterialIndex) const
-{
-	if (GetMaterials().IsValidIndex(MaterialIndex))
-	{
-		// TODO: enable ensure when UVChannelData is setup
-		//ensure(GetMaterials()[MaterialIndex].UVChannelData.bInitialized);
-		return &GetMaterials()[MaterialIndex].UVChannelData;
-	}
-
-	return nullptr;
-}
-
 FString UChaosClothAsset::GetAsyncPropertyName(uint64 Property) const
 {
 	return StaticEnum<EClothAssetAsyncProperties>()->GetNameByValue(Property).ToString();
diff --git a/Engine/Plugins/ChaosClothAsset/Source/ChaosClothAssetEngine/Private/ChaosClothAsset/ClothAssetBase.cpp b/Engine/Plugins/ChaosClothAsset/Source/ChaosClothAssetEngine/Private/ChaosClothAsset/ClothAssetBase.cpp
--- a/Engine/Plugins/ChaosClothAsset/Source/ChaosClothAssetEngine/Private/ChaosClothAsset/ClothAssetBase.cpp
+++ b/Engine/Plugins/ChaosClothAsset/Source/ChaosClothAssetEngine/Private/ChaosClothAsset/ClothAssetBase.cpp
@@ -4,6 +4,7 @@
 #if WITH_EDITORONLY_DATA
 #include "Animation/AnimationAsset.h"
 #endif
+#include "ChaosClothAsset/ClothAssetPrivate.h"
 #include "ChaosClothAsset/ClothComponent.h"
 #include "Engine/Engine.h"
 #include "Engine/SkeletalMesh.h"
@@ -12,6 +13,8 @@
 #if WITH_EDITOR
 #include "Interfaces/ITargetPlatform.h"
 #endif
+#include "UObject/FortniteMainBranchObjectVersion.h"
+
 
 #include UE_INLINE_GENERATED_CPP_BY_NAME(ClothAssetBase)
 namespace UE::Chaos::ClothAsset::Private
@@ -156,6 +159,19 @@ bool UChaosClothAssetBase::IsReadyForFinishDestroy()
 	return ReleaseResourcesFence.IsFenceComplete();
 }
 
+void UChaosClothAssetBase::Serialize(FArchive& Ar)
+{
+	Super::Serialize(Ar);
+
+	Ar.UsingCustomVersion(FFortniteMainBranchObjectVersion::GUID);
+
+	if (Ar.CustomVer(FFortniteMainBranchObjectVersion::GUID) < FFortniteMainBranchObjectVersion::AddClothAssetBase)
+	{
+		return;
+	}
+	Ar << GetRefSkeleton();
+}
+
 FSkeletalMeshLODInfo* UChaosClothAssetBase::GetLODInfo(int32 Index)
 {
 	return LODInfo.IsValidIndex(Index) ? &LODInfo[Index] : nullptr;
@@ -203,6 +219,18 @@ PRAGMA_DISABLE_DEPRECATION_WARNINGS
 PRAGMA_ENABLE_DEPRECATION_WARNINGS
 }
 
+const FMeshUVChannelInfo* UChaosClothAssetBase::GetUVChannelData(int32 MaterialIndex) const
+{
+	if (GetMaterials().IsValidIndex(MaterialIndex))
+	{
+		// TODO: enable ensure when UVChannelData is setup
+		//ensure(GetMaterials()[MaterialIndex].UVChannelData.bInitialized);
+		return &GetMaterials()[MaterialIndex].UVChannelData;
+	}
+
+	return nullptr;
+}
+
 int32 UChaosClothAssetBase::GetMinLodIdx(bool bForceLowestLODIdx) const
 {
 	if (IsMinLodQualityLevelEnable())
@@ -293,6 +321,58 @@ void UChaosClothAssetBase::ReleaseResources()
 	}
 }
 
+void UChaosClothAssetBase::CalculateInvRefMatrices()
+{
+	auto GetRefPoseMatrix = [this](int32 BoneIndex)->FMatrix
+		{
+			check(BoneIndex >= 0 && BoneIndex < GetRefSkeleton().GetRawBoneNum());
+			FTransform BoneTransform = GetRefSkeleton().GetRawRefBonePose()[BoneIndex];
+			BoneTransform.NormalizeRotation();  // Make sure quaternion is normalized!
+			return BoneTransform.ToMatrixWithScale();
+		};
+
+	const int32 NumRealBones = GetRefSkeleton().GetRawBoneNum();
+
+	RefBasesInvMatrix.Empty(NumRealBones);
+	RefBasesInvMatrix.AddUninitialized(NumRealBones);
+
+	// Reset cached mesh-space ref pose
+	TArray<FMatrix> ComposedRefPoseMatrices;
+	ComposedRefPoseMatrices.SetNumUninitialized(NumRealBones);
+
+	// Precompute the Mesh.RefBasesInverse
+	for (int32 BoneIndex = 0; BoneIndex < NumRealBones; ++BoneIndex)
+	{
+		// Render the default pose
+		ComposedRefPoseMatrices[BoneIndex] = GetRefPoseMatrix(BoneIndex);
+
+		// Construct mesh-space skeletal hierarchy
+		if (BoneIndex > 0)
+		{
+			int32 Parent = GetRefSkeleton().GetRawParentIndex(BoneIndex);
+			ComposedRefPoseMatrices[BoneIndex] = ComposedRefPoseMatrices[BoneIndex] * ComposedRefPoseMatrices[Parent];
+		}
+
+		FVector XAxis, YAxis, ZAxis;
+		ComposedRefPoseMatrices[BoneIndex].GetScaledAxes(XAxis, YAxis, ZAxis);
+		if (XAxis.IsNearlyZero(UE_SMALL_NUMBER) &&
+			YAxis.IsNearlyZero(UE_SMALL_NUMBER) &&
+			ZAxis.IsNearlyZero(UE_SMALL_NUMBER))
+		{
+			// This is not allowed, warn them
+			UE_LOG(
+				LogChaosClothAsset,
+				Warning,
+				TEXT("Reference Pose for asset %s for joint (%s) includes NIL matrix. Zero scale isn't allowed on ref pose."),
+				*GetPathName(),
+				*GetRefSkeleton().GetBoneName(BoneIndex).ToString());
+		}
+
+		// Precompute inverse so we can use from-refpose-skin vertices
+		RefBasesInvMatrix[BoneIndex] = FMatrix44f(ComposedRefPoseMatrices[BoneIndex].Inverse());
+	}
+}
+
 void UChaosClothAssetBase::SetResourceForRendering(TUniquePtr<FSkeletalMeshRenderData>&& InSkeletalMeshRenderData)
 {
 	WaitUntilAsyncPropertyReleased(EAsyncProperties::RenderData);
diff --git a/Engine/Plugins/ChaosClothAsset/Source/ChaosClothAssetEngine/Private/ChaosClothAsset/ClothComponent.cpp b/Engine/Plugins/ChaosClothAsset/Source/ChaosClothAssetEngine/Private/ChaosClothAsset/ClothComponent.cpp
--- a/Engine/Plugins/ChaosClothAsset/Source/ChaosClothAssetEngine/Private/ChaosClothAsset/ClothComponent.cpp
+++ b/Engine/Plugins/ChaosClothAsset/Source/ChaosClothAssetEngine/Private/ChaosClothAsset/ClothComponent.cpp
@@ -19,6 +19,22 @@
 
 CSV_DECLARE_CATEGORY_MODULE_EXTERN(ENGINE_API, Animation);
 
+void FChaosClothSimulationProperties::Initialize(const TArray<TSharedRef<FManagedArrayCollection>>& AssetCollections)
+{
+	using namespace ::Chaos::Softs;
+	PropertyCollections.Reset(AssetCollections.Num());
+	CollectionPropertyFacades.Reset(AssetCollections.Num());
+	for (const TSharedRef<FManagedArrayCollection>& AssetCollection : AssetCollections)
+	{
+		TSharedPtr<FManagedArrayCollection>& PropertyCollection = PropertyCollections.Emplace_GetRef(MakeShared<FManagedArrayCollection>());
+		TSharedRef<FCollectionPropertyMutableFacade> CollectionPropertyMutableFacade = MakeShared<FCollectionPropertyMutableFacade>(PropertyCollection);
+		CollectionPropertyMutableFacade->Copy(*AssetCollection);  // This only copy properties from the asset's collection (in case it also contains other groups)
+		CollectionPropertyFacades.Emplace(StaticCastSharedRef<FCollectionPropertyFacade>(CollectionPropertyMutableFacade));  // Mutable facades conveniently inherits from non mutable ones
+	}
+	ClothOutfitInteractor = NewObject<UChaosClothAssetInteractor>();
+	ClothOutfitInteractor->SetProperties(CollectionPropertyFacades);
+}
+
 UChaosClothComponent::UChaosClothComponent(const FObjectInitializer& ObjectInitializer)
 	: Super(ObjectInitializer)
 	, bUseAttachedParentAsPoseComponent(1)  // By default use the parent component as leader pose component
@@ -82,33 +98,19 @@ bool UChaosClothComponent::IsSimulationEnabled() const
 
 void UChaosClothComponent::ResetConfigProperties()
 {
+	ClothSimulationProperties.Reset();
+
 	if (IsRegistered())
 	{
-		if (UChaosClothAsset* const ClothAsset = Cast<UChaosClothAsset>(GetAsset()))
+		if (GetAsset())
 		{
-			const TArray<TSharedRef<FManagedArrayCollection>>& ClothCollections = ClothAsset->GetClothCollections();
-			PropertyCollections.Reset();
-			PropertyCollections.Reserve(ClothCollections.Num());
-			CollectionPropertyFacades.Reset();
-			CollectionPropertyFacades.Reserve(ClothCollections.Num());
-			for (const TSharedRef<FManagedArrayCollection>& ClothCollection : ClothCollections)
-			{
-				TSharedPtr<FManagedArrayCollection>& PropertyCollection = PropertyCollections.Add_GetRef(MakeShared<FManagedArrayCollection>());
+			const int32 NumClothSimulationModels = GetAsset()->GetNumClothSimulationModels();
+			ClothSimulationProperties.SetNum(NumClothSimulationModels);
 
-				::Chaos::Softs::FCollectionPropertyMutableFacade CollectionPropertyMutableFacade(PropertyCollection);
-				CollectionPropertyMutableFacade.Copy(*ClothCollection);
-
-				CollectionPropertyFacades.Add(MakeShared<::Chaos::Softs::FCollectionPropertyFacade>(PropertyCollection));
+			for (int32 ModelIndex = 0; ModelIndex < NumClothSimulationModels; ++ModelIndex)
+			{
+				ClothSimulationProperties[ModelIndex].Initialize(GetAsset()->GetCollections(ModelIndex));
 			}
-			check(ClothOutfitInteractor);
-			ClothOutfitInteractor->SetProperties(CollectionPropertyFacades);
-		}
-		else
-		{
-			check(ClothOutfitInteractor);
-			ClothOutfitInteractor->ResetProperties();
-			PropertyCollections.Reset();
-			CollectionPropertyFacades.Reset();
 		}
 	}
 	else
@@ -120,17 +122,23 @@ void UChaosClothComponent::ResetConfigProperties()
 #if WITH_EDITOR
 void UChaosClothComponent::UpdateConfigProperties()
 {
+	using namespace ::Chaos::Softs;
 	if (IsRegistered())
 	{
-		if (UChaosClothAsset* const ClothAsset = Cast<UChaosClothAsset>(GetAsset()))
+		if (GetAsset() && GetAsset()->GetNumClothSimulationModels() == ClothSimulationProperties.Num())
 		{
-			const TArray<TSharedRef<FManagedArrayCollection>>& ClothCollections = ClothAsset->GetClothCollections();
-			if (ClothCollections.Num() == PropertyCollections.Num())
+			for (int32 ModelIndex = 0; ModelIndex < ClothSimulationProperties.Num(); ++ModelIndex)
 			{
-				check(CollectionPropertyFacades.Num() == ClothCollections.Num());
-				for (int32 LodIndex = 0; LodIndex < ClothCollections.Num(); ++LodIndex)
+				const TArray<TSharedRef<FManagedArrayCollection>>& AssetCollections = GetAsset()->GetCollections(ModelIndex);
+				TArray<TSharedPtr<FManagedArrayCollection>>& PropertyCollections = ClothSimulationProperties[ModelIndex].PropertyCollections;
+				if (AssetCollections.Num() == PropertyCollections.Num())
 				{
-					CollectionPropertyFacades[LodIndex]->UpdateProperties(ClothCollections[LodIndex].ToSharedPtr());
+					TArray<TSharedPtr<FCollectionPropertyFacade>>& CollectionPropertyFacades = ClothSimulationProperties[ModelIndex].CollectionPropertyFacades;
+					check(CollectionPropertyFacades.Num() == AssetCollections.Num());
+					for (int32 LODIndex = 0; LODIndex < AssetCollections.Num(); ++LODIndex)
+					{
+						CollectionPropertyFacades[LODIndex]->UpdateProperties(AssetCollections[LODIndex].ToSharedPtr());
+					}
 				}
 			}
 		}
@@ -151,14 +159,10 @@ void UChaosClothComponent::RecreateClothSimulationProxy()
 	{
 		ClothSimulationProxy.Reset();
 
-		if (UChaosClothAsset* const ClothAsset = Cast<UChaosClothAsset>(GetAsset()))
+		if (GetAsset() && GetAsset()->HasValidClothSimulationModels())
 		{
-			const TSharedPtr<const FChaosClothSimulationModel> ClothSimulationModel = ClothAsset->GetClothSimulationModel();
-			if (ClothSimulationModel && ClothSimulationModel->GetNumLods())
-			{
-				// Create the simulation proxy (note CreateClothSimulationProxy() can be overloaded)
-				CreateClothSimulationProxyImpl();
-			}
+			// Create the simulation proxy (note CreateClothSimulationProxy() can be overloaded)
+			CreateClothSimulationProxyImpl();
 		}
 	}
 	else
@@ -233,7 +237,6 @@ void UChaosClothComponent::OnRegister()
 	UpdateComponentSpaceTransforms();
 
 	// Fill up the property collection with the original cloth asset properties
-	ClothOutfitInteractor = NewObject<UChaosClothAssetInteractor>();
 	ResetConfigProperties();
 
 	// Create the proxy to start the simulation
@@ -256,10 +259,8 @@ void UChaosClothComponent::OnUnregister()
 	// Release cloth simulation
 	ClothSimulationProxy.Reset();
 
-	// Release the runtime simulation collection and facade
-	ClothOutfitInteractor->ResetProperties();
-	CollectionPropertyFacades.Empty();
-	PropertyCollections.Empty();
+	// Release the runtime simulation interactors, collections, and facades
+	ClothSimulationProperties.Reset();
 
 	// Unregister the dataflow simulation interface
 	UE::Dataflow::UnregisterSimulationInterface(this);
@@ -519,11 +520,17 @@ void UChaosClothComponent::GetAdditionalRequiredBonesForLeader(int32 LeaderLODIn
 	}
 
 	// Merge the physics asset bones (the leader's physics asset can be different to this component's cloth asset)
-	if (const UPhysicsAsset* const PhysicsAsset = GetAsset() ? GetAsset()->GetPhysicsAsset() : nullptr)
+	if (GetAsset())
 	{
-		if (const USkinnedAsset* const LeaderSkinnedAsset = ensure(LeaderPoseComponent.IsValid()) ? LeaderPoseComponent->GetSkinnedAsset() : nullptr)  // Needs the leader SkinnedAsset for the correct RefSkeleton
+		for (int32 ModelIndex = 0; ModelIndex < GetAsset()->GetNumClothSimulationModels(); ++ModelIndex)
 		{
-			USkinnedMeshComponent::GetPhysicsRequiredBones(LeaderSkinnedAsset, PhysicsAsset, RequiredBones);
+			if (const UPhysicsAsset* const PhysicsAsset = GetAsset()->GetPhysicsAssetForModel(ModelIndex))
+			{
+				if (const USkinnedAsset* const LeaderSkinnedAsset = ensure(LeaderPoseComponent.IsValid()) ? LeaderPoseComponent->GetSkinnedAsset() : nullptr)  // Needs the leader SkinnedAsset for the correct RefSkeleton
+				{
+					USkinnedMeshComponent::GetPhysicsRequiredBones(LeaderSkinnedAsset, PhysicsAsset, RequiredBones);
+				}
+			}
 		}
 	}
 
@@ -599,15 +606,20 @@ void UChaosClothComponent::SetSimulateInEditor(const bool bNewSimulateState)
 
 void UChaosClothComponent::StartNewParallelSimulation(float DeltaTime)
 {
+	using namespace ::Chaos::Softs;
 	if (ClothSimulationProxy.IsValid())
 	{
 		CSV_SCOPED_TIMING_STAT(Animation, Cloth);
 		const bool bIsSimulating = ClothSimulationProxy->Tick_GameThread(DeltaTime);
 		const int32 CurrentLOD = GetPredictedLODLevel();
 
-		if (bIsSimulating && CollectionPropertyFacades.IsValidIndex(CurrentLOD) && CollectionPropertyFacades[CurrentLOD].IsValid())
+		for (int32 ModelIndex = 0; ModelIndex < ClothSimulationProperties.Num(); ++ModelIndex)
 		{
-			CollectionPropertyFacades[CurrentLOD]->ClearDirtyFlags();
+			TArray<TSharedPtr<FCollectionPropertyFacade>>& CollectionPropertyFacades = ClothSimulationProperties[ModelIndex].CollectionPropertyFacades;
+			if (bIsSimulating && CollectionPropertyFacades.IsValidIndex(CurrentLOD) && CollectionPropertyFacades[CurrentLOD].IsValid())
+			{
+				CollectionPropertyFacades[CurrentLOD]->ClearDirtyFlags();
+			}
 		}
 	}
 }
@@ -664,10 +676,20 @@ void UChaosClothComponent::UpdateVisibility()
 	}
 }
 
-UChaosClothAssetInteractor* UChaosClothComponent::GetClothOutfitInteractor()
+UChaosClothAssetInteractor* UChaosClothComponent::GetClothOutfitInteractor(int32 ModelIndex, const FName ClothSimulationModelName)
 {
 	check(IsInGameThread());
-	return ClothOutfitInteractor;
+	if (ClothSimulationModelName != NAME_None && GetAsset() && GetAsset()->GetNumClothSimulationModels() == ClothSimulationProperties.Num())
+	{
+		for (int32 Index = 0; Index < ClothSimulationProperties.Num(); ++Index)
+		{
+			if (GetAsset()->GetClothSimulationModelName(Index) == ClothSimulationModelName)
+			{
+				return ClothSimulationProperties[Index].ClothOutfitInteractor;
+			}
+		}
+	}
+	return ClothSimulationProperties.IsValidIndex(ModelIndex) ? ClothSimulationProperties[ModelIndex].ClothOutfitInteractor : nullptr;
 }
 
 void UChaosClothComponent::BuildSimulationProxy()
@@ -691,15 +713,25 @@ void UChaosClothComponent::PreProcessSimulation(const float DeltaTime)
 
 void UChaosClothComponent::WriteToSimulation(const float DeltaTime, const bool bAsyncTask)
 {
+	using namespace ::Chaos::Softs;
 	if (ClothSimulationProxy.IsValid())
 	{
 		const bool bIsSimulating = ClothSimulationProxy->PreSimulate_GameThread(DeltaTime);
 
-		// Clear property dirty flags
-		const int32 CurrentLOD = GetPredictedLODLevel();
-		if (bIsSimulating && CollectionPropertyFacades.IsValidIndex(CurrentLOD) && CollectionPropertyFacades[CurrentLOD].IsValid())
+		if (bIsSimulating)
 		{
-			CollectionPropertyFacades[CurrentLOD]->ClearDirtyFlags();
+			// Clear property dirty flags
+			const int32 CurrentLOD = GetPredictedLODLevel();
+
+			for (int32 ModelIndex = 0; ModelIndex < ClothSimulationProperties.Num(); ++ModelIndex)
+			{
+				TArray<TSharedPtr<FCollectionPropertyFacade>>& CollectionPropertyFacades = ClothSimulationProperties[ModelIndex].CollectionPropertyFacades;
+
+				if (CollectionPropertyFacades.IsValidIndex(CurrentLOD) && CollectionPropertyFacades[CurrentLOD].IsValid())
+				{
+					CollectionPropertyFacades[CurrentLOD]->ClearDirtyFlags();
+				}
+			}
 		}
 	}
 }
diff --git a/Engine/Plugins/ChaosClothAsset/Source/ChaosClothAssetEngine/Private/ChaosClothAsset/ClothSimulationContext.cpp b/Engine/Plugins/ChaosClothAsset/Source/ChaosClothAssetEngine/Private/ChaosClothAsset/ClothSimulationContext.cpp
--- a/Engine/Plugins/ChaosClothAsset/Source/ChaosClothAssetEngine/Private/ChaosClothAsset/ClothSimulationContext.cpp
+++ b/Engine/Plugins/ChaosClothAsset/Source/ChaosClothAssetEngine/Private/ChaosClothAsset/ClothSimulationContext.cpp
@@ -38,8 +38,8 @@ namespace UE::Chaos::ClothAsset
 		SolverGeometryScale = ClothComponent.GetClothGeometryScale();;
 
 		// Update bone transforms
-		const UChaosClothAsset* const ClothAsset = Cast<const UChaosClothAsset>(ClothComponent.GetAsset());
-		const FReferenceSkeleton* const ReferenceSkeleton = ClothAsset ? &ClothAsset->GetRefSkeleton() : nullptr;
+		const UChaosClothAssetBase* const Asset = ClothComponent.GetAsset();
+		const FReferenceSkeleton* const ReferenceSkeleton = Asset ? &Asset->GetRefSkeleton() : nullptr;
 		int32 NumBones = ReferenceSkeleton ? ReferenceSkeleton->GetNum() : 0;
 
 		if (USkinnedMeshComponent* const LeaderComponent = ClothComponent.LeaderPoseComponent.Get())
@@ -102,14 +102,15 @@ namespace UE::Chaos::ClothAsset
 		bool bSetRefToLocalMatricesToIdentity = true;
 		if (!bIsInitialization)
 		{
-			if (const FSkeletalMeshRenderData* const RenderData = ClothAsset ? ClothAsset->GetResourceForRendering() : nullptr)
+			if (const FSkeletalMeshRenderData* const RenderData = Asset ? Asset->GetResourceForRendering() : nullptr)
 			{
 				if (RenderData->LODRenderData.IsValidIndex(LodIndex))
 				{
 					const TArray<uint16>* RequiredExtraBones = nullptr;
-					if (TSharedPtr<const FChaosClothSimulationModel > ClothModel = ClothAsset ? ClothAsset->GetClothSimulationModel() : TSharedPtr<const FChaosClothSimulationModel >(nullptr))
+					for (int32 ModelIndex = 0; ModelIndex < Asset->GetNumClothSimulationModels(); ++ModelIndex)
 					{
-						if (ClothModel->IsValidLodIndex(LodIndex))
+						const TSharedPtr<const FChaosClothSimulationModel>& ClothModel = Asset->GetClothSimulationModel(ModelIndex);
+						if (ClothModel && ClothModel->IsValidLodIndex(LodIndex))
 						{
 							RequiredExtraBones = &ClothModel->ClothSimulationLodModels[LodIndex].RequiredExtraBoneIndices;
 						}
diff --git a/Engine/Plugins/ChaosClothAsset/Source/ChaosClothAssetEngine/Private/ChaosClothAsset/ClothSimulationModel.cpp b/Engine/Plugins/ChaosClothAsset/Source/ChaosClothAssetEngine/Private/ChaosClothAsset/ClothSimulationModel.cpp
--- a/Engine/Plugins/ChaosClothAsset/Source/ChaosClothAssetEngine/Private/ChaosClothAsset/ClothSimulationModel.cpp
+++ b/Engine/Plugins/ChaosClothAsset/Source/ChaosClothAssetEngine/Private/ChaosClothAsset/ClothSimulationModel.cpp
@@ -489,3 +489,18 @@ void FChaosClothSimulationModel::GetResourceSizeEx(FResourceSizeEx& CumulativeRe
 
 	CumulativeResourceSize.AddDedicatedSystemMemoryBytes(Ar.TotalSize());
 }
+
+CHAOSCLOTHASSETENGINE_API void FChaosClothSimulationModel::RemapBoneIndices(const TArray<int32>& BoneMap)
+{
+	// Default to the root bone index whenever an invalid index is provided to avoid any issues further on
+	constexpr int32 RootBoneIndex = 0;
+
+	// Remap the reference bone if needed
+	ReferenceBoneIndex = ensure(BoneMap.IsValidIndex(ReferenceBoneIndex)) ? BoneMap[ReferenceBoneIndex] : RootBoneIndex;
+
+	// UsedBoneIndices is used to reindex from the subset index to the actual skeleton bone index, so needs updating
+	for (int32& UsedBoneIndex : UsedBoneIndices)
+	{
+		UsedBoneIndex = ensure(BoneMap.IsValidIndex(UsedBoneIndex)) ? BoneMap[UsedBoneIndex] : RootBoneIndex;
+	}
+}
diff --git a/Engine/Plugins/ChaosClothAsset/Source/ChaosClothAssetEngine/Private/ChaosClothAsset/ClothSimulationModel.h b/Engine/Plugins/ChaosClothAsset/Source/ChaosClothAssetEngine/Private/ChaosClothAsset/ClothSimulationModel.h
deleted file mode 100644
diff --git a/Engine/Plugins/ChaosClothAsset/Source/ChaosClothAssetEngine/Private/ChaosClothAsset/ClothSimulationProxy.cpp b/Engine/Plugins/ChaosClothAsset/Source/ChaosClothAssetEngine/Private/ChaosClothAsset/ClothSimulationProxy.cpp
--- a/Engine/Plugins/ChaosClothAsset/Source/ChaosClothAssetEngine/Private/ChaosClothAsset/ClothSimulationProxy.cpp
+++ b/Engine/Plugins/ChaosClothAsset/Source/ChaosClothAssetEngine/Private/ChaosClothAsset/ClothSimulationProxy.cpp
@@ -156,8 +156,14 @@ namespace UE::Chaos::ClothAsset
 
 		check(IsInGameThread());
 
+		// Reset all simulation arrays
+		Configs.Reset();
+		Meshes.Reset();
+		Colliders.Reset();
+		Cloths.Reset();
+
 		// Create solver config simulation thread object first. Need to know which solver type we're creating.
-		const int32 SolverConfigIndex = Configs.Emplace(MakeUnique<FClothingSimulationConfig>(ClothComponent.GetPropertyCollections(0)));  // TODO: Use a separate solver config for outfits
+		const int32 SolverConfigIndex = Configs.Emplace(MakeUnique<FClothingSimulationConfig>(ClothComponent.GetSolverPropertyCollections()));
 
 		// Use new SoftsEvolution, not PBDEvolution.
 		constexpr bool bUseLegacySolver = false;
@@ -179,38 +185,39 @@ PRAGMA_ENABLE_DEPRECATION_WARNINGS
 		Solver->SetLocalSpaceRotation((FQuat)ClothSimulationContext->ComponentTransform.GetRotation());
 
 		// Create mesh simulation thread object
-		UChaosClothAsset* const ClothAsset = Cast<UChaosClothAsset>(ClothComponent.GetAsset());
-		if (!ClothAsset)
+		const UChaosClothAssetBase* const Asset = ClothComponent.GetAsset();
+		if (!Asset)
 		{
 			return;
 		}
-		ClothSimulationModel = ClothAsset->GetClothSimulationModel();
-		check(ClothSimulationModel.IsValid());
+		const FReferenceSkeleton* const ReferenceSkeleton = &Asset->GetRefSkeleton();
+		
+		for (int32 ModelIndex = 0; ModelIndex < Asset->GetNumClothSimulationModels(); ++ModelIndex)
+		{
+			const TSharedPtr<const FChaosClothSimulationModel>& ClothSimulationModel = Asset->GetClothSimulationModel(ModelIndex);
 
-		FString DebugName;
+			FString DebugName;
 #if !UE_BUILD_SHIPPING
-		DebugName = ClothComponent.GetOwner() ?
-			FString::Format(TEXT("{0}|{1}"), { ClothComponent.GetOwner()->GetActorNameOrLabel(), ClothComponent.GetName() }) :
-			ClothComponent.GetName();
+			DebugName = ClothComponent.GetOwner() ?
+				FString::Format(TEXT("{0}|{1}"), { ClothComponent.GetOwner()->GetActorNameOrLabel(), ClothComponent.GetName() }) :
+				ClothComponent.GetName();
 #endif
-		const int32 MeshIndex = Meshes.Emplace(MakeUnique<FClothSimulationMesh>(*ClothSimulationModel, *ClothSimulationContext, ClothAsset->GetClothCollections(), DebugName));
+			const int32 MeshIndex = Meshes.Emplace(MakeUnique<FClothSimulationMesh>(*ClothSimulationModel, *ClothSimulationContext, Asset->GetCollections(ModelIndex), DebugName));
 
-		// Create collider simulation thread object
-		const FReferenceSkeleton* const ReferenceSkeleton = &ClothAsset->GetRefSkeleton();
-		const int32 ColliderIndex = Colliders.Emplace(MakeUnique<FClothingSimulationCollider>(ClothComponent.GetPhysicsAsset(), ReferenceSkeleton));
-		Colliders[ColliderIndex]->SetCollisionData(&CollisionSourcesProxy->GetCollisionData());
+			// Create collider simulation thread object
+			const int32 ColliderIndex = Colliders.Emplace(MakeUnique<FClothingSimulationCollider>(Asset->GetPhysicsAssetForModel(ModelIndex), ReferenceSkeleton));
+			Colliders[ColliderIndex]->SetCollisionData(&CollisionSourcesProxy->GetCollisionData());
 
-		// Create cloth config simulation thread object
-		const int32 ClothConfigIndex = Configs.Emplace(MakeUnique<FClothingSimulationConfig>(ClothComponent.GetPropertyCollections(0)));
+			// Create cloth config simulation thread object
+			const int32 ClothConfigIndex = Configs.Emplace(MakeUnique<FClothingSimulationConfig>(ClothComponent.GetPropertyCollections(ModelIndex)));
 
-		// Create cloth simulation thread object
-		constexpr uint32 GroupId = 0;
-
-		const int32 ClothIndex = Cloths.Emplace(MakeUnique<FClothingSimulationCloth>(
-			Configs[ClothConfigIndex].Get(),
-			Meshes[MeshIndex].Get(),
-			ColliderIndex != INDEX_NONE ? TArray<FClothingSimulationCollider*>({ Colliders[ColliderIndex].Get() }) : TArray<FClothingSimulationCollider*>(),
-			GroupId));
+			// Create cloth simulation thread object
+			const int32 ClothIndex = Cloths.Emplace(MakeUnique<FClothingSimulationCloth>(
+				Configs[ClothConfigIndex].Get(),
+				Meshes[MeshIndex].Get(),
+				ColliderIndex != INDEX_NONE ? TArray<FClothingSimulationCollider*>({ Colliders[ColliderIndex].Get() }) : TArray<FClothingSimulationCollider*>(),
+				ModelIndex));
+		}
 	}
 
 	void FClothSimulationProxy::ExecuteInitialization()
@@ -332,11 +339,11 @@ PRAGMA_ENABLE_DEPRECATION_WARNINGS
 
 		// Check that the render mesh current LOD isn't just fully skinned
 		const int32 LodIndex = ClothSimulationContext->LodIndex;
-		const UChaosClothAsset* const ClothAsset = Cast<const UChaosClothAsset>(ClothComponent.GetAsset());
-		if (!ClothAsset ||
-			!ClothAsset->GetResourceForRendering() ||
-			!ClothAsset->GetResourceForRendering()->LODRenderData.IsValidIndex(LodIndex) ||
-			!ClothAsset->GetResourceForRendering()->LODRenderData[LodIndex].HasClothData())
+		const UChaosClothAssetBase* const Asset = ClothComponent.GetAsset();
+		if (!Asset ||
+			!Asset->GetResourceForRendering() ||
+			!Asset->GetResourceForRendering()->LODRenderData.IsValidIndex(LodIndex) ||
+			!Asset->GetResourceForRendering()->LODRenderData[LodIndex].HasClothData())
 		{
 			return false;
 		}
@@ -799,11 +806,13 @@ PRAGMA_ENABLE_DEPRECATION_WARNINGS
 	void FClothSimulationProxy::InitializeConfigs()
 	{
 		// Replace physics thread's configs with the game thread's configs
-		for (const TUniquePtr<::Chaos::FClothingSimulationConfig>& Config : Configs)
+		for (int32 ConfigIndex = 0; ConfigIndex < Configs.Num(); ++ConfigIndex)
 		{
-			Config->Initialize(ClothComponent.GetPropertyCollections(0));  // TODO: Outfit and multi-cloths/solver config update
+			const int32 ModelIndex = ConfigIndex - 1;
+			Configs[ConfigIndex]->Initialize(ModelIndex < 0 ?
+				ClothComponent.GetSolverPropertyCollections() :  // Config 0 is the solver config
+				ClothComponent.GetPropertyCollections(ModelIndex));
 		}
-		// TODO: separate solver config lod from cloth component
 		Solver->SetSolverLOD(ClothSimulationContext->LodIndex);
 	}
 
diff --git a/Engine/Plugins/ChaosClothAsset/Source/ChaosClothAssetEngine/Public/ChaosClothAsset/ClothAsset.h b/Engine/Plugins/ChaosClothAsset/Source/ChaosClothAssetEngine/Public/ChaosClothAsset/ClothAsset.h
--- a/Engine/Plugins/ChaosClothAsset/Source/ChaosClothAssetEngine/Public/ChaosClothAsset/ClothAsset.h
+++ b/Engine/Plugins/ChaosClothAsset/Source/ChaosClothAssetEngine/Public/ChaosClothAsset/ClothAsset.h
@@ -13,11 +13,14 @@ namespace Dataflow = UE::Dataflow;
 #else
 namespace UE_DEPRECATED(5.5, "Use UE::Dataflow instead.") Dataflow {}
 #endif
+#if UE_ENABLE_INCLUDE_ORDER_DEPRECATED_IN_5_6
+class UAnimationAsset;
+class USkeletalMesh;
+#endif
 
 class FSkeletalMeshModel;
 struct FSkeletalMeshLODInfo;
 struct FChaosClothAssetLodTransitionDataCache;
-class UAnimationAsset;
 
 UENUM()
 enum class EClothAssetAsyncProperties : uint64
@@ -54,19 +57,7 @@ public:
 	//~ End UObject interface
 
 	//~ Begin USkinnedAsset interface
-	UFUNCTION(BlueprintGetter)
-	virtual class UPhysicsAsset* GetShadowPhysicsAsset() const override			{ return ShadowPhysicsAsset; }
-	virtual FMatrix GetComposedRefPoseMatrix(FName InBoneName) const override;
-	virtual FMatrix GetComposedRefPoseMatrix(int32 InBoneIndex) const override	{ return CachedComposedRefPoseMatrices[InBoneIndex]; }
-	virtual const FMeshUVChannelInfo* GetUVChannelData(int32 MaterialIndex) const override;
-	virtual bool GetSupportRayTracing() const override							{ return bSupportRayTracing; }
-	virtual int32 GetRayTracingMinLOD() const override							{ return RayTracingMinLOD; }
-	virtual TArray<FMatrix44f>& GetRefBasesInvMatrix() override					{ return RefBasesInvMatrix; }
-	virtual const TArray<FMatrix44f>& GetRefBasesInvMatrix() const override		{ return RefBasesInvMatrix; }
 	virtual UPhysicsAsset* GetPhysicsAsset() const								{ return PhysicsAsset; }
-	virtual TArray<FSkeletalMaterial>& GetMaterials() override					{ return Materials; }
-	virtual const TArray<FSkeletalMaterial>& GetMaterials() const override		{ return Materials; }
-	virtual bool IsMaterialUsed(int32 MaterialIndex) const override				{ return true; }
 	virtual USkeleton* GetSkeleton() override									{ return Skeleton; }  // Note: The USkeleton isn't a reliable source of reference skeleton
 	virtual const USkeleton* GetSkeleton() const override						{ return Skeleton; }
 	virtual void SetSkeleton(USkeleton* InSkeleton) override					{ Skeleton = InSkeleton; }
@@ -76,9 +67,6 @@ public:
 	virtual void BuildLODModel(const ITargetPlatform* TargetPlatform, int32 LODIndex) override;
 	virtual FString BuildDerivedDataKey(const ITargetPlatform* TargetPlatform) override;
 	virtual bool IsInitialBuildDone() const override;
-	virtual bool GetEnableLODStreaming(const class ITargetPlatform* TargetPlatform) const override	{ return false; }
-	virtual int32 GetMaxNumStreamedLODs(const class ITargetPlatform* TargetPlatform) const override	{ return 0; }
-	virtual int32 GetMaxNumOptionalLODs(const ITargetPlatform* TargetPlatform) const override { return 0; }
 #endif
 #if WITH_EDITORONLY_DATA
 	virtual FSkeletalMeshModel* GetImportedModel() const override
@@ -210,9 +198,6 @@ private:
 	virtual FString GetAsyncPropertyName(uint64 Property) const override;
 	//~ End USkinnedAsset interface
 
-	/** Pre-calculate refpose-to-local transforms. */
-	void CalculateInvRefMatrices();
-
 	/** Re-calculate the bounds for this asset. */
 	void CalculateBounds();
 
@@ -224,20 +209,16 @@ private:
 	/** Build the clothing simulation meshes from the Cloth Collection. */
 	void BuildClothSimulationModel(TArray<FChaosClothAssetLodTransitionDataCache>* InOutTransitionCache = nullptr);
 
-	/** List of materials for this cloth asset. Set by the Dataflow evaluation. */
-	UPROPERTY(EditAnywhere, Category = Materials, Meta = (EditCondition = "DataflowAsset == nullptr"))
-	TArray<FSkeletalMaterial> Materials;
-
 	/**
 	 * Skeleton asset used at creation time.
 	 * This is of limited use since this USkeleton's reference skeleton might not necessarily match the one created for this asset.
 	 * Set by the Dataflow evaluation.
 	 */
-	UPROPERTY(EditAnywhere, Setter = SetSkeleton, Category = Skeleton, Meta = (EditCondition = "DataflowAsset == nullptr"))
+	UPROPERTY(EditAnywhere, Setter = SetSkeleton, Category = Skeleton, Meta = (EditCondition = "DataflowInstance.DataflowAsset == nullptr"))
 	TObjectPtr<USkeleton> Skeleton;
 
 	/** Physics asset used for collision. Set by the Dataflow evaluation. */
-	UPROPERTY(EditAnywhere, Category = Collision, Meta = (EditCondition = "DataflowAsset == nullptr"))
+	UPROPERTY(EditAnywhere, Category = Collision, Meta = (EditCondition = "DataflowInstance.DataflowAsset == nullptr"))
 	TObjectPtr<UPhysicsAsset> PhysicsAsset;
 
 	/** Whether to blend positions between the skinned/simulated transitions of the cloth render mesh. Only used when not overriden in Dataflow by the ProxyDefomer node. */
@@ -255,27 +236,14 @@ private:
 	UPROPERTY()
 	float SkinningKernelRadius_DEPRECATED = 30.f;
 
-	/**
-	 * Physics asset whose shapes will be used for shadowing when components have bCastCharacterCapsuleDirectShadow or bCastCharacterCapsuleIndirectShadow enabled.
-	 * Only spheres and sphyl shapes in the physics asset can be supported.  The more shapes used, the higher the cost of the capsule shadows will be.
-	 */
-	UPROPERTY(EditAnywhere, AssetRegistrySearchable, BlueprintGetter = GetShadowPhysicsAsset, Category = Lighting)
-	TObjectPtr<class UPhysicsAsset> ShadowPhysicsAsset;
-
 	/** A unique identifier as used by the section rendering code. */
 	UPROPERTY()
 	FGuid AssetGuid;
 
-	/** Mesh-space ref pose, where parent matrices are applied to ref pose matrices. */
-	TArray<FMatrix> CachedComposedRefPoseMatrices;
-
 	/** Cloth Collection containing this asset data. One per LOD. */
 	UE_DEPRECATED(5.4, "This must be protected for async build, always use the accessors even internally.")
 	TArray<TSharedRef<FManagedArrayCollection>> ClothCollections;
 
-	/** Reference skeleton precomputed bases. */
-	TArray<FMatrix44f> RefBasesInvMatrix;
-
 #if WITH_EDITORONLY_DATA
 	/** Source mesh geometry information (not used at runtime). */
 	UE_DEPRECATED(5.4, "This must be protected for async build, always use the accessors even internally.")
diff --git a/Engine/Plugins/ChaosClothAsset/Source/ChaosClothAssetEngine/Public/ChaosClothAsset/ClothAssetBase.h b/Engine/Plugins/ChaosClothAsset/Source/ChaosClothAssetEngine/Public/ChaosClothAsset/ClothAssetBase.h
--- a/Engine/Plugins/ChaosClothAsset/Source/ChaosClothAssetEngine/Public/ChaosClothAsset/ClothAssetBase.h
+++ b/Engine/Plugins/ChaosClothAsset/Source/ChaosClothAssetEngine/Public/ChaosClothAsset/ClothAssetBase.h
@@ -9,7 +9,6 @@
 
 class FSkeletalMeshRenderData;
 class FSkinnedAssetCompilationContext;
-class USkeletalMesh;
 class UDataflow;
 struct FChaosClothSimulationModel;
 struct FManagedArrayCollection;
@@ -75,14 +74,44 @@ public:
 	//~ Begin UObject interface
 	virtual void BeginDestroy() override;
 	virtual bool IsReadyForFinishDestroy() override;
+	virtual void Serialize(FArchive& Ar) override;
 	//~ End UObject interface
 
 	//~ Begin USkinnedAsset interface
-	virtual FSkeletalMeshLODInfo* GetLODInfo(int32 Index);
-	virtual const FSkeletalMeshLODInfo* GetLODInfo(int32 Index) const;
 	virtual FReferenceSkeleton& GetRefSkeleton() override;
 	virtual const FReferenceSkeleton& GetRefSkeleton() const override;
-	virtual FSkeletalMeshRenderData* GetResourceForRendering() const override;
+	virtual FSkeletalMeshLODInfo* GetLODInfo(int32 Index);
+	virtual const FSkeletalMeshLODInfo* GetLODInfo(int32 Index) const;
+	UFUNCTION(BlueprintGetter)
+	virtual class UPhysicsAsset* GetShadowPhysicsAsset() const override
+	{
+		return ShadowPhysicsAsset;
+	}
+	virtual FMatrix GetComposedRefPoseMatrix(FName BoneName) const override
+	{
+		return FMatrix::Identity;
+	}
+	virtual FMatrix GetComposedRefPoseMatrix(int32 BoneIndex) const override
+	{
+		return FMatrix::Identity;
+	}
+	virtual const FMeshUVChannelInfo* GetUVChannelData(int32 MaterialIndex) const override;
+	virtual bool GetSupportRayTracing() const override
+	{
+		return bSupportRayTracing;
+	}
+	virtual int32 GetRayTracingMinLOD() const override
+	{
+		return RayTracingMinLOD;
+	}
+	virtual TArray<FMatrix44f>& GetRefBasesInvMatrix() override
+	{
+		return RefBasesInvMatrix;
+	}
+	virtual const TArray<FMatrix44f>& GetRefBasesInvMatrix() const override
+	{
+		return RefBasesInvMatrix;
+	}
 	virtual TArray<FSkeletalMeshLODInfo>& GetLODInfoArray() override
 	{
 		return LODInfo;
@@ -91,6 +120,7 @@ public:
 	{
 		return LODInfo;
 	}
+	virtual FSkeletalMeshRenderData* GetResourceForRendering() const override;
 	virtual int32 GetDefaultMinLod() const
 	{
 		return 0;
@@ -99,6 +129,18 @@ public:
 	{
 		return MinLod;
 	}
+	virtual TArray<FSkeletalMaterial>& GetMaterials() override
+	{
+		return Materials;
+	}
+	virtual const TArray<FSkeletalMaterial>& GetMaterials() const override
+	{
+		return Materials;
+	}
+	virtual bool IsMaterialUsed(int32 MaterialIndex) const override
+	{
+		return Materials.IsValidIndex(MaterialIndex);
+	}
 	virtual int32 GetLODNum() const override
 	{
 		return LODInfo.Num();
@@ -140,8 +182,26 @@ public:
 	}
 	virtual bool GetHasVertexColors() const override;
 	virtual int32 GetPlatformMinLODIdx(const ITargetPlatform* TargetPlatform) const override;
-	virtual const FPerPlatformBool& GetDisableBelowMinLodStripping() const override { return DisableBelowMinLodStripping; }
+	virtual const FPerPlatformBool& GetDisableBelowMinLodStripping() const override
+	{
+		return DisableBelowMinLodStripping;
+	}
 	virtual bool IsMinLodQualityLevelEnable() const override;
+
+#if WITH_EDITOR
+	virtual bool GetEnableLODStreaming(const class ITargetPlatform* TargetPlatform) const override
+	{
+		return false;
+	}
+	virtual int32 GetMaxNumStreamedLODs(const class ITargetPlatform* TargetPlatform) const override
+	{
+		return 0;
+	}
+	virtual int32 GetMaxNumOptionalLODs(const ITargetPlatform* TargetPlatform) const override
+	{
+		return 0;
+	}
+#endif
 	//~ End USkinnedAsset interface
 
 	const FPerQualityLevelInt& GetQualityLevelMinLod() const
@@ -182,11 +242,11 @@ public:
 	}
 
 #if WITH_EDITORONLY_DATA
-	void SetPreviewSceneSkeletalMesh(USkeletalMesh* Mesh);
-	USkeletalMesh* GetPreviewSceneSkeletalMesh() const;
+	void SetPreviewSceneSkeletalMesh(class USkeletalMesh* Mesh);
+	class USkeletalMesh* GetPreviewSceneSkeletalMesh() const;
 
-	void SetPreviewSceneAnimation(UAnimationAsset* Animation);
-	UAnimationAsset* GetPreviewSceneAnimation() const;
+	void SetPreviewSceneAnimation(class UAnimationAsset* Animation);
+	class UAnimationAsset* GetPreviewSceneAnimation() const;
 #endif
 
 	void SetHasVertexColors(bool InbHasVertexColors);
@@ -215,6 +275,9 @@ protected:
 	/** Safely release the render data. */
 	void ReleaseResources();
 
+	/** Pre-calculate refpose-to-local transforms. */
+	void CalculateInvRefMatrices();
+
 #if WITH_EDITOR
 	/** Initial step for the building process - Can't be done in parallel. USkinnedAsset Interface. */
 	virtual void BeginBuildInternal(FSkinnedAssetBuildContext& Context) override;
@@ -226,16 +289,12 @@ protected:
 	virtual void FinishBuildInternal(FSkinnedAssetBuildContext& Context) override;
 #endif
 
-	UE_DEPRECATED(5.6, "This must be protected for async build, always use the accessors even internally.")
-	UPROPERTY()
-	uint8 bHasVertexColors : 1 = false;
-
 	UPROPERTY(EditAnywhere, Category = "Dataflow")
 	FDataflowInstance DataflowInstance;
 
-	/** Reference skeleton created from the provided skeleton asset. */
-	UE_DEPRECATED(5.4, "This must be protected for async build, always use the accessors even internally.")
-	FReferenceSkeleton RefSkeleton;
+	/** List of materials for this cloth asset. Set by the Dataflow evaluation. */
+	UPROPERTY(EditAnywhere, Category = Materials, Meta = (EditCondition = "DataflowInstance.DataflowAsset == nullptr"))
+	TArray<FSkeletalMaterial> Materials;
 
 	/** Struct containing information for each LOD level, such as materials to use, and when use the LOD. Not currently editable or customizable through the Dataflow. */
 	UPROPERTY(VisibleAnywhere, EditFixedSize, Category = LevelOfDetails)
@@ -260,6 +319,13 @@ protected:
 	UPROPERTY(EditAnywhere, Category = RayTracing)
 	int32 RayTracingMinLOD;
 
+	/**
+	 * Physics asset whose shapes will be used for shadowing when components have bCastCharacterCapsuleDirectShadow or bCastCharacterCapsuleIndirectShadow enabled.
+	 * Only spheres and sphyl shapes in the physics asset can be supported.  The more shapes used, the higher the cost of the capsule shadows will be.
+	 */
+	UPROPERTY(EditAnywhere, AssetRegistrySearchable, BlueprintGetter = GetShadowPhysicsAsset, Category = Lighting)
+	TObjectPtr<class UPhysicsAsset> ShadowPhysicsAsset;
+
 	/** Default translucent material to blend on top of this mesh. Mesh will be rendered twice - once with a base material and once with overlay material */
 	UE_DEPRECATED(5.6, "This must be protected for async build, always use the accessors even internally.")
 	UPROPERTY(EditAnywhere, Category = Rendering)
@@ -270,6 +336,13 @@ protected:
 	UPROPERTY(EditAnywhere, Category = Rendering)
 	float OverlayMaterialMaxDrawDistance;
 
+	UE_DEPRECATED(5.6, "This must be protected for async build, always use the accessors even internally.")
+	UPROPERTY()
+	uint8 bHasVertexColors : 1 = false;
+
+	/** Reference skeleton created from the provided skeleton asset. */
+	UE_DEPRECATED(5.4, "This must be protected for async build, always use the accessors even internally.")
+	FReferenceSkeleton RefSkeleton;
 
 	/** Bounds for this asset. */
 	FBoxSphereBounds Bounds;
@@ -297,11 +370,11 @@ private:
 
 	/** Optional Skeletal Mesh that the cloth asset is attached to in the Preview Scene in the Cloth Editor */
 	UPROPERTY(DuplicateTransient, AssetRegistrySearchable)
-	TSoftObjectPtr<USkeletalMesh> PreviewSceneSkeletalMesh;
+	TSoftObjectPtr<class USkeletalMesh> PreviewSceneSkeletalMesh;
 
 	/** Optional animation attached to PreviewSceneSkeletalMesh in the Preview Scene in the Cloth Editor */
 	UPROPERTY(DuplicateTransient, AssetRegistrySearchable)
-	TSoftObjectPtr<UAnimationAsset> PreviewSceneAnimation;
+	TSoftObjectPtr<class UAnimationAsset> PreviewSceneAnimation;
 #endif
 
 	/** Rendering data. */
@@ -309,4 +382,10 @@ private:
 
 	/** A fence which is used to keep track of the rendering thread releasing the static mesh resources. */
 	FRenderCommandFence ReleaseResourcesFence;
+
+	/** Reference skeleton precomputed bases. */
+	TArray<FMatrix44f> RefBasesInvMatrix;
+
+	/** Mesh-space ref pose, where parent matrices are applied to ref pose matrices. */
+	TArray<FMatrix> CachedComposedRefPoseMatrices;
 };
diff --git a/Engine/Plugins/ChaosClothAsset/Source/ChaosClothAssetEngine/Public/ChaosClothAsset/ClothComponent.h b/Engine/Plugins/ChaosClothAsset/Source/ChaosClothAssetEngine/Public/ChaosClothAsset/ClothComponent.h
--- a/Engine/Plugins/ChaosClothAsset/Source/ChaosClothAssetEngine/Public/ChaosClothAsset/ClothComponent.h
+++ b/Engine/Plugins/ChaosClothAsset/Source/ChaosClothAssetEngine/Public/ChaosClothAsset/ClothComponent.h
@@ -30,6 +30,24 @@ namespace UE::Chaos::ClothAsset
 	class FCollisionSources;
 }
 
+/**
+ * Private structure that contains all simulation properties runtime elements.
+ * These get created per cloth asset/outfit piece when the component is registered.
+ */
+USTRUCT()
+struct FChaosClothSimulationProperties final
+{
+	GENERATED_BODY()
+private:
+	friend UChaosClothComponent;
+	UPROPERTY(Transient)
+	TObjectPtr<UChaosClothAssetInteractor> ClothOutfitInteractor;
+	TArray<TSharedPtr<FManagedArrayCollection>> PropertyCollections;
+	TArray<TSharedPtr<::Chaos::Softs::FCollectionPropertyFacade>> CollectionPropertyFacades;
+
+	void Initialize(const TArray<TSharedRef<FManagedArrayCollection>>& AssetPropertyCollections);
+};
+
 /**
  * Cloth simulation component.
  */
@@ -110,10 +128,16 @@ public:
 	UFUNCTION(CallInEditor, BlueprintCallable, Category = "ClothComponent", Meta = (DisplayName = "Hard Reset Simulation", Keywords = "Chaos Cloth Recreate Simulation Proxy"))
 	void RecreateClothSimulationProxy();
 
-	/** Get the current interactor for the cloth outfit associated with this cloth component. 
-	 * Interact with solver-level properties as well as all cloth assets within the cloth outfit (once multi-asset outfits exist).*/
+	/**
+	 * Find the current interactor for the cloth outfit associated with this cloth component.
+	 * The default parameter values will always find the interactor when the cloth component is using a cloth asset.
+	 * When using an outfit asset, an interactor from each cloth simulation models can be chosen by specifying either the model index or the model name.
+	 * The solver properties interactor can also be obtained on model index 0 as it is shared across all models, but only set on the first one.
+	 * @param ModelIndex The index of the outfit simulation model to retrieve the interactor for. Can be ignored or used as a fallback when seraching by name. Must be 0 for the solver interactor.
+	 * @param ClothSimulationModelName The name of the outfit model to retrieve the interactor for. Leave it empty to use the ModelIndex instead of when retieving the solver interactor.
+	 */
 	UFUNCTION(BlueprintCallable, Category = "ClothComponent")
-	UChaosClothAssetInteractor* GetClothOutfitInteractor();
+	UChaosClothAssetInteractor* GetClothOutfitInteractor(int32 ModelIndex = 0, const FName ClothSimulationModelName = NAME_None);
 
 	/**
 	 * Add a collision source for the cloth on this component.
@@ -160,14 +184,19 @@ public:
 	 * This might be different from the cloth asset's since the component's properties can be modified in code or in blueprints.
 	 * This could also be different from the cloth simulation object until the cloth simulation thread synchronise the properties.
 	 */
-	const TArray<TSharedPtr<FManagedArrayCollection>>& GetPropertyCollections(int32 /*ModelIndex*/) const
+	const TArray<TSharedPtr<FManagedArrayCollection>>& GetPropertyCollections(int32 ModelIndex) const
 	{
-		return PropertyCollections;
+		return ClothSimulationProperties[ModelIndex].PropertyCollections;
+	}
+	/** Return the solver property collections. */
+	const TArray<TSharedPtr<FManagedArrayCollection>>& GetSolverPropertyCollections() const
+	{
+		return ClothSimulationProperties[0].PropertyCollections;
 	}
 	UE_DEPRECATED(5.6, "Use GetPropertyCollections(int32 ModelIndex) instead.")
 	const TArray<TSharedPtr<const FManagedArrayCollection>>& GetPropertyCollections() const
 	{
-		return reinterpret_cast<const TArray<TSharedPtr<const FManagedArrayCollection>>&>(PropertyCollections);
+		return reinterpret_cast<const TArray<TSharedPtr<const FManagedArrayCollection>>&>(ClothSimulationProperties[0].PropertyCollections);
 	}
 
 	const UE::Chaos::ClothAsset::FClothSimulationProxy* GetClothSimulationProxy() const { return ClothSimulationProxy.Get(); }
@@ -308,11 +337,12 @@ private:
 	TObjectPtr<UChaosClothAssetBase> Asset;
 #endif
 
+	/**
+	 * Simulation properties per cloth model. 
+	 * Transient because it doesn't need to be serialized but contains pointers to UObject that needs not be garbage collected.
+	 */
 	UPROPERTY(Transient)
-	TObjectPtr<UChaosClothAssetInteractor> ClothOutfitInteractor;
-
-	TArray<TSharedPtr<FManagedArrayCollection>> PropertyCollections;
-	TArray<TSharedPtr<::Chaos::Softs::FCollectionPropertyFacade>> CollectionPropertyFacades;
+	TArray<FChaosClothSimulationProperties> ClothSimulationProperties;
 
 	TSharedPtr<UE::Chaos::ClothAsset::FClothSimulationProxy> ClothSimulationProxy;
 
diff --git a/Engine/Plugins/ChaosClothAsset/Source/ChaosClothAssetEngine/Public/ChaosClothAsset/ClothSimulationModel.h b/Engine/Plugins/ChaosClothAsset/Source/ChaosClothAssetEngine/Public/ChaosClothAsset/ClothSimulationModel.h
new file mode 100644
--- /dev/null
+++ b/Engine/Plugins/ChaosClothAsset/Source/ChaosClothAssetEngine/Public/ChaosClothAsset/ClothSimulationModel.h
@@ -0,0 +1,142 @@
+// Copyright Epic Games, Inc. All Rights Reserved.
+
+#pragma once
+
+#include "Containers/Array.h"
+#include "Math/MathFwd.h"
+#include "Templates/SharedPointer.h"
+#include "ClothVertBoneData.h"
+#include "ClothTetherData.h"
+#include "SkeletalMeshTypes.h"
+#include "ClothSimulationModel.generated.h"
+
+struct FReferenceSkeleton;
+struct FManagedArrayCollection;
+struct FChaosClothAssetLodTransitionDataCache;
+
+/**
+ * Cloth simulation LOD model.
+ * Contains the LOD data as fed to the solver for constraints creation.
+ */
+USTRUCT()
+struct FChaosClothSimulationLodModel
+{
+	GENERATED_BODY()
+
+	/** Vertex positions. */
+	UPROPERTY()
+	TArray<FVector3f> Positions;
+	
+	/** Vertex Normals. */
+	UPROPERTY()
+	TArray<FVector3f> Normals;
+
+	/** Triangles indices. */
+	UPROPERTY()
+	TArray<uint32> Indices;
+
+	/** Skinning weights. */
+	UPROPERTY()
+	TArray<FClothVertBoneData> BoneData;
+
+	/** Already remapped using UsedBoneIndices. These are bones that are needed by this LOD for skinning that aren't needed for render.*/
+	UPROPERTY()
+	TArray<uint16> RequiredExtraBoneIndices;
+
+	/** LOD Transition mesh to mesh skinning weights. */
+	TArray<FMeshToMeshVertData> LODTransitionUpData;
+	TArray<FMeshToMeshVertData> LODTransitionDownData;
+
+	/** 2d pattern positions. */
+	UPROPERTY()
+	TArray<FVector2f> PatternPositions;
+
+	/** Pattern triangle indices. Indexing into PatternPositions. */
+	UPROPERTY()
+	TArray<uint32> PatternIndices;
+
+	/** Map from PatternsPositions indices to Positions indices. */
+	UPROPERTY()
+	TArray<uint32> PatternToWeldedIndices;
+
+	/** Weight maps for storing painted attributes modifiers on constraint properties. */
+	TMap<FName, TArray<float>> WeightMaps;
+
+	/** Tether data*/
+	UPROPERTY()
+	FClothTetherData TetherData;
+
+	/** Vertex sets */
+	TMap<FName, TSet<int32>> VertexSets;
+
+	/** Face int maps (currently used by cloth collision layers)*/
+	TMap<FName, TArray<int32>> FaceIntMaps;
+
+	/** Face sets */
+	TMap<FName, TSet<int32>> FaceSets;
+
+	// Custom serialize for weight maps
+	bool Serialize(FArchive& Ar);
+
+	void GetResourceSizeEx(FResourceSizeEx& CumulativeResourceSize);
+};
+
+template<>
+struct TStructOpsTypeTraits<FChaosClothSimulationLodModel> : public TStructOpsTypeTraitsBase2<FChaosClothSimulationLodModel>
+{
+	enum
+	{
+		WithSerializer = true,
+	};
+};
+
+/**
+ * Cloth simulation model.
+ * Contains the LOD models.
+ */
+USTRUCT()
+struct FChaosClothSimulationModel
+{
+	GENERATED_BODY()
+
+	/** LOD data. */
+	UPROPERTY()
+	TArray<FChaosClothSimulationLodModel> ClothSimulationLodModels;
+
+	/** List of bones this asset uses from the reference skeleton. */
+	UPROPERTY()
+	TArray<FName> UsedBoneNames;
+
+	/** List of the indices for the bones in UsedBoneNames, used for remapping. */
+	UPROPERTY()
+	TArray<int32> UsedBoneIndices;
+
+	/** Bone to treat as the root of the simulation space. */
+	UPROPERTY()
+	int32 ReferenceBoneIndex = INDEX_NONE;
+
+	FChaosClothSimulationModel() = default;
+	CHAOSCLOTHASSETENGINE_API FChaosClothSimulationModel(const TArray<TSharedRef<const FManagedArrayCollection>>& ClothCollections, const FReferenceSkeleton& ReferenceSkeleton, TArray<FChaosClothAssetLodTransitionDataCache>* InOutTransitionCache = nullptr);
+
+	int32 GetNumLods() const { return ClothSimulationLodModels.Num(); }
+
+	bool IsValidLodIndex(int32 LodIndex) const { return ClothSimulationLodModels.IsValidIndex(LodIndex); }
+
+	int32 GetNumVertices(int32 LodIndex) const { return IsValidLodIndex(LodIndex) ? ClothSimulationLodModels[LodIndex].Positions.Num() : 0; }
+
+	TConstArrayView<FVector3f> GetPositions(int32 LodIndex) const { return IsValidLodIndex(LodIndex) ? ClothSimulationLodModels[LodIndex].Positions : TConstArrayView<FVector3f>(); }
+	TConstArrayView<FVector2f> GetPatternPositions(int32 LodIndex) const { return IsValidLodIndex(LodIndex) ? ClothSimulationLodModels[LodIndex].PatternPositions : TConstArrayView<FVector2f>(); }
+	TConstArrayView<FVector3f> GetNormals(int32 LodIndex) const { return IsValidLodIndex(LodIndex) ? ClothSimulationLodModels[LodIndex].Normals : TConstArrayView<FVector3f>(); }
+	TConstArrayView<uint32> GetIndices(int32 LodIndex) const { return IsValidLodIndex(LodIndex) ? ClothSimulationLodModels[LodIndex].Indices : TConstArrayView<uint32>(); }
+	TConstArrayView<uint32> GetPatternIndices(int32 LodIndex) const { return IsValidLodIndex(LodIndex) ? ClothSimulationLodModels[LodIndex].PatternIndices : TConstArrayView<uint32>(); }
+	TConstArrayView<uint32> GetPatternToWeldedIndices(int32 LodIndex) const { return IsValidLodIndex(LodIndex) ? ClothSimulationLodModels[LodIndex].PatternToWeldedIndices : TConstArrayView<uint32>(); }
+	TConstArrayView<FClothVertBoneData> GetBoneData(int32 LodIndex) const { return IsValidLodIndex(LodIndex) ? ClothSimulationLodModels[LodIndex].BoneData : TConstArrayView<FClothVertBoneData>(); }
+	TArray<TConstArrayView<TTuple<int32, int32, float>>> GetTethers(int32 LodIndex) const;
+
+	void CalculateLODTransitionUpDownData(TArray<FChaosClothAssetLodTransitionDataCache>* InOutTransitionCache = nullptr);
+
+	CHAOSCLOTHASSETENGINE_API void GetResourceSizeEx(FResourceSizeEx& CumulativeResourceSize);
+
+	/** Remap used bones to a new skeleton. */
+	CHAOSCLOTHASSETENGINE_API void RemapBoneIndices(const TArray<int32>& BoneMap);
+};
diff --git a/Engine/Plugins/ChaosClothAsset/Source/ChaosClothAssetEngine/Public/ChaosClothAsset/ClothSimulationProxy.h b/Engine/Plugins/ChaosClothAsset/Source/ChaosClothAssetEngine/Public/ChaosClothAsset/ClothSimulationProxy.h
--- a/Engine/Plugins/ChaosClothAsset/Source/ChaosClothAssetEngine/Public/ChaosClothAsset/ClothSimulationProxy.h
+++ b/Engine/Plugins/ChaosClothAsset/Source/ChaosClothAssetEngine/Public/ChaosClothAsset/ClothSimulationProxy.h
@@ -154,7 +154,7 @@ PRAGMA_ENABLE_DEPRECATION_WARNINGS
 		TUniquePtr<FCollisionSourcesProxy> CollisionSourcesProxy;
 
 		// The cloth simulation model used to create this simulation, ownership might get transferred to this proxy if it changes during the simulation
-		TSharedPtr<const FChaosClothSimulationModel> ClothSimulationModel;
+		TArray<TSharedPtr<const FChaosClothSimulationModel>> ClothSimulationModels;
 
 		// Simulation objects
 		TUniquePtr<::Chaos::FClothingSimulationSolver> Solver;
diff --git a/Engine/Plugins/Experimental/ChaosOutfitAsset/Source/ChaosOutfitAssetEngine/ChaosOutfitAssetEngine.Build.cs b/Engine/Plugins/Experimental/ChaosOutfitAsset/Source/ChaosOutfitAssetEngine/ChaosOutfitAssetEngine.Build.cs
--- a/Engine/Plugins/Experimental/ChaosOutfitAsset/Source/ChaosOutfitAssetEngine/ChaosOutfitAssetEngine.Build.cs
+++ b/Engine/Plugins/Experimental/ChaosOutfitAsset/Source/ChaosOutfitAssetEngine/ChaosOutfitAssetEngine.Build.cs
@@ -21,9 +21,18 @@ public class ChaosOutfitAssetEngine : ModuleRules
 			{
 				"Core",
 				"CoreUObject",
+				"Chaos",
 				"ChaosClothAsset",
 				"ChaosClothAssetEngine",
-				"Engine"
+				"ClothingSystemRuntimeCommon",
+				"Engine",
+			}
+		);
+
+		PrivateIncludePathModuleNames.AddRange(
+			new string[]
+			{
+				"DerivedDataCache",
 			}
 		);
 	}
diff --git a/Engine/Plugins/Experimental/ChaosOutfitAsset/Source/ChaosOutfitAssetEngine/Private/ChaosClothAsset/Outfit.cpp b/Engine/Plugins/Experimental/ChaosOutfitAsset/Source/ChaosOutfitAssetEngine/Private/ChaosClothAsset/Outfit.cpp
new file mode 100644
--- /dev/null
+++ b/Engine/Plugins/Experimental/ChaosOutfitAsset/Source/ChaosOutfitAssetEngine/Private/ChaosClothAsset/Outfit.cpp
@@ -0,0 +1,243 @@
+// Copyright Epic Games, Inc. All Rights Reserved.
+
+#include "ChaosClothAsset/Outfit.h"
+#include "Chaos/CollectionPropertyFacade.h"
+#include "ChaosClothAsset/ClothAsset.h"
+#include "ChaosClothAsset/ClothSimulationModel.h"
+
+#include UE_INLINE_GENERATED_CPP_BY_NAME(Outfit)
+
+namespace UE::Chaos::ClothAsset::Private
+{
+	// Merge bone to a new skeleton and return its new index in the new skeleton
+	// returns INDEX_NONE oif the parent don't match
+	static int32 MergeBoneToSkeleton(const int32 BoneIndex, const FReferenceSkeleton& InSkeleton, FReferenceSkeleton& InOutSkeleton)
+	{
+		// Check if it exists in the new skeleton with the same parent
+		const FName BoneName = InSkeleton.GetBoneName(BoneIndex);
+		const int32 ParentBoneIndex = InSkeleton.GetParentIndex(BoneIndex);
+		const int32 ExistingBoneIndex = InOutSkeleton.FindBoneIndex(BoneName);
+		if (ExistingBoneIndex != INDEX_NONE)  // Already exists, update bone map
+		{
+			// Check if the parents matches
+			const int32 ExistingParentBoneIndex = InOutSkeleton.GetParentIndex(ExistingBoneIndex);
+			if (!ensureMsgf((ParentBoneIndex == INDEX_NONE && ExistingParentBoneIndex == INDEX_NONE) ||
+				(InSkeleton.GetBoneName(ParentBoneIndex) == InOutSkeleton.GetBoneName(ExistingParentBoneIndex)),
+				TEXT("Bone parent must match")))
+			{
+				// TODO Better error management than ensure
+			}
+			// TODO: Warn if the refposes don't match
+
+			// Update the bone map
+			return ExistingBoneIndex;
+		}
+		else  // Doesn't already exists, add the bone
+		{
+			// Find its parent in the new skeleton
+			const FName ParentBoneName = InSkeleton.GetBoneName(ParentBoneIndex);
+			const int32 ExistingParentBoneIndex = InOutSkeleton.FindBoneIndex(ParentBoneName);
+			checkf(ParentBoneIndex == INDEX_NONE || ExistingParentBoneIndex != INDEX_NONE, TEXT("Either root or a valid parent should exist at this point."));
+			FReferenceSkeletonModifier ReferenceSkeletonModifier(InOutSkeleton, nullptr);
+
+			FMeshBoneInfo MeshBoneInfo;
+#if WITH_EDITORONLY_DATA
+			MeshBoneInfo.ExportName = BoneName.ToString();
+#endif
+			MeshBoneInfo.Name = BoneName;
+			MeshBoneInfo.ParentIndex = ExistingParentBoneIndex;
+			ReferenceSkeletonModifier.Add(MeshBoneInfo, InSkeleton.GetRefBonePose()[BoneIndex]);
+
+			// Update the bone map
+			return ReferenceSkeletonModifier.FindBoneIndex(BoneName);
+		}
+	}
+
+	static void MergeSkeletons(const FReferenceSkeleton& InSkeleton, FReferenceSkeleton& InOutSkeleton, TArray<int32>& OutBoneMap)
+	{
+		const int32 NumBones = InSkeleton.GetNum();
+
+		OutBoneMap.SetNumZeroed(NumBones);
+
+		TBitArray BonesToProcess(true, NumBones);
+
+		for (;;)
+		{
+			// Find the first unprocessed bone
+			int32 BoneIndex = BonesToProcess.Find(true);
+			if (BoneIndex == INDEX_NONE)
+			{
+				break;
+			}
+
+			// Replace by its parent in case they're not already processed (this assumes the bones aren't sorted to be on the safe side)
+			for (int32 ParentBoneIndex = InSkeleton.GetParentIndex(BoneIndex);
+				ParentBoneIndex != INDEX_NONE && BonesToProcess[ParentBoneIndex];
+				ParentBoneIndex = InSkeleton.GetParentIndex(ParentBoneIndex))
+			{
+				BoneIndex = ParentBoneIndex;
+			}
+
+			// Merge the bone
+			const int32 NewBoneIndex = MergeBoneToSkeleton(BoneIndex, InSkeleton, InOutSkeleton);
+			OutBoneMap[BoneIndex] = NewBoneIndex;
+
+			// Marked it as processed
+			BonesToProcess[BoneIndex] = false;
+		}
+	}
+}  // namespace UE::Chaos::ClothAsset::Private
+
+FChaosOutfitPiece::FChaosOutfitPiece()
+	: ClothSimulationModel(MakeShared<FChaosClothSimulationModel>())
+{
+}
+
+FChaosOutfitPiece::~FChaosOutfitPiece() = default;
+
+FChaosOutfitPiece::FChaosOutfitPiece(const UChaosClothAssetBase& ClothAssetBase, int32 ModelIndex)
+	: Name(ClothAssetBase.GetClothSimulationModelName(ModelIndex))
+	, PhysicsAsset(ClothAssetBase.GetPhysicsAssetForModel(ModelIndex))
+	, ClothSimulationModel(MakeShared<FChaosClothSimulationModel>(*ClothAssetBase.GetClothSimulationModel(ModelIndex)))  // Deep copy model from one TSharedRef to the other
+{
+	DeepCopyCollections(ClothAssetBase.GetCollections(ModelIndex));
+}
+
+FChaosOutfitPiece::FChaosOutfitPiece(const FChaosOutfitPiece& Other)
+	: Name(Other.Name)
+	, PhysicsAsset(Other.PhysicsAsset)
+	, ClothSimulationModel(MakeShared<FChaosClothSimulationModel>(*Other.ClothSimulationModel))
+{
+	DeepCopyCollections(Other.Collections);
+}
+
+FChaosOutfitPiece::FChaosOutfitPiece(FChaosOutfitPiece&& Other)
+	: Name(MoveTemp(Other.Name))
+	, PhysicsAsset(MoveTemp(Other.PhysicsAsset))
+	, ClothSimulationModel(MoveTemp(Other.ClothSimulationModel))
+	, Collections(MoveTemp(Other.Collections))
+{}
+
+FChaosOutfitPiece& FChaosOutfitPiece::operator=(const FChaosOutfitPiece& Other)
+{
+	if (&Other != this)
+	{
+		Name = Other.Name;
+		PhysicsAsset = Other.PhysicsAsset;
+		*ClothSimulationModel = *Other.ClothSimulationModel;
+		DeepCopyCollections(Other.Collections);
+	}
+	return *this;
+}
+
+FChaosOutfitPiece& FChaosOutfitPiece::operator=(FChaosOutfitPiece&& Other)
+{
+	if (&Other != this)
+	{
+		Name = MoveTemp(Other.Name);
+		PhysicsAsset = MoveTemp(Other.PhysicsAsset);
+		ClothSimulationModel = MoveTemp(Other.ClothSimulationModel);
+		Collections = MoveTemp(Other.Collections);
+	}
+	return *this;
+}
+
+bool FChaosOutfitPiece::Serialize(FArchive& Ar)
+{
+	// Serialize what can be done with tagged properties
+	{
+		UScriptStruct* const Struct = FChaosOutfitPiece::StaticStruct();
+		Struct->SerializeTaggedProperties(Ar, (uint8*)this, Struct, nullptr);
+	}
+	// Custom serialize the model since TSharedRef can't be declared as UPROPERTY
+	{
+		UScriptStruct* const Struct = FChaosClothSimulationModel::StaticStruct();
+		Struct->SerializeTaggedProperties(Ar, (uint8*)&ClothSimulationModel.Get(), Struct, nullptr);
+	}
+	// Custom serialize tThe collections since TArray<TSharedRef> can't be declared as UPROPERTY
+	int32 NumCollections = Collections.Num();
+	Ar << NumCollections;
+	Collections.Reserve(NumCollections);
+	for (int32 Index = 0; Index < NumCollections; ++Index)
+	{
+		if (Index >= Collections.Num())
+		{
+			Collections.Emplace(MakeShared<FManagedArrayCollection>());
+		}
+		Collections[Index]->Serialize(Ar);
+	}
+	return true;
+}
+
+void FChaosOutfitPiece::DeepCopyCollections(const TArray<TSharedRef<FManagedArrayCollection>>& Other)
+{
+	Collections.Reserve(Other.Num());
+	for (const TSharedRef<FManagedArrayCollection>& Collection : Other)
+	{
+		Collections.Emplace(MakeShared<FManagedArrayCollection>(*Collection));
+	}
+}
+
+void FChaosOutfitPiece::RemapBoneIndices(const TArray<int32>& BoneMap)
+{
+	ClothSimulationModel->RemapBoneIndices(BoneMap);
+}
+
+UChaosOutfit::UChaosOutfit(const FObjectInitializer& ObjectInitializer)
+	: Super(ObjectInitializer)
+{
+	// Create a default reference skeleton
+	ReferenceSkeleton.Empty(1);
+	FReferenceSkeletonModifier ReferenceSkeletonModifier(ReferenceSkeleton, nullptr);
+
+	FMeshBoneInfo MeshBoneInfo;
+	constexpr const TCHAR* RootName = TEXT("Root");
+	MeshBoneInfo.ParentIndex = INDEX_NONE;
+#if WITH_EDITORONLY_DATA
+	MeshBoneInfo.ExportName = RootName;
+#endif
+	MeshBoneInfo.Name = FName(RootName);
+	ReferenceSkeletonModifier.Add(MeshBoneInfo, FTransform::Identity);
+}
+
+UChaosOutfit::UChaosOutfit(FVTableHelper& Helper)
+	: Super(Helper)
+{
+}
+
+UChaosOutfit::~UChaosOutfit() = default;
+
+void UChaosOutfit::Append(const UChaosOutfit& Other)
+{
+	TArray<int32> BoneMap;
+	UE::Chaos::ClothAsset::Private::MergeSkeletons(Other.GetReferenceSkeleton(), ReferenceSkeleton, BoneMap);
+
+	Pieces.Reserve(Pieces.Num() + Other.Pieces.Num());
+
+	for (const FChaosOutfitPiece& Piece : Other.Pieces)
+	{
+		FChaosOutfitPiece& NewPiece = Pieces.Emplace_GetRef(Piece);
+		NewPiece.RemapBoneIndices(BoneMap);
+	}
+}
+
+void UChaosOutfit::Add(const UChaosClothAssetBase& ClothAssetBase)
+{
+	TArray<int32> BoneMap;
+	UE::Chaos::ClothAsset::Private::MergeSkeletons(ClothAssetBase.GetRefSkeleton(), ReferenceSkeleton, BoneMap);
+
+	Pieces.Reserve(Pieces.Num() + ClothAssetBase.GetNumClothSimulationModels());
+
+	for (int32 ModelIndex = 0; ModelIndex < ClothAssetBase.GetNumClothSimulationModels(); ++ModelIndex)
+	{
+		FChaosOutfitPiece& NewPiece = Pieces.Emplace_GetRef(ClothAssetBase, ModelIndex);
+		NewPiece.RemapBoneIndices(BoneMap);
+	}
+}
+
+void UChaosOutfit::Serialize(FArchive& Ar)
+{
+	Super::Serialize(Ar);
+
+	Ar << ReferenceSkeleton;
+}
diff --git a/Engine/Plugins/Experimental/ChaosOutfitAsset/Source/ChaosOutfitAssetEngine/Private/ChaosClothAsset/OutfitAsset.cpp b/Engine/Plugins/Experimental/ChaosOutfitAsset/Source/ChaosOutfitAssetEngine/Private/ChaosClothAsset/OutfitAsset.cpp
--- a/Engine/Plugins/Experimental/ChaosOutfitAsset/Source/ChaosOutfitAssetEngine/Private/ChaosClothAsset/OutfitAsset.cpp
+++ b/Engine/Plugins/Experimental/ChaosOutfitAsset/Source/ChaosOutfitAssetEngine/Private/ChaosClothAsset/OutfitAsset.cpp
@@ -1,9 +1,18 @@
 // Copyright Epic Games, Inc. All Rights Reserved.
 
 #include "ChaosClothAsset/OutfitAsset.h"
+#include "ChaosClothAsset/ClothSimulationModel.h"
+#if WITH_EDITOR
+#include "DerivedDataCacheInterface.h"
+#endif
+#include "Rendering/SkeletalMeshRenderData.h"
 
 #include UE_INLINE_GENERATED_CPP_BY_NAME(OutfitAsset)
 
+// If Chaos outfit asset derived data needs to be rebuilt (new format, serialization differences, etc.) replace the version GUID below with a new one.
+// In case of merge conflicts with DDC versions, you MUST generate a new GUID and set this new GUID as the version.
+#define UE_CHAOS_OUTFIT_ASSET_DERIVED_DATA_VERSION TEXT("DD1C25C90FDE4287881A8759CD3646A6")
+
 UChaosOutfitAsset::UChaosOutfitAsset(const FObjectInitializer& ObjectInitializer)
 	: Super(ObjectInitializer)
 {
@@ -17,3 +26,72 @@ UChaosOutfitAsset::UChaosOutfitAsset(FVTableHelper& Helper)
 
 UChaosOutfitAsset::~UChaosOutfitAsset() = default;
 
+bool UChaosOutfitAsset::HasValidClothSimulationModels() const
+{
+	for (const FChaosOutfitPiece& Piece : Pieces)
+	{
+		if (Piece.ClothSimulationModel->GetNumLods())
+		{
+			return true;
+		}
+	}
+	return false;
+}
+
+FName UChaosOutfitAsset::GetClothSimulationModelName(int32 ModelIndex) const
+{
+	return Pieces[ModelIndex].Name;
+}
+
+TSharedPtr<const FChaosClothSimulationModel> UChaosOutfitAsset::GetClothSimulationModel(int32 ModelIndex) const
+{
+	return Pieces[ModelIndex].ClothSimulationModel;
+}
+
+const TArray<TSharedRef<FManagedArrayCollection>>& UChaosOutfitAsset::GetCollections(int32 ModelIndex) const
+{
+	return Pieces[ModelIndex].Collections;
+}
+
+const UPhysicsAsset* UChaosOutfitAsset::GetPhysicsAssetForModel(int32 ModelIndex) const
+{
+	return Pieces[ModelIndex].PhysicsAsset;
+}
+
+#if WITH_EDITOR
+void UChaosOutfitAsset::BuildLODModel(const ITargetPlatform* /*TargetPlatform*/, int32 /*LODIndex*/)
+{
+}
+
+FString UChaosOutfitAsset::BuildDerivedDataKey(const ITargetPlatform* TargetPlatform)
+{
+	const FString KeySuffix;
+	return FDerivedDataCacheInterface::BuildCacheKey(
+		TEXT("CHAOSOUTFIT"),
+		UE_CHAOS_OUTFIT_ASSET_DERIVED_DATA_VERSION,
+		*KeySuffix
+	);
+}
+
+bool UChaosOutfitAsset::IsInitialBuildDone() const
+{
+	return true;
+}
+#endif // #if WITH_EDITOR
+
+void UChaosOutfitAsset::Build(const TObjectPtr<const UChaosOutfit> Outfit)
+{
+	// Copy pieces across
+	const TArray<FChaosOutfitPiece>& InPieces = Outfit->GetPieces();
+	Pieces.Reset(InPieces.Num());
+	Pieces.Append(InPieces);
+
+	// Copy reference skeleton
+	SetReferenceSkeleton(&Outfit->GetReferenceSkeleton());
+
+	// Create dummy skm render data
+	SetResourceForRendering(MakeUnique<FSkeletalMeshRenderData>());
+
+	// Re-register any components using this asset to restart the simulation and update the leader bone map
+	ReregisterComponents();
+}
diff --git a/Engine/Plugins/Experimental/ChaosOutfitAsset/Source/ChaosOutfitAssetEngine/Public/ChaosClothAsset/Outfit.h b/Engine/Plugins/Experimental/ChaosOutfitAsset/Source/ChaosOutfitAssetEngine/Public/ChaosClothAsset/Outfit.h
new file mode 100644
--- /dev/null
+++ b/Engine/Plugins/Experimental/ChaosOutfitAsset/Source/ChaosOutfitAssetEngine/Public/ChaosClothAsset/Outfit.h
@@ -0,0 +1,97 @@
+// Copyright Epic Games, Inc. All Rights Reserved.
+
+#pragma once
+
+#include "Containers/Array.h"
+#include "Templates/SharedPointerFwd.h"
+#include "UObject/Object.h"
+#include "ReferenceSkeleton.h"
+
+#include "Outfit.generated.h"
+
+class UChaosClothAssetBase;
+class UPhysicsAsset;
+struct FChaosClothSimulationModel;
+struct FManagedArrayCollection;
+
+USTRUCT()
+struct CHAOSOUTFITASSETENGINE_API FChaosOutfitPiece final
+{
+	GENERATED_BODY()
+
+	UPROPERTY()
+	FName Name;
+
+	UPROPERTY()
+	TObjectPtr<const UPhysicsAsset> PhysicsAsset;
+
+	TSharedRef<FChaosClothSimulationModel> ClothSimulationModel;
+
+	TArray<TSharedRef<FManagedArrayCollection>> Collections;
+
+	FChaosOutfitPiece();
+	~FChaosOutfitPiece();
+	FChaosOutfitPiece(const UChaosClothAssetBase& ClothAssetBase, int32 ModelIndex);
+
+	FChaosOutfitPiece(const FChaosOutfitPiece& Other);
+	FChaosOutfitPiece(FChaosOutfitPiece&& Other);
+
+	FChaosOutfitPiece& operator=(const FChaosOutfitPiece& Other);
+	FChaosOutfitPiece& operator=(FChaosOutfitPiece&& Other);
+
+	// Custom serialize for TSharedRef
+	bool Serialize(FArchive& Ar);
+
+private:
+	friend UChaosOutfit;
+
+	void DeepCopyCollections(const TArray<TSharedRef<FManagedArrayCollection>>& Other);
+
+	void RemapBoneIndices(const TArray<int32>& BoneMap);
+};
+
+template<>
+struct TStructOpsTypeTraits<FChaosOutfitPiece> : public TStructOpsTypeTraitsBase2<FChaosOutfitPiece>
+{
+	enum
+	{
+		WithSerializer = true,
+	};
+};
+
+/**
+ * Outfit class to handle the assembly of an outfit asset.
+ */
+UCLASS(Experimental)
+class CHAOSOUTFITASSETENGINE_API UChaosOutfit final : public UObject
+{
+	GENERATED_BODY()
+
+public:
+	UChaosOutfit(const FObjectInitializer& ObjectInitializer);
+	UChaosOutfit(FVTableHelper& Helper);  // This is declared so we can use TUniquePtr<FClothSimulationModel> with just a forward declare of that class
+	virtual ~UChaosOutfit() override;
+
+	void Append(const UChaosOutfit& Other);
+	void Add(const UChaosClothAssetBase& ClothAssetBase);
+
+	const TArray<FChaosOutfitPiece>& GetPieces() const
+	{
+		return Pieces;
+	}
+
+	const FReferenceSkeleton& GetReferenceSkeleton() const
+	{
+		return ReferenceSkeleton;
+	}
+
+	//~ Begin UObject interface
+	virtual void Serialize(FArchive& Ar) override;
+	//~ End UObject interface
+
+private:
+	UPROPERTY()
+	TArray<FChaosOutfitPiece> Pieces;
+
+	FReferenceSkeleton ReferenceSkeleton;
+};
diff --git a/Engine/Plugins/Experimental/ChaosOutfitAsset/Source/ChaosOutfitAssetEngine/Public/ChaosClothAsset/OutfitAsset.h b/Engine/Plugins/Experimental/ChaosOutfitAsset/Source/ChaosOutfitAssetEngine/Public/ChaosClothAsset/OutfitAsset.h
--- a/Engine/Plugins/Experimental/ChaosOutfitAsset/Source/ChaosOutfitAssetEngine/Public/ChaosClothAsset/OutfitAsset.h
+++ b/Engine/Plugins/Experimental/ChaosOutfitAsset/Source/ChaosOutfitAssetEngine/Public/ChaosClothAsset/OutfitAsset.h
@@ -2,21 +2,68 @@
 
 #pragma once
 
-#include "ChaosClothAsset/ClothAsset.h"
+#include "ChaosClothAsset/ClothAssetBase.h"
+#include "ChaosClothAsset/Outfit.h"
 #include "OutfitAsset.generated.h"
 
-class USkeletalMesh;
-class UDataflow;
+struct FChaosOutfitPiece;
+struct FSkeletalMaterial;
 
 /**
  * Outfit asset for character clothing and simulation.
  */
 UCLASS(Experimental, HideCategories = Object, BlueprintType)
-class CHAOSOUTFITASSETENGINE_API UChaosOutfitAsset final : public UChaosClothAsset
+class CHAOSOUTFITASSETENGINE_API UChaosOutfitAsset final : public UChaosClothAssetBase
 {
 	GENERATED_BODY()
 public:
 	UChaosOutfitAsset(const FObjectInitializer& ObjectInitializer);
 	UChaosOutfitAsset(FVTableHelper& Helper);  // This is declared so we can use TUniquePtr<FClothSimulationModel> with just a forward declare of that class
 	virtual ~UChaosOutfitAsset() override;
+
+	void Build(const TObjectPtr<const UChaosOutfit> Outfit = nullptr);  // TODO: Remove default value once the terminal node code is in
+
+private:
+	//~ Begin UChaosClothAssetBase interface
+	virtual bool HasValidClothSimulationModels() const override;
+	virtual int32 GetNumClothSimulationModels() const override
+	{
+		return Pieces.Num();
+	}
+	virtual FName GetClothSimulationModelName(int32 ModelIndex) const override;
+	virtual TSharedPtr<const FChaosClothSimulationModel> GetClothSimulationModel(int32 ModelIndex) const override;
+	virtual const TArray<TSharedRef<FManagedArrayCollection>>& GetCollections(int32 ModelIndex) const override;
+	virtual const UPhysicsAsset* GetPhysicsAssetForModel(int32 ModelIndex) const override;
+	//~ End UChaosClothAssetBase interface
+
+	//~ Begin USkinnedAsset interface
+	virtual UPhysicsAsset* GetPhysicsAsset() const override
+	{
+		return nullptr;  // There isn't a single Physics Asset anymore, this could return the first one but that wouldn't be accurate
+	}
+	virtual USkeleton* GetSkeleton() override
+	{
+		return nullptr;  // Note: The USkeleton isn't a reliable source of reference skeleton
+	}
+	virtual const USkeleton* GetSkeleton() const override
+	{
+		return nullptr;
+	}
+	virtual void SetSkeleton(USkeleton* InSkeleton) override {}
+#if WITH_EDITOR
+	virtual void BuildLODModel(const ITargetPlatform* TargetPlatform, int32 LODIndex) override;
+	virtual FString BuildDerivedDataKey(const ITargetPlatform* TargetPlatform) override;
+	virtual bool IsInitialBuildDone() const override;
+#endif
+#if WITH_EDITORONLY_DATA
+	virtual class FSkeletalMeshModel* GetImportedModel() const override
+	{
+		return nullptr;
+	}
+#endif // WITH_EDITORONLY_DATA
+
+	//~ End USkinnedAsset interface
+
+	UPROPERTY()
+	TArray<FChaosOutfitPiece> Pieces;
 };
diff --git a/Engine/Source/Runtime/Core/Public/UObject/FortniteMainBranchObjectVersions.inl b/Engine/Source/Runtime/Core/Public/UObject/FortniteMainBranchObjectVersions.inl
--- a/Engine/Source/Runtime/Core/Public/UObject/FortniteMainBranchObjectVersions.inl
+++ b/Engine/Source/Runtime/Core/Public/UObject/FortniteMainBranchObjectVersions.inl
@@ -568,5 +568,8 @@ FORTNITE_MAIN_VERSION(ResetLevelInstanceHLODRelevancy, 184)
 // Updated default scene capture post-processing settings to reflect the underlying implementation overrides
 FORTNITE_MAIN_VERSION(SceneCaptureDefaultSettings, 185)
 
+// Add Cloth Asset Base class serialization
+FORTNITE_MAIN_VERSION(AddClothAssetBase, 186)
+
 // -----<new versions can be added above this line>-------------------------------------------------
 #undef FORTNITE_MAIN_VERSION
\ No newline at end of file
