Chaos Cloth Asset :
- Fixed LogEditCondition: EditCondition parsing failed: Field name "DataflowInstance.DataflowAsset"
- Deprecated the DataflowInstance property so that a new bHasDataflowAsset property can be updated to be used in the edit condition.
- Fixed the OverlayMaterial PostEditChangeProperty ReregisterComponents location.

#rnx

diff --git a/Engine/Plugins/ChaosClothAsset/Source/ChaosClothAssetEngine/Private/ChaosClothAsset/ClothAsset.cpp b/Engine/Plugins/ChaosClothAsset/Source/ChaosClothAssetEngine/Private/ChaosClothAsset/ClothAsset.cpp
--- a/Engine/Plugins/ChaosClothAsset/Source/ChaosClothAssetEngine/Private/ChaosClothAsset/ClothAsset.cpp
+++ b/Engine/Plugins/ChaosClothAsset/Source/ChaosClothAssetEngine/Private/ChaosClothAsset/ClothAsset.cpp
@@ -155,7 +155,9 @@ UChaosClothAsset::UChaosClothAsset(const FObjectInitializer& ObjectInitializer)
 	, MeshModel(MakeShareable(new FSkeletalMeshModel()))
 #endif
 {
+PRAGMA_DISABLE_DEPRECATION_WARNINGS
 	DataflowInstance.SetDataflowTerminal(TEXT("ClothAssetTerminal"));
+PRAGMA_ENABLE_DEPRECATION_WARNINGS
 
 	// Setup a single LOD's Cloth Collection
 	TSharedRef<FManagedArrayCollection>& ClothCollection = GetClothCollections().Emplace_GetRef(MakeShared<FManagedArrayCollection>());
@@ -223,8 +225,10 @@ void UChaosClothAsset::PostLoad()
 #if WITH_EDITORONLY_DATA
 	if (DataflowAsset_DEPRECATED != nullptr)
 	{
-		DataflowInstance.SetDataflowAsset(DataflowAsset_DEPRECATED);
+		SetDataflow(DataflowAsset_DEPRECATED);
+PRAGMA_DISABLE_DEPRECATION_WARNINGS
 		DataflowInstance.SetDataflowTerminal(FName(DataflowTerminal_DEPRECATED));
+PRAGMA_ENABLE_DEPRECATION_WARNINGS
 
 		DataflowAsset_DEPRECATED = nullptr;
 		DataflowTerminal_DEPRECATED.Empty();
@@ -235,15 +239,11 @@ void UChaosClothAsset::PostLoad()
 #if WITH_EDITOR
 void UChaosClothAsset::PostEditChangeProperty(FPropertyChangedEvent& PropertyChangedEvent)
 {
-	if (PropertyChangedEvent.GetPropertyName() == GET_MEMBER_NAME_CHECKED(UChaosClothAsset, PhysicsAsset) ||
-		PRAGMA_DISABLE_DEPRECATION_WARNINGS
-		PropertyChangedEvent.GetPropertyName() == GET_MEMBER_NAME_CHECKED(UChaosClothAsset, OverlayMaterial))
-		PRAGMA_ENABLE_DEPRECATION_WARNINGS
+	if (PropertyChangedEvent.GetPropertyName() == GET_MEMBER_NAME_CHECKED(UChaosClothAsset, PhysicsAsset))
 	{
 		ReregisterComponents();
 	}
 	Super::PostEditChangeProperty(PropertyChangedEvent);
-	InvalidateDataflowContents();
 }
 #endif // #if WITH_EDITOR
 
diff --git a/Engine/Plugins/ChaosClothAsset/Source/ChaosClothAssetEngine/Private/ChaosClothAsset/ClothAssetBase.cpp b/Engine/Plugins/ChaosClothAsset/Source/ChaosClothAssetEngine/Private/ChaosClothAsset/ClothAssetBase.cpp
--- a/Engine/Plugins/ChaosClothAsset/Source/ChaosClothAssetEngine/Private/ChaosClothAsset/ClothAssetBase.cpp
+++ b/Engine/Plugins/ChaosClothAsset/Source/ChaosClothAssetEngine/Private/ChaosClothAsset/ClothAssetBase.cpp
@@ -60,7 +60,9 @@ enum class UChaosClothAssetBase::EAsyncProperties : uint32
 
 UChaosClothAssetBase::UChaosClothAssetBase(const FObjectInitializer& ObjectInitializer)
 	: Super(ObjectInitializer)
+PRAGMA_DISABLE_DEPRECATION_WARNINGS
 	, DataflowInstance(this)
+PRAGMA_ENABLE_DEPRECATION_WARNINGS
 	, MinQualityLevelLOD(0)
 	, DisableBelowMinLodStripping(FPerPlatformBool(false))
 	, MinLod(0)
@@ -115,8 +117,8 @@ void UChaosClothAssetBase::WriteDataflowContent(const TObjectPtr<UDataflowBaseCo
 {
 	if (const TObjectPtr<UDataflowSkeletalContent> SkeletalContent = Cast<UDataflowSkeletalContent>(DataflowContent))
 	{
-		SkeletalContent->SetDataflowAsset(DataflowInstance.GetDataflowAsset());
-		SkeletalContent->SetDataflowTerminal(DataflowInstance.GetDataflowTerminal().ToString());
+		SkeletalContent->SetDataflowAsset(GetDataflowInstance().GetDataflowAsset());
+		SkeletalContent->SetDataflowTerminal(GetDataflowInstance().GetDataflowTerminal().ToString());
 
 #if WITH_EDITORONLY_DATA
 		SkeletalContent->SetAnimationAsset(GetPreviewSceneAnimation());
@@ -136,6 +138,13 @@ void UChaosClothAssetBase::ReadDataflowContent(const TObjectPtr<UDataflowBaseCon
 	}
 }
 
+const FDataflowInstance& UChaosClothAssetBase::GetDataflowInstance() const
+{
+PRAGMA_DISABLE_DEPRECATION_WARNINGS
+	return DataflowInstance;
+PRAGMA_ENABLE_DEPRECATION_WARNINGS
+}
+
 void UChaosClothAssetBase::BeginDestroy()
 {
 	check(IsInGameThread());
@@ -172,6 +181,38 @@ void UChaosClothAssetBase::Serialize(FArchive& Ar)
 	Ar << GetRefSkeleton();
 }
 
+void UChaosClothAssetBase::PostLoad()
+{
+	Super::PostLoad();
+
+#if WITH_EDITORONLY_DATA
+	bHasDataflowAsset = (GetDataflowInstance().GetDataflowAsset() != nullptr);
+#endif  // WITH_EDITORONLY_DATA
+}
+
+#if WITH_EDITOR
+void UChaosClothAssetBase::PostEditChangeProperty(FPropertyChangedEvent& PropertyChangedEvent)
+{
+PRAGMA_DISABLE_DEPRECATION_WARNINGS
+	if (PropertyChangedEvent.GetPropertyName() == GET_MEMBER_NAME_CHECKED(UChaosClothAssetBase, OverlayMaterial))
+	{
+		ReregisterComponents();
+	}
+	else if (PropertyChangedEvent.GetMemberPropertyName() == GET_MEMBER_NAME_CHECKED(UChaosClothAssetBase, DataflowInstance) && 
+		PropertyChangedEvent.GetPropertyName() == FName(TEXT("DataflowAsset")))  // Can't use GET_MEMBER_NAME_CHECKED, because DataflowAsset is private
+	{
+		bHasDataflowAsset = (GetDataflowInstance().GetDataflowAsset() != nullptr);
+	}
+PRAGMA_ENABLE_DEPRECATION_WARNINGS
+
+	// This invalidates the IDataflowContentOwner, not the Dataflow itself
+	InvalidateDataflowContents();
+
+	Super::PostEditChangeProperty(PropertyChangedEvent);
+}
+#endif // #if WITH_EDITOR
+
+
 FSkeletalMeshLODInfo* UChaosClothAssetBase::GetLODInfo(int32 Index)
 {
 	return LODInfo.IsValidIndex(Index) ? &LODInfo[Index] : nullptr;
@@ -290,6 +331,31 @@ PRAGMA_DISABLE_DEPRECATION_WARNINGS
 PRAGMA_ENABLE_DEPRECATION_WARNINGS
 }
 
+void UChaosClothAssetBase::SetDataflow(UDataflow* InDataflow)
+{
+PRAGMA_DISABLE_DEPRECATION_WARNINGS
+	DataflowInstance.SetDataflowAsset(InDataflow);
+PRAGMA_ENABLE_DEPRECATION_WARNINGS
+
+#if WITH_EDITORONLY_DATA
+	bHasDataflowAsset = (GetDataflowInstance().GetDataflowAsset() != nullptr);
+#endif
+}
+
+UDataflow* UChaosClothAssetBase::GetDataflow()
+{
+PRAGMA_DISABLE_DEPRECATION_WARNINGS
+	return DataflowInstance.GetDataflowAsset();
+PRAGMA_ENABLE_DEPRECATION_WARNINGS
+}
+
+const UDataflow* UChaosClothAssetBase::GetDataflow() const
+{
+PRAGMA_DISABLE_DEPRECATION_WARNINGS
+	return DataflowInstance.GetDataflowAsset();
+PRAGMA_ENABLE_DEPRECATION_WARNINGS
+}
+
 void UChaosClothAssetBase::InitResources()
 {
 	LLM_SCOPE_BYNAME(TEXT("ClothAsset/InitResources"));
diff --git a/Engine/Plugins/ChaosClothAsset/Source/ChaosClothAssetEngine/Public/ChaosClothAsset/ClothAsset.h b/Engine/Plugins/ChaosClothAsset/Source/ChaosClothAssetEngine/Public/ChaosClothAsset/ClothAsset.h
--- a/Engine/Plugins/ChaosClothAsset/Source/ChaosClothAssetEngine/Public/ChaosClothAsset/ClothAsset.h
+++ b/Engine/Plugins/ChaosClothAsset/Source/ChaosClothAssetEngine/Public/ChaosClothAsset/ClothAsset.h
@@ -134,14 +134,18 @@ public:
 	UE_DEPRECATED(5.6, "Function will be removed as the DataflowTerminal shouldn't be set by name/string.")
 	void SetDataflowTerminal(const FString& InDataflowTerminal) 
 	{ 
+PRAGMA_DISABLE_DEPRECATION_WARNINGS
 		DataflowInstance.SetDataflowTerminal(FName(*InDataflowTerminal));
+PRAGMA_ENABLE_DEPRECATION_WARNINGS
 	}
 
 	/** Return the Dataflow graph asset associated to this cloth asset if any. */
 	UE_DEPRECATED(5.6, "Use GetDataflowInstance() instead")
 	FString GetDataflowTerminal() const 
 	{
+PRAGMA_DISABLE_DEPRECATION_WARNINGS
 		return DataflowInstance.GetDataflowTerminal().ToString();
+PRAGMA_ENABLE_DEPRECATION_WARNINGS
 	}
 
 	//~ Begin UChaosClothAssetBase interface
@@ -214,11 +218,11 @@ private:
 	 * This is of limited use since this USkeleton's reference skeleton might not necessarily match the one created for this asset.
 	 * Set by the Dataflow evaluation.
 	 */
-	UPROPERTY(EditAnywhere, Setter = SetSkeleton, Category = Skeleton, Meta = (EditCondition = "DataflowInstance.DataflowAsset == nullptr"))
+	UPROPERTY(EditAnywhere, Setter = SetSkeleton, Category = Skeleton, Meta = (EditCondition = "!bHasDataflowAsset"))
 	TObjectPtr<USkeleton> Skeleton;
 
 	/** Physics asset used for collision. Set by the Dataflow evaluation. */
-	UPROPERTY(EditAnywhere, Category = Collision, Meta = (EditCondition = "DataflowInstance.DataflowAsset == nullptr"))
+	UPROPERTY(EditAnywhere, Category = Collision, Meta = (EditCondition = "!bHasDataflowAsset"))
 	TObjectPtr<UPhysicsAsset> PhysicsAsset;
 
 	/** Whether to blend positions between the skinned/simulated transitions of the cloth render mesh. Only used when not overriden in Dataflow by the ProxyDefomer node. */
diff --git a/Engine/Plugins/ChaosClothAsset/Source/ChaosClothAssetEngine/Public/ChaosClothAsset/ClothAssetBase.h b/Engine/Plugins/ChaosClothAsset/Source/ChaosClothAssetEngine/Public/ChaosClothAsset/ClothAssetBase.h
--- a/Engine/Plugins/ChaosClothAsset/Source/ChaosClothAssetEngine/Public/ChaosClothAsset/ClothAssetBase.h
+++ b/Engine/Plugins/ChaosClothAsset/Source/ChaosClothAssetEngine/Public/ChaosClothAsset/ClothAssetBase.h
@@ -68,13 +68,17 @@ public:
 	//~ End IDataflowContentOwner interface 
 
 	//~ Begin IDataflowInstanceInterface interface
-	virtual const FDataflowInstance& GetDataflowInstance() const { return DataflowInstance; };
+	virtual const FDataflowInstance& GetDataflowInstance() const;
 	//~ End IDataflowInstanceInterface interface
 
 	//~ Begin UObject interface
 	virtual void BeginDestroy() override;
 	virtual bool IsReadyForFinishDestroy() override;
 	virtual void Serialize(FArchive& Ar) override;
+	virtual void PostLoad() override;
+#if WITH_EDITOR
+	virtual void PostEditChangeProperty(FPropertyChangedEvent& PropertyChangedEvent) override;
+#endif
 	//~ End UObject interface
 
 	//~ Begin USkinnedAsset interface
@@ -224,22 +228,13 @@ public:
 	 * Set the Dataflow graph asset for this asset.
 	 * @param InDataflow The new dataflow asset.
 	 */
-	void SetDataflow(UDataflow* InDataflow)
-	{
-		DataflowInstance.SetDataflowAsset(InDataflow);
-	}
+	void SetDataflow(UDataflow* InDataflow);
 
 	/** Return the Dataflow graph asset associated to this asset if any. */
-	UDataflow* GetDataflow()
-	{
-		return DataflowInstance.GetDataflowAsset();
-	}
+	UDataflow* GetDataflow();
 
 	/** Return the Dataflow graph asset associated to this asset if any, const version. */
-	const UDataflow* GetDataflow() const
-	{
-		return DataflowInstance.GetDataflowAsset();
-	}
+	const UDataflow* GetDataflow() const;
 
 #if WITH_EDITORONLY_DATA
 	void SetPreviewSceneSkeletalMesh(class USkeletalMesh* Mesh);
@@ -289,11 +284,12 @@ protected:
 	virtual void FinishBuildInternal(FSkinnedAssetBuildContext& Context) override;
 #endif
 
+	UE_DEPRECATED(5.6, "This property isn't deprecated, but the proper getter and setter should be used instead in oreder to preserve correct behavior.")
 	UPROPERTY(EditAnywhere, Category = "Dataflow")
 	FDataflowInstance DataflowInstance;
 
 	/** List of materials for this cloth asset. Set by the Dataflow evaluation. */
-	UPROPERTY(EditAnywhere, Category = Materials, Meta = (EditCondition = "DataflowInstance.DataflowAsset == nullptr"))
+	UPROPERTY(EditAnywhere, Category = Materials, Meta = (EditCondition = "!bHasDataflowAsset"))
 	TArray<FSkeletalMaterial> Materials;
 
 	/** Struct containing information for each LOD level, such as materials to use, and when use the LOD. Not currently editable or customizable through the Dataflow. */
@@ -344,6 +340,12 @@ protected:
 	UE_DEPRECATED(5.4, "This must be protected for async build, always use the accessors even internally.")
 	FReferenceSkeleton RefSkeleton;
 
+#if WITH_EDITORONLY_DATA
+	/** Property used in various edit conditions to enable property changes when this asset has no Dataflow asset. */
+	UPROPERTY(Transient)
+	bool bHasDataflowAsset = false;
+#endif
+
 	/** Bounds for this asset. */
 	FBoxSphereBounds Bounds;
 
diff --git a/Engine/Plugins/Experimental/ChaosOutfitAsset/Source/ChaosOutfitAssetEngine/Private/ChaosClothAsset/OutfitAsset.cpp b/Engine/Plugins/Experimental/ChaosOutfitAsset/Source/ChaosOutfitAssetEngine/Private/ChaosClothAsset/OutfitAsset.cpp
--- a/Engine/Plugins/Experimental/ChaosOutfitAsset/Source/ChaosOutfitAssetEngine/Private/ChaosClothAsset/OutfitAsset.cpp
+++ b/Engine/Plugins/Experimental/ChaosOutfitAsset/Source/ChaosOutfitAssetEngine/Private/ChaosClothAsset/OutfitAsset.cpp
@@ -16,7 +16,9 @@
 UChaosOutfitAsset::UChaosOutfitAsset(const FObjectInitializer& ObjectInitializer)
 	: Super(ObjectInitializer)
 {
+PRAGMA_DISABLE_DEPRECATION_WARNINGS
 	DataflowInstance.SetDataflowTerminal(TEXT("OutfitAssetTerminal"));
+PRAGMA_ENABLE_DEPRECATION_WARNINGS
 }
 
 UChaosOutfitAsset::UChaosOutfitAsset(FVTableHelper& Helper)
