Make ChaosClothAsset MaxDeltaTimeTeleport work the same as SKM cloth.

diff --git a/Engine/Plugins/ChaosClothAsset/Source/ChaosClothAssetEngine/Private/ChaosClothAsset/ClothSimulationContext.cpp b/Engine/Plugins/ChaosClothAsset/Source/ChaosClothAssetEngine/Private/ChaosClothAsset/ClothSimulationContext.cpp
--- a/Engine/Plugins/ChaosClothAsset/Source/ChaosClothAssetEngine/Private/ChaosClothAsset/ClothSimulationContext.cpp
+++ b/Engine/Plugins/ChaosClothAsset/Source/ChaosClothAssetEngine/Private/ChaosClothAsset/ClothSimulationContext.cpp
@@ -24,7 +24,7 @@ namespace UE::Chaos::ClothAsset
 		constexpr float MaxDeltaTimeTeleportMultiplierDefault = 1.5f;
 		const float MaxDeltaTimeTeleportMultiplier = CVarMaxDeltaTimeTeleportMultiplier ? CVarMaxDeltaTimeTeleportMultiplier->GetFloat() : MaxDeltaTimeTeleportMultiplierDefault;
 
-		bTeleport = (DeltaTime > MaxDeltaTime * MaxDeltaTimeTeleportMultiplier) ? true : ClothComponent.NeedsTeleport();
+		bTeleport = (InDeltaTime > MaxDeltaTime * MaxDeltaTimeTeleportMultiplier) ? true : ClothComponent.NeedsTeleport();
 		bReset = ClothComponent.NeedsReset();
 
 		VelocityScale = (!bTeleport && !bReset && InDeltaTime > 0.f) ?
