Set default value of new TeleportDistanceThreshold on UChaosClothComponent to 0 to match existing behavior.

diff --git a/Engine/Plugins/ChaosClothAsset/Source/ChaosClothAssetEngine/Private/ChaosClothAsset/ClothComponent.cpp b/Engine/Plugins/ChaosClothAsset/Source/ChaosClothAssetEngine/Private/ChaosClothAsset/ClothComponent.cpp
--- a/Engine/Plugins/ChaosClothAsset/Source/ChaosClothAssetEngine/Private/ChaosClothAsset/ClothComponent.cpp
+++ b/Engine/Plugins/ChaosClothAsset/Source/ChaosClothAssetEngine/Private/ChaosClothAsset/ClothComponent.cpp
@@ -47,7 +47,7 @@ UChaosClothComponent::UChaosClothComponent(const FObjectInitializer& ObjectIniti
 #if WITH_EDITORONLY_DATA
 	, bSimulateInEditor(0)
 #endif
-	, TeleportDistanceThreshold(300.f)
+	, TeleportDistanceThreshold(0.f)
 	, TeleportRotationThreshold(0.f)
 	, ClothTeleportDistThresholdSquared(UE::ClothingSimulation::TeleportHelpers::ComputeTeleportDistanceThresholdSquared(TeleportDistanceThreshold))
 	, ClothTeleportCosineThresholdInRad(UE::ClothingSimulation::TeleportHelpers::ComputeTeleportCosineRotationThreshold(TeleportRotationThreshold))
