Fix ChaosClothComponent teleport settings--they were not getting reset after a teleport was detected.

This time added ability to test in the Cloth Editor and verified it's working as expected.

diff --git a/Engine/Plugins/ChaosClothAsset/Source/ChaosClothAssetEngine/Private/ChaosClothAsset/ClothComponent.cpp b/Engine/Plugins/ChaosClothAsset/Source/ChaosClothAssetEngine/Private/ChaosClothAsset/ClothComponent.cpp
--- a/Engine/Plugins/ChaosClothAsset/Source/ChaosClothAssetEngine/Private/ChaosClothAsset/ClothComponent.cpp
+++ b/Engine/Plugins/ChaosClothAsset/Source/ChaosClothAssetEngine/Private/ChaosClothAsset/ClothComponent.cpp
@@ -51,6 +51,7 @@ UChaosClothComponent::UChaosClothComponent(const FObjectInitializer& ObjectIniti
 	, TeleportRotationThreshold(0.f)
 	, ClothTeleportDistThresholdSquared(UE::ClothingSimulation::TeleportHelpers::ComputeTeleportDistanceThresholdSquared(TeleportDistanceThreshold))
 	, ClothTeleportCosineThresholdInRad(UE::ClothingSimulation::TeleportHelpers::ComputeTeleportCosineRotationThreshold(TeleportRotationThreshold))
+	, ClothTeleportMode(EClothingTeleportMode::None)
 	, CollisionSources(MakeUnique<UE::Chaos::ClothAsset::FCollisionSources>(this, (bool)bCollideWithEnvironment))
 {
 	PrimaryComponentTick.EndTickGroup = TG_PostPhysics;
@@ -646,23 +647,8 @@ void UChaosClothComponent::StartNewParallelSimulation(float DeltaTime)
 void UChaosClothComponent::UpdateClothTeleport()
 {
 	const FMatrix CurRootBoneMat = GetBoneMatrix(0);
-	EClothingTeleportMode TeleportMode = bTeleport ? (bReset ? EClothingTeleportMode::TeleportAndReset : EClothingTeleportMode::Teleport) : EClothingTeleportMode::None;
-	TeleportMode = UE::ClothingSimulation::TeleportHelpers::CalculateClothingTeleport(TeleportMode, CurRootBoneMat, PrevRootBoneMatrix, bReset, ClothTeleportDistThresholdSquared, ClothTeleportCosineThresholdInRad);
-
-	switch (TeleportMode)
-	{
-	default:
-	case EClothingTeleportMode::None:
-		bTeleport = false;
-		break;
-	case EClothingTeleportMode::TeleportAndReset:
-		bReset = true;
-		// fallthrough
-	case EClothingTeleportMode::Teleport:
-		bTeleport = true;
-		break;
-	}
-
+	ClothTeleportMode = bTeleport ? (bReset ? EClothingTeleportMode::TeleportAndReset : EClothingTeleportMode::Teleport) : EClothingTeleportMode::None;
+	ClothTeleportMode = UE::ClothingSimulation::TeleportHelpers::CalculateClothingTeleport(ClothTeleportMode, CurRootBoneMat, PrevRootBoneMatrix, bReset, ClothTeleportDistThresholdSquared, ClothTeleportCosineThresholdInRad);
 	PrevRootBoneMatrix = CurRootBoneMat;
 }
 
diff --git a/Engine/Plugins/ChaosClothAsset/Source/ChaosClothAssetEngine/Private/ChaosClothAsset/ClothSimulationContext.cpp b/Engine/Plugins/ChaosClothAsset/Source/ChaosClothAssetEngine/Private/ChaosClothAsset/ClothSimulationContext.cpp
--- a/Engine/Plugins/ChaosClothAsset/Source/ChaosClothAssetEngine/Private/ChaosClothAsset/ClothSimulationContext.cpp
+++ b/Engine/Plugins/ChaosClothAsset/Source/ChaosClothAssetEngine/Private/ChaosClothAsset/ClothSimulationContext.cpp
@@ -24,8 +24,11 @@ namespace UE::Chaos::ClothAsset
 		constexpr float MaxDeltaTimeTeleportMultiplierDefault = 1.5f;
 		const float MaxDeltaTimeTeleportMultiplier = CVarMaxDeltaTimeTeleportMultiplier ? CVarMaxDeltaTimeTeleportMultiplier->GetFloat() : MaxDeltaTimeTeleportMultiplierDefault;
 
-		bTeleport = (InDeltaTime > MaxDeltaTime * MaxDeltaTimeTeleportMultiplier) ? true : ClothComponent.NeedsTeleport();
-		bReset = ClothComponent.NeedsReset();
+		const EClothingTeleportMode ComponentTeleportMode = ClothComponent.GetClothTeleportMode();
+
+		bTeleport = (InDeltaTime > MaxDeltaTime * MaxDeltaTimeTeleportMultiplier) ? true : 
+			(ComponentTeleportMode >= EClothingTeleportMode::Teleport);
+		bReset = ComponentTeleportMode == EClothingTeleportMode::TeleportAndReset;
 
 		VelocityScale = (!bTeleport && !bReset && InDeltaTime > 0.f) ?
 			FMath::Min(InDeltaTime, MaxDeltaTime) / InDeltaTime :
diff --git a/Engine/Plugins/ChaosClothAsset/Source/ChaosClothAssetEngine/Public/ChaosClothAsset/ClothComponent.h b/Engine/Plugins/ChaosClothAsset/Source/ChaosClothAssetEngine/Public/ChaosClothAsset/ClothComponent.h
--- a/Engine/Plugins/ChaosClothAsset/Source/ChaosClothAssetEngine/Public/ChaosClothAsset/ClothComponent.h
+++ b/Engine/Plugins/ChaosClothAsset/Source/ChaosClothAssetEngine/Public/ChaosClothAsset/ClothComponent.h
@@ -94,12 +94,15 @@ public:
 	UFUNCTION(BlueprintCallable, Category = "ClothComponent", Meta = (Keywords = "Chaos Cloth Teleport Reset"))
 	void ForceNextUpdateTeleportAndReset() { bTeleport = bReset = true; }
 
-	/** Return whether teleport is currently requested. */
+	/** Return whether teleport is currently requested. Use GetClothTeleportMode to get teleport mode including any auto-teleport based on teleport thresholds. */
 	bool NeedsTeleport() const { return bTeleport; }
 
-	/** Return whether reseting the pose is currently requested. */
+	/** Return whether reseting the pose is currently requested. Use GetClothTeleportMode to get teleport mode including any auto-teleport based on teleport thresholds*/
 	bool NeedsReset() const { return bReset; }
 
+	/** Get currently calculated teleport mode. */
+	EClothingTeleportMode GetClothTeleportMode() const { return ClothTeleportMode; }
+
 	/** Stop the simulation, and keep the cloth in its last pose. */
 	UFUNCTION(CallInEditor, BlueprintCallable, Category = "ClothComponent", Meta = (UnsafeDuringActorConstruction, Keywords = "Chaos Cloth Simulation Suspend"))
 	void SuspendSimulation() { bSuspendSimulation = true; }
@@ -395,6 +398,9 @@ private:
 	/** previous root bone matrix to compare the difference and decide to do clothing teleport  */
 	FMatrix	PrevRootBoneMatrix;
 
+	/** Currently calculated ClothTeleport based on bTeleport, bReset as well as any teleport calculated based on TeleportDistanceThreshold and TeleportRotationThreshold */
+	EClothingTeleportMode ClothTeleportMode;
+
 	/**
 	 * Simulation properties per cloth model. 
 	 * Transient because it doesn't need to be serialized but contains pointers to UObject that needs not be garbage collected.
diff --git a/Engine/Plugins/ChaosClothAssetEditor/Source/ChaosClothAssetEditor/Private/ChaosClothAsset/ClothEditorPreviewScene.cpp b/Engine/Plugins/ChaosClothAssetEditor/Source/ChaosClothAssetEditor/Private/ChaosClothAsset/ClothEditorPreviewScene.cpp
--- a/Engine/Plugins/ChaosClothAssetEditor/Source/ChaosClothAssetEditor/Private/ChaosClothAsset/ClothEditorPreviewScene.cpp
+++ b/Engine/Plugins/ChaosClothAssetEditor/Source/ChaosClothAssetEditor/Private/ChaosClothAsset/ClothEditorPreviewScene.cpp
@@ -299,6 +299,30 @@ void FChaosClothPreviewScene::SceneDescriptionPropertyChanged(const FName& Prope
 			ClothComponent->SetClothGeometryScale(PreviewSceneDescription->SolverGeometryScale);
 		}
 	}
+
+	if (PropertyName == GET_MEMBER_NAME_CHECKED(UChaosClothPreviewSceneDescription, SolverGeometryScale))
+	{
+		if (ClothComponent)
+		{
+			ClothComponent->SetClothGeometryScale(PreviewSceneDescription->SolverGeometryScale);
+		}
+	}
+
+	if (PropertyName == GET_MEMBER_NAME_CHECKED(UChaosClothPreviewSceneDescription, TeleportDistanceThreshold))
+	{
+		if (ClothComponent)
+		{
+			ClothComponent->SetTeleportDistanceThreshold(PreviewSceneDescription->TeleportDistanceThreshold);
+		}
+	}
+
+	if (PropertyName == GET_MEMBER_NAME_CHECKED(UChaosClothPreviewSceneDescription, TeleportRotationThreshold))
+	{
+		if (ClothComponent)
+		{
+			ClothComponent->SetTeleportRotationThreshold(PreviewSceneDescription->TeleportRotationThreshold);
+		}
+	}
 }
 
 UChaosClothComponent* FChaosClothPreviewScene::GetClothComponent()
diff --git a/Engine/Plugins/ChaosClothAssetEditor/Source/ChaosClothAssetEditor/Private/ChaosClothAsset/ClothEditorPreviewScene.h b/Engine/Plugins/ChaosClothAssetEditor/Source/ChaosClothAssetEditor/Private/ChaosClothAsset/ClothEditorPreviewScene.h
--- a/Engine/Plugins/ChaosClothAssetEditor/Source/ChaosClothAssetEditor/Private/ChaosClothAsset/ClothEditorPreviewScene.h
+++ b/Engine/Plugins/ChaosClothAssetEditor/Source/ChaosClothAssetEditor/Private/ChaosClothAsset/ClothEditorPreviewScene.h
@@ -65,6 +65,21 @@ public:
 	UPROPERTY(EditAnywhere, Transient, Category = "ClothComponent", Meta = (UIMin = 0.0, UIMax = 10.0, ClampMin = 0.0, ClampMax = 10000.0))
 	float SolverGeometryScale = 1.f;
 
+	/**
+	* Conduct teleportation if the character's movement is greater than this threshold in 1 frame.
+	* Zero or negative values will skip the check.
+	*/
+	UPROPERTY(EditAnywhere, Transient, Category = ClothComponent)
+	float TeleportDistanceThreshold = 0.f;
+
+	/**
+	* Rotation threshold in degrees, ranging from 0 to 180.
+	* Conduct teleportation if the character's rotation is greater than this threshold in 1 frame.
+	* Zero or negative values will skip the check.
+	*/
+	UPROPERTY(EditAnywhere, Transient, Category = ClothComponent)
+	float TeleportRotationThreshold = 0.f;
+
 	UPROPERTY(Transient)
 	bool bValidSelectionForTransform = false;
 
