Add LLM instrumentation for Chaos Cloth.

diff --git a/Engine/Plugins/ChaosClothAsset/Source/ChaosClothAssetEngine/Private/ChaosClothAsset/ClothAsset.cpp b/Engine/Plugins/ChaosClothAsset/Source/ChaosClothAssetEngine/Private/ChaosClothAsset/ClothAsset.cpp
--- a/Engine/Plugins/ChaosClothAsset/Source/ChaosClothAssetEngine/Private/ChaosClothAsset/ClothAsset.cpp
+++ b/Engine/Plugins/ChaosClothAsset/Source/ChaosClothAssetEngine/Private/ChaosClothAsset/ClothAsset.cpp
@@ -176,6 +176,7 @@ PRAGMA_ENABLE_DEPRECATION_WARNINGS
 
 void UChaosClothAsset::Serialize(FArchive& Ar)
 {
+	LLM_SCOPE_BYNAME(TEXT("Physics/Cloth"));
 	Super::Serialize(Ar);
 
 	Ar.UsingCustomVersion(FFortniteMainBranchObjectVersion::GUID);
@@ -203,11 +204,14 @@ void UChaosClothAsset::Serialize(FArchive& Ar)
 
 	if (bCooked && !IsTemplate() && !Ar.IsCountingMemory())  // Counting of these resources are done in GetResourceSizeEx, so skip these when counting memory
 	{
-		if (Ar.IsLoading())
 		{
-			SetResourceForRendering(MakeUnique<FSkeletalMeshRenderData>());
+			LLM_SCOPE_BYNAME(TEXT("Physics/ClothRendering"));
+			if (Ar.IsLoading())
+			{
+				SetResourceForRendering(MakeUnique<FSkeletalMeshRenderData>());
+			}
+			GetResourceForRendering()->Serialize(Ar, this);
 		}
-		GetResourceForRendering()->Serialize(Ar, this);
 
 		if (!ClothSimulationModel.IsValid())
 		{
diff --git a/Engine/Plugins/ChaosClothAsset/Source/ChaosClothAssetEngine/Private/ChaosClothAsset/ClothSimulationProxy.cpp b/Engine/Plugins/ChaosClothAsset/Source/ChaosClothAssetEngine/Private/ChaosClothAsset/ClothSimulationProxy.cpp
--- a/Engine/Plugins/ChaosClothAsset/Source/ChaosClothAssetEngine/Private/ChaosClothAsset/ClothSimulationProxy.cpp
+++ b/Engine/Plugins/ChaosClothAsset/Source/ChaosClothAssetEngine/Private/ChaosClothAsset/ClothSimulationProxy.cpp
@@ -152,6 +152,7 @@ namespace UE::Chaos::ClothAsset
 
 	void FClothSimulationProxy::BeginInitialization_GameThread()
 	{
+		LLM_SCOPE_BYNAME(TEXT("Physics/Cloth"));
 		using namespace ::Chaos;
 
 		check(IsInGameThread());
@@ -222,6 +223,7 @@ PRAGMA_ENABLE_DEPRECATION_WARNINGS
 
 	void FClothSimulationProxy::ExecuteInitialization()
 	{
+		LLM_SCOPE_BYNAME(TEXT("Physics/Cloth"));
 		int32 LocalNumCloths = 0;
 		int32 LocalNumKinematicParticles = 0;
 		int32 LocalNumDynamicParticles = 0;
