Chaos Cloth Asset - Make the GetCollection method return a const collection.

#rnx

diff --git a/Engine/Plugins/ChaosCloth/Source/ChaosCloth/Private/ChaosCloth/ChaosClothingSimulationConfig.cpp b/Engine/Plugins/ChaosCloth/Source/ChaosCloth/Private/ChaosCloth/ChaosClothingSimulationConfig.cpp
--- a/Engine/Plugins/ChaosCloth/Source/ChaosCloth/Private/ChaosCloth/ChaosClothingSimulationConfig.cpp
+++ b/Engine/Plugins/ChaosCloth/Source/ChaosCloth/Private/ChaosCloth/ChaosClothingSimulationConfig.cpp
@@ -12,23 +12,9 @@ namespace Chaos
 	{
 	}
 
-	FClothingSimulationConfig::FClothingSimulationConfig(const TArray<TSharedPtr<FManagedArrayCollection>>& InPropertyCollections)
-	{
-		Initialize(InPropertyCollections);
-	}
-
-	FClothingSimulationConfig::FClothingSimulationConfig(std::initializer_list<TSharedPtr<FManagedArrayCollection>> InitList)
-	{
-PRAGMA_DISABLE_DEPRECATION_WARNINGS
-		Initialize(InitList);
-PRAGMA_ENABLE_DEPRECATION_WARNINGS
-	}
-
 	FClothingSimulationConfig::FClothingSimulationConfig(const TArray<TSharedPtr<const FManagedArrayCollection>>& InPropertyCollections)
 	{
-PRAGMA_DISABLE_DEPRECATION_WARNINGS
 		Initialize(InPropertyCollections);
-PRAGMA_ENABLE_DEPRECATION_WARNINGS
 	}
 
 	FClothingSimulationConfig::~FClothingSimulationConfig() = default;
@@ -254,34 +240,6 @@ PRAGMA_ENABLE_DEPRECATION_WARNINGS
 		Property->AddValue(TEXT("UseLegacyConfig"), bUseLegacyConfig, NonAnimatablePropertyFlags);
 	}
 
-	void FClothingSimulationConfig::Initialize(const TArray<TSharedPtr<FManagedArrayCollection>>& InPropertyCollections)
-	{
-		PropertyCollections.Reset(InPropertyCollections.Num());
-		Properties.Reset(InPropertyCollections.Num());
-
-		for (const TSharedPtr<FManagedArrayCollection>& InPropertyCollection : InPropertyCollections)
-		{
-			TSharedPtr<FManagedArrayCollection>& PropertyCollection = PropertyCollections.Add_GetRef(MakeShared<FManagedArrayCollection>());
-			TUniquePtr<Softs::FCollectionPropertyMutableFacade>& Property = Properties.Add_GetRef(MakeUnique<Softs::FCollectionPropertyMutableFacade>(PropertyCollection));
-			Property->Copy(*InPropertyCollection);
-		}
-	}
-
-	// UE_DEPRECATED(5.6, "Use Initialize(const TArray<TSharedPtr<FManagedArrayCollection>>&).")
-	CHAOSCLOTH_API void FClothingSimulationConfig::Initialize(std::initializer_list<TSharedPtr<FManagedArrayCollection>> InitList)
-	{
-		PropertyCollections.Reset((int32)InitList.size());
-		Properties.Reset((int32)InitList.size());
-
-		for (const TSharedPtr<FManagedArrayCollection>* It = InitList.begin(); It != InitList.end(); ++It)
-		{
-			TSharedPtr<FManagedArrayCollection>& PropertyCollection = PropertyCollections.Add_GetRef(MakeShared<FManagedArrayCollection>());
-			TUniquePtr<Softs::FCollectionPropertyMutableFacade>& Property = Properties.Add_GetRef(MakeUnique<Softs::FCollectionPropertyMutableFacade>(PropertyCollection));
-			Property->Copy(**It);
-		}
-	}
-
-	// UE_DEPRECATED(5.6, "Use Initialize(const TArray<TSharedPtr<FManagedArrayCollection>>&).")
 	void FClothingSimulationConfig::Initialize(const TArray<TSharedPtr<const FManagedArrayCollection>>& InPropertyCollections)
 	{
 		PropertyCollections.Reset(InPropertyCollections.Num());
diff --git a/Engine/Plugins/ChaosCloth/Source/ChaosCloth/Public/ChaosCloth/ChaosClothingSimulationConfig.h b/Engine/Plugins/ChaosCloth/Source/ChaosCloth/Public/ChaosCloth/ChaosClothingSimulationConfig.h
--- a/Engine/Plugins/ChaosCloth/Source/ChaosCloth/Public/ChaosCloth/ChaosClothingSimulationConfig.h
+++ b/Engine/Plugins/ChaosCloth/Source/ChaosCloth/Public/ChaosCloth/ChaosClothingSimulationConfig.h
@@ -22,10 +22,6 @@ namespace Chaos
 	{
 	public:
 		CHAOSCLOTH_API FClothingSimulationConfig();
-		explicit CHAOSCLOTH_API FClothingSimulationConfig(const TArray<TSharedPtr<FManagedArrayCollection>>& InPropertyCollections);
-		UE_DEPRECATED(5.6, "Use FClothingSimulationConfig(const TArray<TSharedPtr<FManagedArrayCollection>>&) instead.")
-		explicit CHAOSCLOTH_API FClothingSimulationConfig(std::initializer_list<TSharedPtr<FManagedArrayCollection>> InitList);
-		UE_DEPRECATED(5.6, "Use FClothingSimulationConfig(const TArray<TSharedPtr<FManagedArrayCollection>>&) instead.")
 		CHAOSCLOTH_API FClothingSimulationConfig(const TArray<TSharedPtr<const FManagedArrayCollection>>& InPropertyCollections);
 
 		CHAOSCLOTH_API ~FClothingSimulationConfig();
@@ -44,10 +40,6 @@ namespace Chaos
 		CHAOSCLOTH_API void Initialize(const UChaosClothConfig* ClothConfig, const UChaosClothSharedSimConfig* ClothSharedConfig, bool bUseLegacyConfig = false);
 
 		/** Initialize config from an array of property collections (one per LOD). */
-		CHAOSCLOTH_API void Initialize(const TArray<TSharedPtr<FManagedArrayCollection>>& InPropertyCollections);
-		UE_DEPRECATED(5.6, "Use Initialize(const TArray<TSharedPtr<FManagedArrayCollection>>&) instead.")
-		CHAOSCLOTH_API void Initialize(std::initializer_list<TSharedPtr<FManagedArrayCollection>> InitList);
-		UE_DEPRECATED(5.6, "Use Initialize(const TArray<TSharedPtr<FManagedArrayCollection>>&) instead.")
 		CHAOSCLOTH_API void Initialize(const TArray<TSharedPtr<const FManagedArrayCollection>>& InPropertyCollections);
 
 		/** Return a property collection facade for reading properties from this configuration. */
diff --git a/Engine/Plugins/ChaosClothAsset/Source/ChaosClothAssetEngine/Private/ChaosClothAsset/ClothAsset.cpp b/Engine/Plugins/ChaosClothAsset/Source/ChaosClothAssetEngine/Private/ChaosClothAsset/ClothAsset.cpp
--- a/Engine/Plugins/ChaosClothAsset/Source/ChaosClothAssetEngine/Private/ChaosClothAsset/ClothAsset.cpp
+++ b/Engine/Plugins/ChaosClothAsset/Source/ChaosClothAssetEngine/Private/ChaosClothAsset/ClothAsset.cpp
@@ -37,116 +37,120 @@
 
 namespace UE::Chaos::ClothAsset::Private
 {
-bool bClothCollectionOnlyCookRequiredFacades = true;
-FAutoConsoleVariableRef CVarClothCollectionOnlyCookRequiredFacades(
-	TEXT("p.ClothCollectionOnlyCookRequiredFacades"),
-	bClothCollectionOnlyCookRequiredFacades,
-	TEXT("Default setting for culling managed arrays on the cloth collection during the cook. Default[true]"));
+	static bool bClothCollectionOnlyCookRequiredFacades = true;
 
-::Chaos::FChaosArchive& Serialize(::Chaos::FChaosArchive& Ar, TArray<TSharedRef<FManagedArrayCollection>>& ClothCollections)
-{
-	Ar.UsingCustomVersion(FUE5MainStreamObjectVersion::GUID);
+	FAutoConsoleVariableRef CVarClothCollectionOnlyCookRequiredFacades(
+		TEXT("p.ClothCollectionOnlyCookRequiredFacades"),
+		bClothCollectionOnlyCookRequiredFacades,
+		TEXT("Default setting for culling managed arrays on the cloth collection during the cook. Default[true]"));
 
-	if (Ar.IsLoading() && Ar.CustomVer(FUE5MainStreamObjectVersion::GUID) < FUE5MainStreamObjectVersion::ClothCollectionSingleLodSchema)
+	::Chaos::FChaosArchive& Serialize(::Chaos::FChaosArchive& Ar, TArray<TSharedRef<const FManagedArrayCollection>>& ClothCollections)
 	{
-		// Cloth assets before this version had a single ClothCollection with a completely different schema.
-		ClothCollections.Empty(1);
-		TSharedRef<FManagedArrayCollection>& ClothCollection = ClothCollections.Emplace_GetRef(MakeShared<FManagedArrayCollection>());
-		ClothCollection->Serialize(Ar);
-
-		// Now we're just going to hard reset and define a new schema.
-		ClothCollection->Reset();
-		FCollectionClothFacade ClothFacade(ClothCollection);
-		ClothFacade.DefineSchema();
+		Ar.UsingCustomVersion(FUE5MainStreamObjectVersion::GUID);
 
-		return Ar;
-	}
-	else
-	{
-		// This is following Serialize for Arrays
-		ClothCollections.CountBytes(Ar);
-		int32 SerializeNum = Ar.IsLoading() ? 0 : ClothCollections.Num();
-		Ar << SerializeNum;
-		if (SerializeNum == 0)
+		if (Ar.IsLoading() && Ar.CustomVer(FUE5MainStreamObjectVersion::GUID) < FUE5MainStreamObjectVersion::ClothCollectionSingleLodSchema)
 		{
-			// if we are loading, then we have to reset the size to 0, in case it isn't currently 0
-			if (Ar.IsLoading())
-			{
-				ClothCollections.Empty();
-			}
-			return Ar;
-		}
-		check(SerializeNum >= 0);
+			// Cloth assets before this version had a single ClothCollection with a completely different schema.
+			ClothCollections.Empty(1);
+			TSharedRef<FManagedArrayCollection> ClothCollection = MakeShared<FManagedArrayCollection>();
+			ClothCollection->Serialize(Ar);
+
+			// Now we're just going to hard reset and define a new schema.
+			ClothCollection->Reset();
+			FCollectionClothFacade ClothFacade(ClothCollection);
+			ClothFacade.DefineSchema();
+
+			ClothCollections.Emplace(MoveTemp(ClothCollection));
 
-		if (Ar.IsError() || SerializeNum < 0)
-		{
-			Ar.SetError();
 			return Ar;
 		}
-		if (Ar.IsLoading())
+		else
 		{
-			// Required for resetting ArrayNum
-			ClothCollections.Empty(SerializeNum);
-
-			for (int32 i = 0; i < SerializeNum; i++)
+			// This is following Serialize for Arrays
+			ClothCollections.CountBytes(Ar);
+			int32 SerializeNum = Ar.IsLoading() ? 0 : ClothCollections.Num();
+			Ar << SerializeNum;
+			if (SerializeNum == 0)
 			{
-				TSharedRef<FManagedArrayCollection>& ClothCollection = ClothCollections.Emplace_GetRef(MakeShared<FManagedArrayCollection>());
-				ClothCollection->Serialize(Ar);
+				// if we are loading, then we have to reset the size to 0, in case it isn't currently 0
+				if (Ar.IsLoading())
+				{
+					ClothCollections.Empty();
+				}
+				return Ar;
 			}
-		}
-		else
-		{
-			check(SerializeNum == ClothCollections.Num());
+			check(SerializeNum >= 0);
 
-			for (int32 i = 0; i < SerializeNum; i++)
+			if (Ar.IsError() || SerializeNum < 0)
 			{
-				ClothCollections[i]->Serialize(Ar);
+				Ar.SetError();
+				return Ar;
 			}
-		}
+			if (Ar.IsLoading())
+			{
+				// Required for resetting ArrayNum
+				ClothCollections.Empty(SerializeNum);
+
+				for (int32 i = 0; i < SerializeNum; i++)
+				{
+					TSharedRef<FManagedArrayCollection> ClothCollection = MakeShared<FManagedArrayCollection>();
+					ClothCollection->Serialize(Ar);
+					ClothCollections.Emplace(MoveTemp(ClothCollection));
+				}
+			}
+			else
+			{
+				check(SerializeNum == ClothCollections.Num());
 
-		return Ar;
+				for (int32 i = 0; i < SerializeNum; i++)
+				{
+					ConstCastSharedRef<FManagedArrayCollection>(ClothCollections[i])->Serialize(Ar);
+				}
+			}
+
+			return Ar;
+		}
 	}
-}
 
-TArray<TSharedRef<FManagedArrayCollection>> TrimOnCook(const FString InAssetName, const TArray<TSharedRef<FManagedArrayCollection>>& InClothCollections)
-{
-	int32 Index = 0;
-#if WITH_EDITORONLY_DATA
-	if (bClothCollectionOnlyCookRequiredFacades)
+	TArray<TSharedRef<const FManagedArrayCollection>> TrimOnCook(const FString InAssetName, const TArray<TSharedRef<const FManagedArrayCollection>>& InClothCollections)
 	{
-		TArray<TSharedRef<FManagedArrayCollection>> OutputCollections;
-		for (TSharedRef<FManagedArrayCollection> ClothCollection : InClothCollections)
+		int32 Index = 0;
+#if WITH_EDITORONLY_DATA
+		if (bClothCollectionOnlyCookRequiredFacades)
 		{
-			// Properties
-			TSharedPtr<FManagedArrayCollection> PropertyCollection(new FManagedArrayCollection());
-			::Chaos::Softs::FCollectionPropertyMutableFacade CollectionPropertyMutableFacade(PropertyCollection);
-			CollectionPropertyMutableFacade.Copy(*ClothCollection);
-
-			// Springs
-			const ::Chaos::Softs::FEmbeddedSpringFacade InEmbeddedSpringFacade(ClothCollection.Get(), UE::Chaos::ClothAsset::ClothCollectionGroup::SimVertices3D);
-			if (InEmbeddedSpringFacade.IsValid())
+			TArray<TSharedRef<const FManagedArrayCollection>> OutputCollections;
+			for (TSharedRef<const FManagedArrayCollection> ClothCollection : InClothCollections)
 			{
-				::Chaos::Softs::FEmbeddedSpringFacade EmbeddedSpringFacade(*PropertyCollection, ClothCollectionGroup::SimVertices3D);
-				EmbeddedSpringFacade.DefineSchema();
-				constexpr int32 VertexOffset = 0;
-				EmbeddedSpringFacade.Append(InEmbeddedSpringFacade, VertexOffset);
+				// Properties
+				TSharedRef<FManagedArrayCollection> PropertyCollection = MakeShared<FManagedArrayCollection>();
+				::Chaos::Softs::FCollectionPropertyMutableFacade CollectionPropertyMutableFacade(PropertyCollection);
+				CollectionPropertyMutableFacade.Copy(*ClothCollection);
+
+				// Springs
+				const ::Chaos::Softs::FEmbeddedSpringFacade InEmbeddedSpringFacade(ClothCollection.Get(), UE::Chaos::ClothAsset::ClothCollectionGroup::SimVertices3D);
+				if (InEmbeddedSpringFacade.IsValid())
+				{
+					::Chaos::Softs::FEmbeddedSpringFacade EmbeddedSpringFacade(*PropertyCollection, ClothCollectionGroup::SimVertices3D);
+					EmbeddedSpringFacade.DefineSchema();
+					constexpr int32 VertexOffset = 0;
+					EmbeddedSpringFacade.Append(InEmbeddedSpringFacade, VertexOffset);
+				}
+
+				OutputCollections.Emplace(MoveTemp(PropertyCollection));
+
+				UE_LOG(LogChaosClothAsset, Display, TEXT("TrimOnCook[ON] %s:[%d] [size:%d]"), 
+					*InAssetName, Index++, PropertyCollection->GetAllocatedSize());
 			}
-
-			OutputCollections.Add(TSharedRef<FManagedArrayCollection>(PropertyCollection.ToSharedRef()));
-			UE_LOG(LogChaosClothAsset, Display, TEXT("TrimOnCook[ON] %s:[%d] [size:%d]"),
-				*InAssetName, Index++, PropertyCollection->GetAllocatedSize());
+			return OutputCollections;
 		}
-		return OutputCollections;
-	}
 #endif
-	for (TSharedRef<FManagedArrayCollection> ClothCollection : InClothCollections)
-	{
-		UE_LOG(LogChaosClothAsset, Display, TEXT("TrimOnCook [OFF] %s:[%d] [size:%d]"),
-			*InAssetName, Index++, ClothCollection->GetAllocatedSize());
+		for (TSharedRef<const FManagedArrayCollection> ClothCollection : InClothCollections)
+		{
+			UE_LOG(LogChaosClothAsset, Display, TEXT("TrimOnCook [OFF] %s:[%d] [size:%d]"),
+				*InAssetName, Index++, ClothCollection->GetAllocatedSize());
+		}
+		return InClothCollections;
 	}
-	return InClothCollections;
-}
-
 }
 
 UChaosClothAsset::UChaosClothAsset(const FObjectInitializer& ObjectInitializer)
@@ -160,9 +164,10 @@ PRAGMA_DISABLE_DEPRECATION_WARNINGS
 PRAGMA_ENABLE_DEPRECATION_WARNINGS
 
 	// Setup a single LOD's Cloth Collection
-	TSharedRef<FManagedArrayCollection>& ClothCollection = GetClothCollections().Emplace_GetRef(MakeShared<FManagedArrayCollection>());
+	TSharedRef<FManagedArrayCollection> ClothCollection = MakeShared<FManagedArrayCollection>();
 	UE::Chaos::ClothAsset::FCollectionClothFacade ClothFacade(ClothCollection);
 	ClothFacade.DefineSchema();
+	GetClothCollections().Emplace(MoveTemp(ClothCollection));
 }
 
 UChaosClothAsset::UChaosClothAsset(FVTableHelper& Helper)
@@ -186,7 +191,7 @@ void UChaosClothAsset::Serialize(FArchive& Ar)
 
 	if (bCooked && Ar.IsSaving())
 	{
-		TArray<TSharedRef<FManagedArrayCollection>> OutputCollections = 
+		TArray<TSharedRef<const FManagedArrayCollection>> OutputCollections = 
 			UE::Chaos::ClothAsset::Private::TrimOnCook(GetPathName(), GetClothCollections());
 		Chaos::FChaosArchive ChaosArchive(Ar);
 		UE::Chaos::ClothAsset::Private::Serialize(ChaosArchive, OutputCollections);
@@ -322,6 +327,8 @@ void UChaosClothAsset::BeginPostLoadInternal(FSkinnedAssetPostLoadContext& Conte
 #if WITH_EDITOR
 	TRACE_CPUPROFILER_EVENT_SCOPE(UChaosClothAsset::BeginPostLoadInternal);
 
+	using namespace UE::Chaos::ClothAsset;
+
 	checkf(IsInGameThread(), TEXT("Cannot execute function UChaosClothAsset::BeginPostLoadInternal asynchronously. Asset: %s"), *GetFullName());
 	SetInternalFlags(EInternalObjectFlags::Async);
 
@@ -336,9 +343,10 @@ void UChaosClothAsset::BeginPostLoadInternal(FSkinnedAssetPostLoadContext& Conte
 	if (GetClothCollections().IsEmpty())
 	{
 		UE_LOG(LogChaosClothAsset, Warning, TEXT("Invalid Cloth Collection (no LODs) found while loading Cloth Asset %s."), *GetFullName());
-		TSharedRef<FManagedArrayCollection>& ClothCollection = GetClothCollections().Emplace_GetRef(MakeShared<FManagedArrayCollection>());
-		UE::Chaos::ClothAsset::FCollectionClothFacade ClothFacade(ClothCollection);
+		TSharedRef<FManagedArrayCollection> ClothCollection = MakeShared<FManagedArrayCollection>();
+		FCollectionClothFacade ClothFacade(ClothCollection);
 		ClothFacade.DefineSchema();
+		GetClothCollections().Emplace(MoveTemp(ClothCollection));
 		bAnyInvalidLods = true;
 	}
 
@@ -349,19 +357,20 @@ void UChaosClothAsset::BeginPostLoadInternal(FSkinnedAssetPostLoadContext& Conte
 	bool bAnyInvalidSkeletons = false;
 	for (int32 LODIndex = 0; LODIndex < NumLods; ++LODIndex)
 	{
-		TSharedRef<FManagedArrayCollection>& ClothCollection = GetClothCollections()[LODIndex];
+		const TSharedRef<const FManagedArrayCollection>& ClothCollection = GetClothCollections()[LODIndex];
 
-		UE::Chaos::ClothAsset::FCollectionClothFacade ClothFacade(ClothCollection);
-		if (!ClothFacade.IsValid())
+		const FCollectionClothConstFacade ClothConstFacade(ClothCollection);
+		if (!ClothConstFacade.IsValid())
 		{
 			UE_LOG(LogChaosClothAsset, Warning, TEXT("Invalid Cloth Collection found at LOD %i while loading Cloth Asset %s."), LODIndex, *GetFullName());
-			ClothCollection = MakeShared<FManagedArrayCollection>();
-			ClothFacade = UE::Chaos::ClothAsset::FCollectionClothFacade(ClothCollection);
-			ClothFacade.DefineSchema();
+			TSharedRef<FManagedArrayCollection> NewClothCollection = MakeShared<FManagedArrayCollection>();
+			FCollectionClothFacade NewClothFacade = FCollectionClothFacade(NewClothCollection);
+			NewClothFacade.DefineSchema();
+			GetClothCollections()[LODIndex] = MoveTemp(NewClothCollection);
 			bAnyInvalidLods = true;
 			bAnyInvalidSkeletons = true;
 		}
-		else if (ClothFacade.GetSkeletalMeshPathName().IsEmpty())
+		else if (ClothConstFacade.GetSkeletalMeshPathName().IsEmpty())
 		{
 			bAnyInvalidSkeletons = true;
 		}
@@ -552,7 +561,7 @@ FString UChaosClothAsset::GetAsyncPropertyName(uint64 Property) const
 	return StaticEnum<EClothAssetAsyncProperties>()->GetNameByValue(Property).ToString();
 }
 
-const TArray<TSharedRef<FManagedArrayCollection>>& UChaosClothAsset::GetCollections(int32) const
+const TArray<TSharedRef<const FManagedArrayCollection>>& UChaosClothAsset::GetCollections(int32 /*ModelIndex*/) const
 {
 	WaitUntilAsyncPropertyReleased(EClothAssetAsyncProperties::ClothCollection);
 PRAGMA_DISABLE_DEPRECATION_WARNINGS
@@ -669,10 +678,12 @@ void UChaosClothAsset::SetPhysicsAsset(UPhysicsAsset* InPhysicsAsset)
 	using namespace UE::Chaos::ClothAsset;
 
 	PhysicsAsset = InPhysicsAsset;
-	for (TSharedRef<FManagedArrayCollection>& ClothCollection : GetClothCollections())
+	for (TSharedRef<const FManagedArrayCollection>& ClothCollection : GetClothCollections())
 	{
-		FCollectionClothFacade Cloth(ClothCollection);
-		Cloth.SetPhysicsAssetPathName(PhysicsAsset ? PhysicsAsset->GetPathName() : FString());
+		TSharedRef<FManagedArrayCollection> NewClothCollection = MakeShared<FManagedArrayCollection>(*ClothCollection);
+		FCollectionClothFacade ClothFacade(NewClothCollection);
+		ClothFacade.SetPhysicsAssetPathName(PhysicsAsset ? PhysicsAsset->GetPathName() : FString());
+		ClothCollection = MoveTemp(NewClothCollection);
 	}
 }
 
@@ -691,9 +702,11 @@ void UChaosClothAsset::SetReferenceSkeleton(const FReferenceSkeleton* ReferenceS
 	// Rebind the meshes
 	if (bRebindMeshes)
 	{
-		for (TSharedRef<FManagedArrayCollection>& ClothCollection : GetClothCollections())
+		for (TSharedRef<const FManagedArrayCollection>& ClothCollection : GetClothCollections())
 		{
-			FClothGeometryTools::BindMeshToRootBone(ClothCollection, true, true);
+			TSharedRef<FManagedArrayCollection> NewClothCollection = MakeShared<FManagedArrayCollection>(*ClothCollection);
+			FClothGeometryTools::BindMeshToRootBone(NewClothCollection, true, true);
+			ClothCollection = MoveTemp(NewClothCollection);
 		}
 	}
 
@@ -733,11 +746,13 @@ void UChaosClothAsset::CopySimMeshToRenderMesh(UMaterialInterface* Material)
 		FString(TEXT("/Engine/EditorMaterials/Cloth/CameraLitDoubleSided.CameraLitDoubleSided"));
 
 	bool bAnyLodHasRenderMesh = false;
-	for (TSharedRef<FManagedArrayCollection>& ClothCollection : GetClothCollections())
+	for (TSharedRef<const FManagedArrayCollection>& ClothCollection : GetClothCollections())
 	{
+		TSharedRef<FManagedArrayCollection> NewClothCollection = MakeShared<FManagedArrayCollection>(*ClothCollection);
 		constexpr bool bSingleRenderPattern = true;
-		FClothGeometryTools::CopySimMeshToRenderMesh(ClothCollection, RenderMaterialPathName, bSingleRenderPattern);
-		bAnyLodHasRenderMesh = bAnyLodHasRenderMesh || FClothGeometryTools::HasRenderMesh(ClothCollection);
+		FClothGeometryTools::CopySimMeshToRenderMesh(NewClothCollection, RenderMaterialPathName, bSingleRenderPattern);
+		bAnyLodHasRenderMesh = bAnyLodHasRenderMesh || FClothGeometryTools::HasRenderMesh(NewClothCollection);
+		ClothCollection = MoveTemp(NewClothCollection);
 	}
 
 	// Set new material
diff --git a/Engine/Plugins/ChaosClothAsset/Source/ChaosClothAssetEngine/Private/ChaosClothAsset/ClothComponent.cpp b/Engine/Plugins/ChaosClothAsset/Source/ChaosClothAssetEngine/Private/ChaosClothAsset/ClothComponent.cpp
--- a/Engine/Plugins/ChaosClothAsset/Source/ChaosClothAssetEngine/Private/ChaosClothAsset/ClothComponent.cpp
+++ b/Engine/Plugins/ChaosClothAsset/Source/ChaosClothAssetEngine/Private/ChaosClothAsset/ClothComponent.cpp
@@ -20,17 +20,18 @@
 
 CSV_DECLARE_CATEGORY_MODULE_EXTERN(ENGINE_API, Animation);
 
-void FChaosClothSimulationProperties::Initialize(const TArray<TSharedRef<FManagedArrayCollection>>& AssetCollections)
+void FChaosClothSimulationProperties::Initialize(const TArray<TSharedRef<const FManagedArrayCollection>>& AssetCollections)
 {
 	using namespace ::Chaos::Softs;
 	PropertyCollections.Reset(AssetCollections.Num());
 	CollectionPropertyFacades.Reset(AssetCollections.Num());
-	for (const TSharedRef<FManagedArrayCollection>& AssetCollection : AssetCollections)
+	for (const TSharedRef<const FManagedArrayCollection>& AssetCollection : AssetCollections)
 	{
-		TSharedPtr<FManagedArrayCollection>& PropertyCollection = PropertyCollections.Emplace_GetRef(MakeShared<FManagedArrayCollection>());
+		TSharedPtr<FManagedArrayCollection> PropertyCollection = MakeShared<FManagedArrayCollection>();
 		TSharedRef<FCollectionPropertyMutableFacade> CollectionPropertyMutableFacade = MakeShared<FCollectionPropertyMutableFacade>(PropertyCollection);
 		CollectionPropertyMutableFacade->Copy(*AssetCollection);  // This only copy properties from the asset's collection (in case it also contains other groups)
 		CollectionPropertyFacades.Emplace(StaticCastSharedRef<FCollectionPropertyFacade>(CollectionPropertyMutableFacade));  // Mutable facades conveniently inherits from non mutable ones
+		PropertyCollections.Emplace(MoveTemp(PropertyCollection));
 	}
 	ClothOutfitInteractor = NewObject<UChaosClothAssetInteractor>();
 	ClothOutfitInteractor->SetProperties(CollectionPropertyFacades);
@@ -136,8 +137,8 @@ void UChaosClothComponent::UpdateConfigProperties()
 		{
 			for (int32 ModelIndex = 0; ModelIndex < ClothSimulationProperties.Num(); ++ModelIndex)
 			{
-				const TArray<TSharedRef<FManagedArrayCollection>>& AssetCollections = GetAsset()->GetCollections(ModelIndex);
-				TArray<TSharedPtr<FManagedArrayCollection>>& PropertyCollections = ClothSimulationProperties[ModelIndex].PropertyCollections;
+				const TArray<TSharedRef<const FManagedArrayCollection>>& AssetCollections = GetAsset()->GetCollections(ModelIndex);
+				TArray<TSharedPtr<const FManagedArrayCollection>>& PropertyCollections = ClothSimulationProperties[ModelIndex].PropertyCollections;
 				if (AssetCollections.Num() == PropertyCollections.Num())
 				{
 					TArray<TSharedPtr<FCollectionPropertyFacade>>& CollectionPropertyFacades = ClothSimulationProperties[ModelIndex].CollectionPropertyFacades;
diff --git a/Engine/Plugins/ChaosClothAsset/Source/ChaosClothAssetEngine/Private/ChaosClothAsset/ClothSimulationMesh.cpp b/Engine/Plugins/ChaosClothAsset/Source/ChaosClothAssetEngine/Private/ChaosClothAsset/ClothSimulationMesh.cpp
--- a/Engine/Plugins/ChaosClothAsset/Source/ChaosClothAssetEngine/Private/ChaosClothAsset/ClothSimulationMesh.cpp
+++ b/Engine/Plugins/ChaosClothAsset/Source/ChaosClothAssetEngine/Private/ChaosClothAsset/ClothSimulationMesh.cpp
@@ -12,7 +12,7 @@ namespace UE::Chaos::ClothAsset
 	FClothSimulationMesh::FClothSimulationMesh(
 		const FChaosClothSimulationModel& InClothSimulationModel,
 		const FClothSimulationContext& InClothSimulationContext,
-		const TArray<TSharedRef<FManagedArrayCollection>>& InManagedArrayCollections,
+		const TArray<TSharedRef<const FManagedArrayCollection>>& InManagedArrayCollections,
 		const FString& DebugName)
 		: ::Chaos::FClothingSimulationMesh(DebugName)
 		, ClothSimulationModel(InClothSimulationModel)
diff --git a/Engine/Plugins/ChaosClothAsset/Source/ChaosClothAssetEngine/Private/ChaosClothAsset/ClothSimulationMesh.h b/Engine/Plugins/ChaosClothAsset/Source/ChaosClothAssetEngine/Private/ChaosClothAsset/ClothSimulationMesh.h
--- a/Engine/Plugins/ChaosClothAsset/Source/ChaosClothAssetEngine/Private/ChaosClothAsset/ClothSimulationMesh.h
+++ b/Engine/Plugins/ChaosClothAsset/Source/ChaosClothAssetEngine/Private/ChaosClothAsset/ClothSimulationMesh.h
@@ -16,7 +16,7 @@ namespace UE::Chaos::ClothAsset
 		FClothSimulationMesh(
 			const FChaosClothSimulationModel& InClothSimulationModel,
 			const FClothSimulationContext& InClothSimulationContext,
-			const TArray<TSharedRef<FManagedArrayCollection>>& InManagedArrayCollections,
+			const TArray<TSharedRef<const FManagedArrayCollection>>& InManagedArrayCollections,
 			const FString& DebugName);
 
 		virtual ~FClothSimulationMesh() override = default;
@@ -62,6 +62,6 @@ namespace UE::Chaos::ClothAsset
 	private:
 		const FChaosClothSimulationModel& ClothSimulationModel;
 		const FClothSimulationContext& ClothSimulationContext;
-		const TArray<TSharedRef<FManagedArrayCollection>>& ManagedArrayCollections;
+		const TArray<TSharedRef<const FManagedArrayCollection>>& ManagedArrayCollections;
 	};
 }
diff --git a/Engine/Plugins/ChaosClothAsset/Source/ChaosClothAssetEngine/Public/ChaosClothAsset/ClothAsset.h b/Engine/Plugins/ChaosClothAsset/Source/ChaosClothAssetEngine/Public/ChaosClothAsset/ClothAsset.h
--- a/Engine/Plugins/ChaosClothAsset/Source/ChaosClothAssetEngine/Public/ChaosClothAsset/ClothAsset.h
+++ b/Engine/Plugins/ChaosClothAsset/Source/ChaosClothAssetEngine/Public/ChaosClothAsset/ClothAsset.h
@@ -80,7 +80,7 @@ public:
 	//~ End USkinnedAsset interface
 
 	/** Return the enclosed Cloth Collection object. */
-	TArray<TSharedRef<FManagedArrayCollection>>& GetClothCollections() 
+	TArray<TSharedRef<const FManagedArrayCollection>>& GetClothCollections()
 	{
 		WaitUntilAsyncPropertyReleased(EClothAssetAsyncProperties::ClothCollection);
 		PRAGMA_DISABLE_DEPRECATION_WARNINGS
@@ -89,11 +89,11 @@ public:
 	}
 
 	/** Return the enclosed Cloth Collection object, const version. */
-	const TArray<TSharedRef<const FManagedArrayCollection>>& GetClothCollections() const 
+	const TArray<TSharedRef<const FManagedArrayCollection>>& GetClothCollections() const
 	{
 		WaitUntilAsyncPropertyReleased(EClothAssetAsyncProperties::ClothCollection, ESkinnedAssetAsyncPropertyLockType::ReadOnly);
 		PRAGMA_DISABLE_DEPRECATION_WARNINGS
-		return reinterpret_cast<const TArray<TSharedRef<const FManagedArrayCollection>>&>(ClothCollections);
+		return ClothCollections;
 		PRAGMA_ENABLE_DEPRECATION_WARNINGS
 	}
 
@@ -163,7 +163,7 @@ PRAGMA_ENABLE_DEPRECATION_WARNINGS
 	{
 		return GetFName();
 	}
-	virtual const TArray<TSharedRef<FManagedArrayCollection>>& GetCollections(int32 /*ModelIndex*/) const override;
+	virtual const TArray<TSharedRef<const FManagedArrayCollection>>& GetCollections(int32 /*ModelIndex*/) const override;
 	virtual const UPhysicsAsset* GetPhysicsAssetForModel(int32 /*ModelIndex*/) const override
 	{
 		return PhysicsAsset;
@@ -246,7 +246,7 @@ private:
 
 	/** Cloth Collection containing this asset data. One per LOD. */
 	UE_DEPRECATED(5.4, "This must be protected for async build, always use the accessors even internally.")
-	TArray<TSharedRef<FManagedArrayCollection>> ClothCollections;
+	TArray<TSharedRef<const FManagedArrayCollection>> ClothCollections;
 
 #if WITH_EDITORONLY_DATA
 	/** Source mesh geometry information (not used at runtime). */
diff --git a/Engine/Plugins/ChaosClothAsset/Source/ChaosClothAssetEngine/Public/ChaosClothAsset/ClothAssetBase.h b/Engine/Plugins/ChaosClothAsset/Source/ChaosClothAssetEngine/Public/ChaosClothAsset/ClothAssetBase.h
--- a/Engine/Plugins/ChaosClothAsset/Source/ChaosClothAssetEngine/Public/ChaosClothAsset/ClothAssetBase.h
+++ b/Engine/Plugins/ChaosClothAsset/Source/ChaosClothAssetEngine/Public/ChaosClothAsset/ClothAssetBase.h
@@ -44,9 +44,9 @@ public:
 	virtual FName GetClothSimulationModelName(int32 ModelIndex) const
 	PURE_VIRTUAL(UChaosClothAssetBase::GetClothSimulationModelName, return NAME_None;);
 
-	/** Return the collections that use the property facade for this asset model (one per LOD). */
-	virtual const TArray<TSharedRef<FManagedArrayCollection>>& GetCollections(int32 ModelIndex) const
-	PURE_VIRTUAL(UChaosClothAssetBase::GetCollections, static const TArray<TSharedRef<FManagedArrayCollection>> EmptyArray; return EmptyArray;);
+	/** Return the collections for this asset model (one per LOD). */
+	virtual const TArray<TSharedRef<const FManagedArrayCollection>>& GetCollections(int32 ModelIndex) const
+	PURE_VIRTUAL(UChaosClothAssetBase::GetCollections, static const TArray<TSharedRef<const FManagedArrayCollection>> EmptyArray; return EmptyArray;);
 
 	/** Return the physics asset used for the simulation by this asset model. */
 	virtual const UPhysicsAsset* GetPhysicsAssetForModel(int32 ModelIndex) const
@@ -72,7 +72,6 @@ public:
 	virtual FDataflowInstance& GetDataflowInstance() override;
 	//~ End IDataflowInstanceInterface interface
 
-
 	//~ Begin UObject interface
 	virtual void BeginDestroy() override;
 	virtual bool IsReadyForFinishDestroy() override;
diff --git a/Engine/Plugins/ChaosClothAsset/Source/ChaosClothAssetEngine/Public/ChaosClothAsset/ClothComponent.h b/Engine/Plugins/ChaosClothAsset/Source/ChaosClothAssetEngine/Public/ChaosClothAsset/ClothComponent.h
--- a/Engine/Plugins/ChaosClothAsset/Source/ChaosClothAssetEngine/Public/ChaosClothAsset/ClothComponent.h
+++ b/Engine/Plugins/ChaosClothAsset/Source/ChaosClothAssetEngine/Public/ChaosClothAsset/ClothComponent.h
@@ -42,10 +42,10 @@ private:
 	friend UChaosClothComponent;
 	UPROPERTY(Transient)
 	TObjectPtr<UChaosClothAssetInteractor> ClothOutfitInteractor;
-	TArray<TSharedPtr<FManagedArrayCollection>> PropertyCollections;
+	TArray<TSharedPtr<const FManagedArrayCollection>> PropertyCollections;
 	TArray<TSharedPtr<::Chaos::Softs::FCollectionPropertyFacade>> CollectionPropertyFacades;
 
-	void Initialize(const TArray<TSharedRef<FManagedArrayCollection>>& AssetPropertyCollections);
+	void Initialize(const TArray<TSharedRef<const FManagedArrayCollection>>& AssetPropertyCollections);
 };
 
 /**
@@ -187,12 +187,12 @@ public:
 	 * This might be different from the cloth asset's since the component's properties can be modified in code or in blueprints.
 	 * This could also be different from the cloth simulation object until the cloth simulation thread synchronise the properties.
 	 */
-	const TArray<TSharedPtr<FManagedArrayCollection>>& GetPropertyCollections(int32 ModelIndex) const
+	const TArray<TSharedPtr<const FManagedArrayCollection>>& GetPropertyCollections(int32 ModelIndex) const
 	{
 		return ClothSimulationProperties[ModelIndex].PropertyCollections;
 	}
 	/** Return the solver property collections. */
-	const TArray<TSharedPtr<FManagedArrayCollection>>& GetSolverPropertyCollections() const
+	const TArray<TSharedPtr<const FManagedArrayCollection>>& GetSolverPropertyCollections() const
 	{
 		return ClothSimulationProperties[0].PropertyCollections;
 	}
diff --git a/Engine/Plugins/ChaosClothAssetEditor/Source/ChaosClothAssetDataflowNodes/Private/ChaosClothAsset/TerminalNode.cpp b/Engine/Plugins/ChaosClothAssetEditor/Source/ChaosClothAssetDataflowNodes/Private/ChaosClothAsset/TerminalNode.cpp
--- a/Engine/Plugins/ChaosClothAssetEditor/Source/ChaosClothAssetDataflowNodes/Private/ChaosClothAsset/TerminalNode.cpp
+++ b/Engine/Plugins/ChaosClothAssetEditor/Source/ChaosClothAssetDataflowNodes/Private/ChaosClothAsset/TerminalNode.cpp
@@ -22,10 +22,10 @@
 
 namespace UE::Chaos::ClothAsset::Private
 {
-	uint32 CalculateClothChecksum(const TArray<TSharedRef<FManagedArrayCollection>>& InClothCollections)
+	uint32 CalculateClothChecksum(const TArray<TSharedRef<const FManagedArrayCollection>>& InClothCollections)
 	{
 		uint32 Checksum = 0;
-		for (const TSharedRef<FManagedArrayCollection>& ClothCollection : InClothCollections)
+		for (const TSharedRef<const FManagedArrayCollection>& ClothCollection : InClothCollections)
 		{
 			constexpr bool bIncludeWeightMapsTrue = true; // Currently, editing weight maps is destructive
 			FCollectionClothConstFacade Cloth(ClothCollection);
@@ -67,7 +67,7 @@ namespace UE::Chaos::ClothAsset::Private
 		return Checksum;
 	}
 
-	bool PropertyKeysAndSolverTypesMatch(const TArray<TSharedRef<FManagedArrayCollection>>& Collections0, const TArray<TSharedRef<FManagedArrayCollection>>& Collections1)
+	bool PropertyKeysAndSolverTypesMatch(const TArray<TSharedRef<const FManagedArrayCollection>>& Collections0, const TArray<TSharedRef<const FManagedArrayCollection>>& Collections1)
 	{
 		if (Collections0.Num() != Collections1.Num())
 		{
@@ -111,8 +111,8 @@ void FChaosClothAssetTerminalNode_v2::SetAssetValue(TObjectPtr<UObject> Asset, U
 	{
 		using namespace UE::Chaos::ClothAsset;
 
-		TArray<TSharedRef<FManagedArrayCollection>> InClothCollections = GetCleanedCollectionLodValues(Context);
-		TArray<TSharedRef<FManagedArrayCollection>>& ClothCollections = ClothAsset->GetClothCollections();
+		TArray<TSharedRef<const FManagedArrayCollection>> InClothCollections = GetCleanedCollectionLodValues(Context);
+		TArray<TSharedRef<const FManagedArrayCollection>>& ClothCollections = ClothAsset->GetClothCollections();
 
 		const uint32 PreviousChecksum = ClothColllectionChecksum;
 		const bool bPreviousChecksumsValid = bClothCollectionChecksumValid;
@@ -126,8 +126,10 @@ void FChaosClothAssetTerminalNode_v2::SetAssetValue(TObjectPtr<UObject> Asset, U
 			check(ClothCollections.Num() > 0);
 			for (int32 LODIndex = 0; LODIndex < InClothCollections.Num(); ++LODIndex)
 			{
-				Chaos::Softs::FCollectionPropertyFacade Properties(ClothCollections[LODIndex]);
+				TSharedRef<FManagedArrayCollection> ClothCollection = MakeShared<FManagedArrayCollection>(*ClothCollections[LODIndex]);
+				Chaos::Softs::FCollectionPropertyFacade Properties(ClothCollection);
 				Properties.UpdateProperties(InClothCollections[LODIndex].ToSharedPtr());
+				ClothCollections[LODIndex] = MoveTemp(ClothCollection);
 			}
 
 			// Asset must be resaved
@@ -148,7 +150,7 @@ void FChaosClothAssetTerminalNode_v2::SetAssetValue(TObjectPtr<UObject> Asset, U
 		for (int32 LodIndex = 0; LodIndex < InClothCollections.Num(); ++LodIndex)
 		{
 			// New LOD
-			TSharedRef<FManagedArrayCollection>& ClothCollection = ClothCollections.Emplace_GetRef(MakeShared<FManagedArrayCollection>());
+			TSharedRef<FManagedArrayCollection> ClothCollection = MakeShared<FManagedArrayCollection>();
 			UE::Chaos::ClothAsset::FCollectionClothFacade ClothFacade(ClothCollection);
 			ClothFacade.DefineSchema();
 
@@ -160,6 +162,7 @@ void FChaosClothAssetTerminalNode_v2::SetAssetValue(TObjectPtr<UObject> Asset, U
 						LOCTEXT("InvalidLODDetails", "LOD {0} has no valid data."),
 						LodIndex));
 
+				ClothCollections.Emplace(MoveTemp(ClothCollection));
 				continue;
 			}
 
@@ -226,14 +229,17 @@ void FChaosClothAssetTerminalNode_v2::SetAssetValue(TObjectPtr<UObject> Asset, U
 				using namespace ::Chaos::Softs;
 				PhysicsAssetPathName = InClothFacade.GetPhysicsAssetPathName();
 			}
+
+			ClothCollections.Emplace(MoveTemp(ClothCollection));
 		}
 
 		// Make sure that whatever happens there is always at least one empty LOD to avoid crashing the render data
 		if (ClothCollections.Num() < 1)
 		{
-			TSharedRef<FManagedArrayCollection>& ClothCollection = ClothCollections.Emplace_GetRef(MakeShared<FManagedArrayCollection>());
+			TSharedRef<FManagedArrayCollection> ClothCollection = MakeShared<FManagedArrayCollection>();
 			FCollectionClothFacade ClothFacade(ClothCollection);
 			ClothFacade.DefineSchema();
+			ClothCollections.Emplace(MoveTemp(ClothCollection));
 		}
 
 		// Set reference skeleton
@@ -288,17 +294,17 @@ void FChaosClothAssetTerminalNode_v2::OnPinRemoved(const UE::Dataflow::FPin& Pin
 	return Super::OnPinRemoved(Pin);
 }
 
-TArray<TSharedRef<FManagedArrayCollection>> FChaosClothAssetTerminalNode_v2::GetCleanedCollectionLodValues(UE::Dataflow::FContext& Context) const
+TArray<TSharedRef<const FManagedArrayCollection>> FChaosClothAssetTerminalNode_v2::GetCleanedCollectionLodValues(UE::Dataflow::FContext& Context) const
 {
 	using namespace UE::Chaos::ClothAsset;
 
-	TArray<TSharedRef<FManagedArrayCollection>> CollectionLodValues;
+	TArray<TSharedRef<const FManagedArrayCollection>> CollectionLodValues;
 	CollectionLodValues.Reserve(CollectionLods.Num());
 
 	int32 LastValidLodIndex = INDEX_NONE;
 	for (int32 LodIndex = 0; LodIndex < CollectionLods.Num(); ++LodIndex)
 	{
-		TSharedRef<FManagedArrayCollection>& CollectionLodValue = CollectionLodValues.Emplace_GetRef(MakeShared<FManagedArrayCollection>(GetValue<FManagedArrayCollection>(Context, GetConnectionReference(LodIndex))));
+		TSharedRef<FManagedArrayCollection> CollectionLodValue = MakeShared<FManagedArrayCollection>(GetValue<FManagedArrayCollection>(Context, GetConnectionReference(LodIndex)));
 
 		FCollectionClothFacade ClothFacade(CollectionLodValue);
 		if (ClothFacade.HasValidRenderData())  // The cloth collection must at least have a render mesh
@@ -306,31 +312,30 @@ TArray<TSharedRef<FManagedArrayCollection>> FChaosClothAssetTerminalNode_v2::Get
 			FClothGeometryTools::CleanupAndCompactMesh(CollectionLodValue);
 			LastValidLodIndex = LodIndex;
 		}
+		else if (LastValidLodIndex >= 0)
+		{
+			ClothFacade.DefineSchema();
+			ClothFacade.Initialize(FCollectionClothConstFacade(CollectionLodValues[LastValidLodIndex]));
+
+			FClothDataflowTools::LogAndToastWarning(*this,
+				LOCTEXT("InvalidInputLodNHeadline", "Invalid input LOD."),
+				FText::Format(
+					LOCTEXT("InvalidInputLodNDetails",
+						"Invalid or empty input LOD for LOD {0}.\n"
+						"Using the previous valid LOD {1} instead."),
+					LodIndex,
+					LastValidLodIndex));
+		}
 		else
 		{
-			if (LastValidLodIndex >= 0)
-			{
-				ClothFacade.DefineSchema();
-				ClothFacade.Initialize(FCollectionClothConstFacade(CollectionLodValues[LastValidLodIndex]));
-				FClothDataflowTools::LogAndToastWarning(*this,
-					LOCTEXT("InvalidInputLodNHeadline", "Invalid input LOD."),
-					FText::Format(
-						LOCTEXT("InvalidInputLodNDetails",
-							"Invalid or empty input LOD for LOD {0}.\n"
-							"Using the previous valid LOD {1} instead."),
-						LodIndex,
-						LastValidLodIndex));
-			}
-			else
-			{
-				FClothDataflowTools::LogAndToastWarning(*this,
-					LOCTEXT("InvalidInputLod0Headline", "Invalid input LOD 0."),
-					LOCTEXT("InvalidInputLod0Details",
-						"Invalid or empty input LOD for LOD 0.\n"
-						"LOD 0 cannot be empty in order to construct a valid Cloth Asset."));
-				break;
-			}
+			FClothDataflowTools::LogAndToastWarning(*this,
+				LOCTEXT("InvalidInputLod0Headline", "Invalid input LOD 0."),
+				LOCTEXT("InvalidInputLod0Details",
+					"Invalid or empty input LOD for LOD 0.\n"
+					"LOD 0 cannot be empty in order to construct a valid Cloth Asset."));
+			break;
 		}
+		CollectionLodValues.Emplace(MoveTemp(CollectionLodValue));
 	}
 	return CollectionLodValues;
 }
@@ -409,8 +414,8 @@ PRAGMA_DISABLE_DEPRECATION_WARNINGS
 		}
 PRAGMA_ENABLE_DEPRECATION_WARNINGS
 
-		TArray<TSharedRef<FManagedArrayCollection>> InClothCollections = GetCleanedCollectionLodValues(Context);
-		TArray<TSharedRef<FManagedArrayCollection>>& ClothCollections = ClothAsset->GetClothCollections();
+		TArray<TSharedRef<const FManagedArrayCollection>> InClothCollections = GetCleanedCollectionLodValues(Context);
+		TArray<TSharedRef<const FManagedArrayCollection>>& ClothCollections = ClothAsset->GetClothCollections();
 
 		const uint32 PreviousChecksum = ClothColllectionChecksum;
 		const bool bPreviousChecksumsValid = bClothCollectionChecksumValid;
@@ -424,8 +429,10 @@ PRAGMA_ENABLE_DEPRECATION_WARNINGS
 			check(ClothCollections.Num() > 0);
 			for (int32 LODIndex = 0; LODIndex < InClothCollections.Num(); ++LODIndex)
 			{
-				Chaos::Softs::FCollectionPropertyFacade Properties(ClothCollections[LODIndex]);
+				TSharedRef<FManagedArrayCollection> ClothCollection = MakeShared<FManagedArrayCollection>(*ClothCollections[LODIndex]);
+				Chaos::Softs::FCollectionPropertyFacade Properties(ClothCollection);
 				Properties.UpdateProperties(InClothCollections[LODIndex].ToSharedPtr());
+				ClothCollections[LODIndex] = MoveTemp(ClothCollection);
 			}
 
 			// Asset must be resaved
@@ -446,7 +453,7 @@ PRAGMA_ENABLE_DEPRECATION_WARNINGS
 		for (int32 LodIndex = 0; LodIndex < InClothCollections.Num(); ++LodIndex)
 		{
 			// New LOD
-			TSharedRef<FManagedArrayCollection>& ClothCollection = ClothCollections.Emplace_GetRef(MakeShared<FManagedArrayCollection>());
+			TSharedRef<FManagedArrayCollection> ClothCollection = MakeShared<FManagedArrayCollection>();
 			UE::Chaos::ClothAsset::FCollectionClothFacade ClothFacade(ClothCollection);
 			ClothFacade.DefineSchema();
 
@@ -458,6 +465,7 @@ PRAGMA_ENABLE_DEPRECATION_WARNINGS
 						LOCTEXT("InvalidLODDetails", "LOD {0} has no valid data."),
 						LodIndex));
 
+				ClothCollections.Emplace(MoveTemp(ClothCollection));
 				continue;
 			}
 
@@ -507,21 +515,23 @@ PRAGMA_ENABLE_DEPRECATION_WARNINGS
 				}
 			}
 
-
 			// Set physics asset only with LOD 0 at the moment
 			if (LodIndex == 0)
 			{
 				using namespace ::Chaos::Softs;
 				PhysicsAssetPathName = InClothFacade.GetPhysicsAssetPathName();
 			}
+
+			ClothCollections.Emplace(MoveTemp(ClothCollection));
 		}
 
 		// Make sure that whatever happens there is always at least one empty LOD to avoid crashing the render data
 		if (ClothCollections.Num() < 1)
 		{
-			TSharedRef<FManagedArrayCollection>& ClothCollection = ClothCollections.Emplace_GetRef(MakeShared<FManagedArrayCollection>());
+			TSharedRef<FManagedArrayCollection> ClothCollection = MakeShared<FManagedArrayCollection>();
 			FCollectionClothFacade ClothFacade(ClothCollection);
 			ClothFacade.DefineSchema();
+			ClothCollections.Emplace(MoveTemp(ClothCollection));
 		}
 
 		// Set reference skeleton
@@ -665,18 +675,18 @@ PRAGMA_ENABLE_DEPRECATION_WARNINGS
 	}
 }
 
-TArray<TSharedRef<FManagedArrayCollection>> FChaosClothAssetTerminalNode::GetCleanedCollectionLodValues(UE::Dataflow::FContext& Context) const
+TArray<TSharedRef<const FManagedArrayCollection>> FChaosClothAssetTerminalNode::GetCleanedCollectionLodValues(UE::Dataflow::FContext& Context) const
 {
 	using namespace UE::Chaos::ClothAsset;
 
 	const TArray<const FManagedArrayCollection*> CollectionLods = GetCollectionLods();
-	TArray<TSharedRef<FManagedArrayCollection>> CollectionLodValues;
+	TArray<TSharedRef<const FManagedArrayCollection>> CollectionLodValues;
 	CollectionLodValues.Reserve(CollectionLods.Num());
 
-	int32 LastValidLodIndex = INDEX_NONE;	
+	int32 LastValidLodIndex = INDEX_NONE;
 	for (int32 LodIndex = 0; LodIndex < CollectionLods.Num(); ++LodIndex)
 	{
-		TSharedRef<FManagedArrayCollection>& CollectionLodValue = CollectionLodValues.Emplace_GetRef(MakeShared<FManagedArrayCollection>(GetValue<FManagedArrayCollection>(Context, CollectionLods[LodIndex])));
+		TSharedRef<FManagedArrayCollection> CollectionLodValue = MakeShared<FManagedArrayCollection>(GetValue<FManagedArrayCollection>(Context, CollectionLods[LodIndex]));
 
 		FCollectionClothFacade ClothFacade(CollectionLodValue);
 		if (ClothFacade.HasValidRenderData())  // The cloth collection must at least have a render mesh
@@ -684,31 +694,30 @@ TArray<TSharedRef<FManagedArrayCollection>> FChaosClothAssetTerminalNode::GetCle
 			FClothGeometryTools::CleanupAndCompactMesh(CollectionLodValue);
 			LastValidLodIndex = LodIndex;
 		}
+		else if (LastValidLodIndex >= 0)
+		{
+			ClothFacade.DefineSchema();
+			ClothFacade.Initialize(FCollectionClothConstFacade(CollectionLodValues[LastValidLodIndex]));
+
+			FClothDataflowTools::LogAndToastWarning(*this,
+				LOCTEXT("InvalidInputLodNHeadline", "Invalid input LOD."),
+				FText::Format(
+					LOCTEXT("InvalidInputLodNDetails",
+						"Invalid or empty input LOD for LOD {0}.\n"
+						"Using the previous valid LOD {1} instead."),
+					LodIndex,
+					LastValidLodIndex));
+		}
 		else
 		{
-			if (LastValidLodIndex >= 0)
-			{
-				ClothFacade.DefineSchema();
-				ClothFacade.Initialize(FCollectionClothConstFacade(CollectionLodValues[LastValidLodIndex]));
-				FClothDataflowTools::LogAndToastWarning(*this,
-					LOCTEXT("InvalidInputLodNHeadline", "Invalid input LOD."),
-					FText::Format(
-						LOCTEXT("InvalidInputLodNDetails",
-							"Invalid or empty input LOD for LOD {0}.\n"
-							"Using the previous valid LOD {1} instead."),
-						LodIndex,
-						LastValidLodIndex));
-			}
-			else
-			{
-				FClothDataflowTools::LogAndToastWarning(*this,
-					LOCTEXT("InvalidInputLod0Headline", "Invalid input LOD 0."),
-					LOCTEXT("InvalidInputLod0Details",
-						"Invalid or empty input LOD for LOD 0.\n"
-						"LOD 0 cannot be empty in order to construct a valid Cloth Asset."));
-				break;
-			}
+			FClothDataflowTools::LogAndToastWarning(*this,
+				LOCTEXT("InvalidInputLod0Headline", "Invalid input LOD 0."),
+				LOCTEXT("InvalidInputLod0Details",
+					"Invalid or empty input LOD for LOD 0.\n"
+					"LOD 0 cannot be empty in order to construct a valid Cloth Asset."));
+			break;
 		}
+		CollectionLodValues.Emplace(MoveTemp(CollectionLodValue));
 	}
 	return CollectionLodValues;
 }
diff --git a/Engine/Plugins/ChaosClothAssetEditor/Source/ChaosClothAssetDataflowNodes/Public/ChaosClothAsset/TerminalNode.h b/Engine/Plugins/ChaosClothAssetEditor/Source/ChaosClothAssetDataflowNodes/Public/ChaosClothAsset/TerminalNode.h
--- a/Engine/Plugins/ChaosClothAssetEditor/Source/ChaosClothAssetDataflowNodes/Public/ChaosClothAsset/TerminalNode.h
+++ b/Engine/Plugins/ChaosClothAssetEditor/Source/ChaosClothAssetDataflowNodes/Public/ChaosClothAsset/TerminalNode.h
@@ -59,7 +59,7 @@ private:
 	virtual void PostSerialize(const FArchive& Ar) override;
 	//~ End FDataflowNode interface
 
-	TArray<TSharedRef<FManagedArrayCollection>> GetCleanedCollectionLodValues(UE::Dataflow::FContext& Context) const;
+	TArray<TSharedRef<const FManagedArrayCollection>> GetCleanedCollectionLodValues(UE::Dataflow::FContext& Context) const;
 	UE::Dataflow::TConnectionReference<FManagedArrayCollection> GetConnectionReference(int32 Index) const;
 
 	UPROPERTY()
@@ -133,7 +133,7 @@ private:
 	//~ End FDataflowNode interface
 
 	TArray<const FManagedArrayCollection*> GetCollectionLods() const;
-	TArray<TSharedRef<FManagedArrayCollection>> GetCleanedCollectionLodValues(UE::Dataflow::FContext& Context) const;
+	TArray<TSharedRef<const FManagedArrayCollection>> GetCleanedCollectionLodValues(UE::Dataflow::FContext& Context) const;
 	const FManagedArrayCollection* GetCollectionLod(int32 LodIndex) const;
 
 	UPROPERTY()
diff --git a/Engine/Plugins/ChaosClothAssetEditor/Source/ChaosClothAssetTools/Private/ChaosClothAsset/ClothingAssetToClothAssetExporter.cpp b/Engine/Plugins/ChaosClothAssetEditor/Source/ChaosClothAssetTools/Private/ChaosClothAsset/ClothingAssetToClothAssetExporter.cpp
--- a/Engine/Plugins/ChaosClothAssetEditor/Source/ChaosClothAssetTools/Private/ChaosClothAsset/ClothingAssetToClothAssetExporter.cpp
+++ b/Engine/Plugins/ChaosClothAssetEditor/Source/ChaosClothAssetTools/Private/ChaosClothAsset/ClothingAssetToClothAssetExporter.cpp
@@ -95,7 +95,7 @@ void UClothingAssetToChaosClothAssetExporter::Export(const UClothingAssetBase* C
 
 	UChaosClothAsset* const ClothAsset = CastChecked<UChaosClothAsset>(ExportedAsset);
 	check(ClothAsset);
-	TArray<TSharedRef<FManagedArrayCollection>>& ClothCollections = ClothAsset->GetClothCollections();
+	TArray<TSharedRef<const FManagedArrayCollection>>& ClothCollections = ClothAsset->GetClothCollections();
 	
 	const int32 NumLods = ClothingAssetCommon->LodData.Num();
 	ClothCollections.Reset(NumLods);
@@ -106,7 +106,7 @@ void UClothingAssetToChaosClothAssetExporter::Export(const UClothingAssetBase* C
 		const FClothLODDataCommon& ClothLODData = ClothingAssetCommon->LodData[LodIndex];
 		const FClothPhysicalMeshData& PhysicalMeshData = ClothLODData.PhysicalMeshData;
 
-		TSharedRef<FManagedArrayCollection>& ClothCollection = ClothCollections.Emplace_GetRef(MakeShared<FManagedArrayCollection>());
+		TSharedRef<FManagedArrayCollection> ClothCollection = MakeShared<FManagedArrayCollection>();
 		FCollectionClothFacade ClothFacade(ClothCollection);
 		ClothFacade.DefineSchema();
 
@@ -123,6 +123,8 @@ void UClothingAssetToChaosClothAssetExporter::Export(const UClothingAssetBase* C
 		constexpr int32 UVChannelIndexNone = INDEX_NONE;
 		constexpr bool bAppend = false;
 		FClothGeometryTools::BuildSimMeshFromDynamicMesh(ClothCollection, DynamicMesh, UVChannelIndexNone, FVector2f(1.f), bAppend);
+
+		ClothCollections.Emplace(MoveTemp(ClothCollection));
 	}
 
 	if (NumLods)
@@ -133,9 +135,10 @@ void UClothingAssetToChaosClothAssetExporter::Export(const UClothingAssetBase* C
 	else
 	{
 		// Make sure that at least one empty LOD is always created
-		TSharedRef<FManagedArrayCollection>& ClothCollection = ClothCollections.Emplace_GetRef(MakeShared<FManagedArrayCollection>());
+		TSharedRef<FManagedArrayCollection> ClothCollection = MakeShared<FManagedArrayCollection>();
 		FCollectionClothFacade Cloth(ClothCollection);
 		Cloth.DefineSchema();
+		ClothCollections.Emplace(MoveTemp(ClothCollection));
 	}
 
 	// Assign the physics asset if any (must be done after having added the LODs)
diff --git a/Engine/Plugins/Experimental/ChaosOutfitAsset/Source/ChaosOutfitAssetEngine/Private/ChaosClothAsset/Outfit.cpp b/Engine/Plugins/Experimental/ChaosOutfitAsset/Source/ChaosOutfitAssetEngine/Private/ChaosClothAsset/Outfit.cpp
--- a/Engine/Plugins/Experimental/ChaosOutfitAsset/Source/ChaosOutfitAssetEngine/Private/ChaosClothAsset/Outfit.cpp
+++ b/Engine/Plugins/Experimental/ChaosOutfitAsset/Source/ChaosOutfitAssetEngine/Private/ChaosClothAsset/Outfit.cpp
@@ -154,27 +154,36 @@ bool FChaosOutfitPiece::Serialize(FArchive& Ar)
 		UScriptStruct* const Struct = FChaosClothSimulationModel::StaticStruct();
 		Struct->SerializeTaggedProperties(Ar, (uint8*)&ClothSimulationModel.Get(), Struct, nullptr);
 	}
-	// Custom serialize tThe collections since TArray<TSharedRef> can't be declared as UPROPERTY
+	// Custom serialize the collections since TArray<TSharedRef> can't be declared as UPROPERTY
 	int32 NumCollections = Collections.Num();
 	Ar << NumCollections;
-	Collections.Reserve(NumCollections);
-	for (int32 Index = 0; Index < NumCollections; ++Index)
+
+	if (Ar.IsLoading())
+	{
+		Collections.Reset(NumCollections);
+		for (int32 Index = 0; Index < NumCollections; ++Index)
+		{
+			TSharedRef<FManagedArrayCollection> NewCollection = MakeShared<FManagedArrayCollection>();
+			NewCollection->Serialize(Ar);
+			Collections.Emplace(MoveTemp(NewCollection));
+		}
+	}
+	else
 	{
-		if (Index >= Collections.Num())
+		for (int32 Index = 0; Index < NumCollections; ++Index)
 		{
-			Collections.Emplace(MakeShared<FManagedArrayCollection>());
+			ConstCastSharedRef<FManagedArrayCollection>(Collections[Index])->Serialize(Ar);
 		}
-		Collections[Index]->Serialize(Ar);
 	}
 	return true;
 }
 
-void FChaosOutfitPiece::DeepCopyCollections(const TArray<TSharedRef<FManagedArrayCollection>>& Other)
+void FChaosOutfitPiece::DeepCopyCollections(const TArray<TSharedRef<const FManagedArrayCollection>>& Other)
 {
 	Collections.Reserve(Other.Num());
-	for (const TSharedRef<FManagedArrayCollection>& Collection : Other)
+	for (const TSharedRef<const FManagedArrayCollection>& Collection : Other)
 	{
-		Collections.Emplace(MakeShared<FManagedArrayCollection>(*Collection));
+		Collections.Emplace(MakeShared<const FManagedArrayCollection>(*Collection));
 	}
 }
 
diff --git a/Engine/Plugins/Experimental/ChaosOutfitAsset/Source/ChaosOutfitAssetEngine/Private/ChaosClothAsset/OutfitAsset.cpp b/Engine/Plugins/Experimental/ChaosOutfitAsset/Source/ChaosOutfitAssetEngine/Private/ChaosClothAsset/OutfitAsset.cpp
--- a/Engine/Plugins/Experimental/ChaosOutfitAsset/Source/ChaosOutfitAssetEngine/Private/ChaosClothAsset/OutfitAsset.cpp
+++ b/Engine/Plugins/Experimental/ChaosOutfitAsset/Source/ChaosOutfitAssetEngine/Private/ChaosClothAsset/OutfitAsset.cpp
@@ -50,7 +50,7 @@ TSharedPtr<const FChaosClothSimulationModel> UChaosOutfitAsset::GetClothSimulati
 	return Pieces[ModelIndex].ClothSimulationModel;
 }
 
-const TArray<TSharedRef<FManagedArrayCollection>>& UChaosOutfitAsset::GetCollections(int32 ModelIndex) const
+const TArray<TSharedRef<const FManagedArrayCollection>>& UChaosOutfitAsset::GetCollections(int32 ModelIndex) const
 {
 	return Pieces[ModelIndex].Collections;
 }
diff --git a/Engine/Plugins/Experimental/ChaosOutfitAsset/Source/ChaosOutfitAssetEngine/Public/ChaosClothAsset/Outfit.h b/Engine/Plugins/Experimental/ChaosOutfitAsset/Source/ChaosOutfitAssetEngine/Public/ChaosClothAsset/Outfit.h
--- a/Engine/Plugins/Experimental/ChaosOutfitAsset/Source/ChaosOutfitAssetEngine/Public/ChaosClothAsset/Outfit.h
+++ b/Engine/Plugins/Experimental/ChaosOutfitAsset/Source/ChaosOutfitAssetEngine/Public/ChaosClothAsset/Outfit.h
@@ -27,7 +27,7 @@ struct CHAOSOUTFITASSETENGINE_API FChaosOutfitPiece final
 
 	TSharedRef<FChaosClothSimulationModel> ClothSimulationModel;
 
-	TArray<TSharedRef<FManagedArrayCollection>> Collections;
+	TArray<TSharedRef<const FManagedArrayCollection>> Collections;
 
 	FChaosOutfitPiece();
 	~FChaosOutfitPiece();
@@ -45,7 +45,7 @@ struct CHAOSOUTFITASSETENGINE_API FChaosOutfitPiece final
 private:
 	friend UChaosOutfit;
 
-	void DeepCopyCollections(const TArray<TSharedRef<FManagedArrayCollection>>& Other);
+	void DeepCopyCollections(const TArray<TSharedRef<const FManagedArrayCollection>>& Other);
 
 	void RemapBoneIndices(const TArray<int32>& BoneMap);
 };
diff --git a/Engine/Plugins/Experimental/ChaosOutfitAsset/Source/ChaosOutfitAssetEngine/Public/ChaosClothAsset/OutfitAsset.h b/Engine/Plugins/Experimental/ChaosOutfitAsset/Source/ChaosOutfitAssetEngine/Public/ChaosClothAsset/OutfitAsset.h
--- a/Engine/Plugins/Experimental/ChaosOutfitAsset/Source/ChaosOutfitAssetEngine/Public/ChaosClothAsset/OutfitAsset.h
+++ b/Engine/Plugins/Experimental/ChaosOutfitAsset/Source/ChaosOutfitAssetEngine/Public/ChaosClothAsset/OutfitAsset.h
@@ -32,7 +32,7 @@ private:
 	}
 	virtual FName GetClothSimulationModelName(int32 ModelIndex) const override;
 	virtual TSharedPtr<const FChaosClothSimulationModel> GetClothSimulationModel(int32 ModelIndex) const override;
-	virtual const TArray<TSharedRef<FManagedArrayCollection>>& GetCollections(int32 ModelIndex) const override;
+	virtual const TArray<TSharedRef<const FManagedArrayCollection>>& GetCollections(int32 ModelIndex) const override;
 	virtual const UPhysicsAsset* GetPhysicsAssetForModel(int32 ModelIndex) const override;
 	//~ End UChaosClothAssetBase interface
 
