Cloth bounding box fixes. Also, do not create a ClothProxy if p.ClothPhysics == 0.

Serialize the ClothAsset bounding box. The data needed to calculate the bounding box is stripped on cook. This is in line with how SKM stores its ImportedBounds (via Serialize, not UPROPERTY, in its case).

Dirty the cached bounds when the cloth simulation is finished.

diff --git a/Engine/Plugins/ChaosClothAsset/Source/ChaosClothAssetEngine/Private/ChaosClothAsset/ClothAsset.cpp b/Engine/Plugins/ChaosClothAsset/Source/ChaosClothAssetEngine/Private/ChaosClothAsset/ClothAsset.cpp
--- a/Engine/Plugins/ChaosClothAsset/Source/ChaosClothAssetEngine/Private/ChaosClothAsset/ClothAsset.cpp
+++ b/Engine/Plugins/ChaosClothAsset/Source/ChaosClothAssetEngine/Private/ChaosClothAsset/ClothAsset.cpp
@@ -446,13 +446,13 @@ void UChaosClothAsset::FinishPostLoadInternal(FSkinnedAssetPostLoadContext& Cont
 	}
 
 	CalculateInvRefMatrices();
-	CalculateBounds();
 
 #if WITH_EDITOR
 	ReleaseAsyncProperty();
 #endif
 }
 
+#if WITH_EDITOR
 void UChaosClothAsset::CalculateBounds()
 {
 	using namespace UE::Chaos::ClothAsset;
@@ -472,6 +472,7 @@ void UChaosClothAsset::CalculateBounds()
 
 	Bounds = FBoxSphereBounds(BoundingBox);
 }
+#endif
 
 void UChaosClothAsset::Build(TArray<FChaosClothAssetLodTransitionDataCache>* InOutTransitionCache)
 {
@@ -492,9 +493,6 @@ void UChaosClothAsset::Build(TArray<FChaosClothAssetLodTransitionDataCache>* InO
 	// Rebuild matrices
 	CalculateInvRefMatrices();
 
-	// Update bounds
-	CalculateBounds();
-
 	// Add LODs to the render data
 	const int32 NumLods = FMath::Max(GetClothCollections().Num(), 1);  // The render data will always look for at least one default LOD 0
 
@@ -547,6 +545,9 @@ void UChaosClothAsset::BuildMeshModel()
 		GetImportedModel()->LODModels.Add(new FSkeletalMeshLODModel());
 		FBuilder::BuildLod(GetImportedModel()->LODModels[LodIndex], *this, LodIndex, TargetPlatform);
 	}
+
+	// Update bounds
+	CalculateBounds();
 }
 #endif  // #if WITH_EDITOR
 
diff --git a/Engine/Plugins/ChaosClothAsset/Source/ChaosClothAssetEngine/Private/ChaosClothAsset/ClothComponent.cpp b/Engine/Plugins/ChaosClothAsset/Source/ChaosClothAssetEngine/Private/ChaosClothAsset/ClothComponent.cpp
--- a/Engine/Plugins/ChaosClothAsset/Source/ChaosClothAssetEngine/Private/ChaosClothAsset/ClothComponent.cpp
+++ b/Engine/Plugins/ChaosClothAsset/Source/ChaosClothAssetEngine/Private/ChaosClothAsset/ClothComponent.cpp
@@ -587,9 +587,16 @@ TSharedPtr<UE::Chaos::ClothAsset::FClothSimulationProxy> UChaosClothComponent::C
 
 void UChaosClothComponent::CreateClothSimulationProxyImpl()
 {
-	// Note: CreateClothSimulationProxy and ClothSimulationProxy::PostConstructor are both virtual.
-	ClothSimulationProxy = CreateClothSimulationProxy();
-	ClothSimulationProxy->PostConstructor();
+	static IConsoleVariable* const CVarClothPhysics = IConsoleManager::Get().FindConsoleVariable(TEXT("p.ClothPhysics"));
+	if (!CVarClothPhysics || CVarClothPhysics->GetBool())
+	{
+		// Note: CreateClothSimulationProxy and ClothSimulationProxy::PostConstructor are both virtual.
+		ClothSimulationProxy = CreateClothSimulationProxy();
+		if (ClothSimulationProxy)
+		{
+			ClothSimulationProxy->PostConstructor();
+		}
+	}
 }
 
 void UChaosClothComponent::AddCollisionSource(USkinnedMeshComponent* SourceComponent, const UPhysicsAsset* SourcePhysicsAsset, bool bUseSphylsOnly)
@@ -663,9 +670,10 @@ void UChaosClothComponent::HandleExistingParallelSimulation()
 		}
 	}
 
-	if (ClothSimulationProxy.IsValid())
+	if (ClothSimulationProxy.IsValid() && ClothSimulationProxy->IsParallelSimulationTaskValid())
 	{
 		ClothSimulationProxy->CompleteParallelSimulation_GameThread();
+		InvalidateCachedBounds();
 	}
 }
 
@@ -770,6 +778,7 @@ void UChaosClothComponent::ReadFromSimulation(const float /*DeltaTime*/, const b
 	if (ClothSimulationProxy.IsValid())
 	{
 		ClothSimulationProxy->PostSimulate_GameThread();
+		InvalidateCachedBounds();
 	}
 }
 
diff --git a/Engine/Plugins/ChaosClothAsset/Source/ChaosClothAssetEngine/Public/ChaosClothAsset/ClothAssetBase.h b/Engine/Plugins/ChaosClothAsset/Source/ChaosClothAssetEngine/Public/ChaosClothAsset/ClothAssetBase.h
--- a/Engine/Plugins/ChaosClothAsset/Source/ChaosClothAssetEngine/Public/ChaosClothAsset/ClothAssetBase.h
+++ b/Engine/Plugins/ChaosClothAsset/Source/ChaosClothAssetEngine/Public/ChaosClothAsset/ClothAssetBase.h
@@ -348,6 +348,7 @@ protected:
 #endif
 
 	/** Bounds for this asset. */
+	UPROPERTY()
 	FBoxSphereBounds Bounds;
 
 private:
diff --git a/Engine/Plugins/ChaosClothAsset/Source/ChaosClothAssetEngine/Public/ChaosClothAsset/ClothSimulationProxy.h b/Engine/Plugins/ChaosClothAsset/Source/ChaosClothAssetEngine/Public/ChaosClothAsset/ClothSimulationProxy.h
--- a/Engine/Plugins/ChaosClothAsset/Source/ChaosClothAssetEngine/Public/ChaosClothAsset/ClothSimulationProxy.h
+++ b/Engine/Plugins/ChaosClothAsset/Source/ChaosClothAssetEngine/Public/ChaosClothAsset/ClothSimulationProxy.h
@@ -61,6 +61,12 @@ namespace UE::Chaos::ClothAsset
 		/** Wait for the parallel task to complete if one was running, and update the simulation data. */
 		void CompleteParallelSimulation_GameThread();
 
+		/** Returns true if a parallel simulation task has been launched and the simulation data has not been updated, i.e., CompleteParallelSimulation_GameThread will do something. */
+		bool IsParallelSimulationTaskValid() const
+		{
+			return IsValidRef(ParallelTask);
+		}
+
 		/** Prepare the proxy before the simulation is setup. For specialized usage when not calling Tick_GameThread. 
 		 * @param bForceWaitForInitialization will ignore the value of p.ChaosClothAsset.WaitForAsyncClothInitialization and force waiting for any in flight initialization, guaranteeing successful preprocessing..*/
 		void PreProcess_GameThread(float DeltaTime, bool bForceWaitForInitialization = false);
