Chaos Cloth Asset:
- Added a new MakeClothAsset node.
- Moved the terminal node's build from collection code into the cloth asset class.
- Moved a few editor only skinned asset overrides (BuildInternal) from the base class to the cloth asset class.
- Added GetAssetGuid() method to the cloth asset base class.

#rnx

diff --git a/Engine/Plugins/ChaosClothAsset/Source/ChaosClothAssetEngine/Private/ChaosClothAsset/ClothAsset.cpp b/Engine/Plugins/ChaosClothAsset/Source/ChaosClothAssetEngine/Private/ChaosClothAsset/ClothAsset.cpp
--- a/Engine/Plugins/ChaosClothAsset/Source/ChaosClothAssetEngine/Private/ChaosClothAsset/ClothAsset.cpp
+++ b/Engine/Plugins/ChaosClothAsset/Source/ChaosClothAssetEngine/Private/ChaosClothAsset/ClothAsset.cpp
@@ -7,6 +7,7 @@
 #include "ChaosClothAsset/ClothGeometryTools.h"
 #include "ChaosClothAsset/ClothSimulationModel.h"
 #include "ChaosClothAsset/CollectionClothFacade.h"
+#include "ChaosClothAsset/CollectionClothSelectionFacade.h"
 #include "Chaos/CollectionEmbeddedSpringConstraintFacade.h"
 #include "Chaos/CollectionPropertyFacade.h"
 #include "Animation/Skeleton.h"
@@ -34,6 +35,7 @@
 // In case of merge conflicts with DDC versions, you MUST generate a new GUID and set this new GUID as the version.
 #define CHAOS_CLOTH_ASSET_DERIVED_DATA_VERSION TEXT("5C041E93563044A69646B5E98F152B7C")
 
+#define LOCTEXT_NAMESPACE "ChaosClothAsset"
 
 namespace UE::Chaos::ClothAsset::Private
 {
@@ -322,6 +324,51 @@ void UChaosClothAsset::GetResourceSizeEx(FResourceSizeEx& CumulativeResourceSize
 #endif
 }
 
+#if WITH_EDITOR
+void UChaosClothAsset::ExecuteBuildInternal(FSkinnedAssetBuildContext& Context)
+{
+	TRACE_CPUPROFILER_EVENT_SCOPE(UChaosClothAsset::ExecuteBuildInternal);
+
+	// This scope allows us to use any locked properties without causing stalls
+	FSkinnedAssetAsyncBuildScope AsyncBuildScope(this);
+
+	// rebuild render data from imported model
+	CacheDerivedData(&Context);
+
+	// Build the material channel data used by the texture streamer
+	UpdateUVChannelData(true);
+}
+
+void UChaosClothAsset::BeginBuildInternal(FSkinnedAssetBuildContext& Context)
+{
+	TRACE_CPUPROFILER_EVENT_SCOPE(UChaosClothAsset::BeginBuildInternal);
+
+	SetInternalFlags(EInternalObjectFlags::Async);
+
+	// Unregister all instances of this component
+	Context.RecreateRenderStateContext = MakeUnique<FSkinnedMeshComponentRecreateRenderStateContext>(this, false);
+
+	// Release the render data resources.
+	ReleaseResources();
+
+	// Flush the resource release commands to the rendering thread to ensure that the build doesn't occur while a resource is still
+	// allocated, and potentially accessing the UChaosClothAsset.
+	ReleaseResourcesFence.Wait();
+
+	// Lock all properties that should not be modified/accessed during async post-load
+	AcquireAsyncProperty();
+}
+
+void UChaosClothAsset::FinishBuildInternal(FSkinnedAssetBuildContext& Context)
+{
+	TRACE_CPUPROFILER_EVENT_SCOPE(UChaosClothAsset::FinishBuildInternal);
+
+	ClearInternalFlags(EInternalObjectFlags::Async);
+
+	ReleaseAsyncProperty();
+}
+#endif // #if WITH_EDITOR
+
 void UChaosClothAsset::BeginPostLoadInternal(FSkinnedAssetPostLoadContext& Context)
 {
 #if WITH_EDITOR
@@ -338,8 +385,7 @@ void UChaosClothAsset::BeginPostLoadInternal(FSkinnedAssetPostLoadContext& Conte
 	// This scope allows us to use any locked properties without causing stalls
 	FSkinnedAssetAsyncBuildScope AsyncBuildScope(this);
 
-	// Make sure that the collection is still compatible and valid
-	bool bAnyInvalidLods = false;
+	// Make sure that there is at least one valid collection
 	if (GetClothCollections().IsEmpty())
 	{
 		UE_LOG(LogChaosClothAsset, Warning, TEXT("Invalid Cloth Collection (no LODs) found while loading Cloth Asset %s."), *GetFullName());
@@ -347,14 +393,11 @@ void UChaosClothAsset::BeginPostLoadInternal(FSkinnedAssetPostLoadContext& Conte
 		FCollectionClothFacade ClothFacade(ClothCollection);
 		ClothFacade.DefineSchema();
 		GetClothCollections().Emplace(MoveTemp(ClothCollection));
-		bAnyInvalidLods = true;
 	}
 
+	// Check that all LODs have the cloth schema
 	const int32 NumLods = GetClothCollections().Num();
 	check(NumLods >= 1);  // The default LOD 0 should be present now if it ever was missing
-	LODInfo.SetNum(NumLods);  // Always keep a matching number of LODInfos
-
-	bool bAnyInvalidSkeletons = false;
 	for (int32 LODIndex = 0; LODIndex < NumLods; ++LODIndex)
 	{
 		const TSharedRef<const FManagedArrayCollection>& ClothCollection = GetClothCollections()[LODIndex];
@@ -367,32 +410,12 @@ void UChaosClothAsset::BeginPostLoadInternal(FSkinnedAssetPostLoadContext& Conte
 			FCollectionClothFacade NewClothFacade = FCollectionClothFacade(NewClothCollection);
 			NewClothFacade.DefineSchema();
 			GetClothCollections()[LODIndex] = MoveTemp(NewClothCollection);
-			bAnyInvalidLods = true;
-			bAnyInvalidSkeletons = true;
 		}
-		else if (ClothConstFacade.GetSkeletalMeshPathName().IsEmpty())
-		{
-			bAnyInvalidSkeletons = true;
-		}
-	}
-	if (bAnyInvalidLods)
-	{
-		SetPhysicsAsset(PhysicsAsset);  // Re-update the collection with the physics asset information if any
-	}
-	if (bAnyInvalidSkeletons)
-	{
-		constexpr bool bRebuildModels = false;
-		constexpr bool bRebindMeshes = true;  // Best to rebind the mesh when reloading broken data
-		SetReferenceSkeleton(nullptr, bRebuildModels, bRebindMeshes);
 	}
 
 	// We're done touching the ClothCollections, so can unlock for read
 	ReleaseAsyncProperty((uint64)EClothAssetAsyncProperties::ClothCollection, ESkinnedAssetAsyncPropertyLockType::WriteOnly);
 
-	BuildClothSimulationModel();  // TODO: Cache ClothSimulationModel?
-
-	BuildMeshModel();
-
 	// Convert PerPlatForm data to PerQuality if perQuality data have not been serialized.
 	// Also test default value, since PerPlatformData can have Default !=0 and no PerPlatform data overrides.
 	const bool bConvertMinLODData = (MinQualityLevelLOD.PerQuality.Num() == 0 && MinQualityLevelLOD.Default == 0) && (MinLod.PerPlatform.Num() != 0 || MinLod.Default != 0);
@@ -474,6 +497,147 @@ void UChaosClothAsset::CalculateBounds()
 }
 #endif
 
+void UChaosClothAsset::Build(
+	const TArray<TSharedRef<const FManagedArrayCollection>>& InClothCollections,
+	TArray<FChaosClothAssetLodTransitionDataCache>* InOutTransitionCache,
+	FText* ErrorText,
+	FText* VerboseText)
+{
+	using namespace UE::Chaos::ClothAsset;
+
+	// Reset the asset's collection
+	TArray<TSharedRef<const FManagedArrayCollection>>& OutClothCollections = GetClothCollections();
+	OutClothCollections.Reset(InClothCollections.Num());
+
+	// Reset the asset's material list
+	TArray<FSkeletalMaterial>& OutMaterials = GetMaterials();
+	OutMaterials.Reset();
+
+	// Iterate through the LODs
+	FString PhysicsAssetPathName;
+
+	for (int32 LodIndex = 0; LodIndex < InClothCollections.Num(); ++LodIndex)
+	{
+		// New LOD
+		TSharedRef<FManagedArrayCollection> ClothCollection = MakeShared<FManagedArrayCollection>();
+		UE::Chaos::ClothAsset::FCollectionClothFacade ClothFacade(ClothCollection);
+		ClothFacade.DefineSchema();
+
+		const FCollectionClothConstFacade InClothFacade(InClothCollections[LodIndex]);
+		if (!InClothFacade.HasValidRenderData())  // The cloth collection must at least have a render mesh
+		{
+			if (ErrorText && ErrorText->IsEmpty())
+			{
+				*ErrorText = LOCTEXT("BuildErrorText", "Invalid LOD.");
+				if (VerboseText)
+				{
+					*VerboseText = FText::Format(LOCTEXT("BuildVerboseTextFirstError", "LOD {0} has no valid data."), LodIndex);
+				}
+			}
+			else if (ErrorText && VerboseText)
+			{
+				*VerboseText = FText::Format(LOCTEXT("BuildVerboseTextThereafter", "{0}\nLOD {1} has no valid data."), *VerboseText, LodIndex);
+			}
+			// else no error reporting
+
+			OutClothCollections.Emplace(MoveTemp(ClothCollection));
+			continue;
+		}
+
+		// Copy input LOD to current output LOD
+		ClothFacade.Initialize(InClothFacade);
+
+		// Add this LOD's materials to the asset
+		const int32 NumLodMaterials = ClothFacade.GetNumRenderPatterns();
+
+		OutMaterials.Reserve(OutMaterials.Num() + NumLodMaterials);
+
+		const TConstArrayView<FString> LodRenderMaterialPathName = ClothFacade.GetRenderMaterialPathName();
+		for (int32 LodMaterialIndex = 0; LodMaterialIndex < NumLodMaterials; ++LodMaterialIndex)
+		{
+			const FString& RenderMaterialPathName = LodRenderMaterialPathName[LodMaterialIndex];
+
+			if (UMaterialInterface* const Material = LoadObject<UMaterialInterface>(this, *RenderMaterialPathName, nullptr, LOAD_None, nullptr))
+			{
+				OutMaterials.Emplace(Material, true, false, Material->GetFName());
+			}
+			else
+			{
+				OutMaterials.Emplace();
+			}
+		}
+
+		// Set properties
+		constexpr bool bUpdateExistingProperties = false;
+		Chaos::Softs::FCollectionPropertyMutableFacade PropertyFacade(ClothCollection);
+		PropertyFacade.Append(InClothCollections[LodIndex].ToSharedPtr(), bUpdateExistingProperties);
+
+		// Set selections
+		FCollectionClothSelectionFacade Selection(ClothCollection);
+		FCollectionClothSelectionConstFacade InSelection(InClothCollections[LodIndex]);
+		if (InSelection.IsValid())
+		{
+			Selection.DefineSchema();
+			const TArray<FName> InSelectionNames = InSelection.GetNames();
+			for (const FName& InSelectionName : InSelectionNames)
+			{
+				const FName SelectionGroup = InSelection.GetSelectionGroup(InSelectionName);
+				if (SelectionGroup == ClothCollectionGroup::SimVertices3D ||
+					SelectionGroup == ClothCollectionGroup::SimFaces)
+				{
+					Selection.FindOrAddSelectionSet(InSelectionName, SelectionGroup) = InSelection.GetSelectionSet(InSelectionName);
+				}
+			}
+		}
+
+		// Set springs
+		Chaos::Softs::FEmbeddedSpringFacade EmbeddedSpringFacade(ClothCollection.Get(), ClothCollectionGroup::SimVertices3D);
+		const Chaos::Softs::FEmbeddedSpringFacade InEmbeddedSpringFacade(InClothCollections[LodIndex].Get(), ClothCollectionGroup::SimVertices3D);
+		if (InEmbeddedSpringFacade.IsValid())
+		{
+			EmbeddedSpringFacade.DefineSchema();
+			constexpr int32 VertexOffset = 0;
+			EmbeddedSpringFacade.Append(InEmbeddedSpringFacade, VertexOffset);
+		}
+
+		// Set physics asset only with LOD 0 at the moment
+		if (LodIndex == 0)
+		{
+			using namespace ::Chaos::Softs;
+			PhysicsAssetPathName = InClothFacade.GetPhysicsAssetPathName();
+		}
+
+		OutClothCollections.Emplace(MoveTemp(ClothCollection));
+	}
+
+	// Make sure that whatever happens there is always at least one empty LOD to avoid crashing the render data
+	if (OutClothCollections.Num() < 1)
+	{
+		TSharedRef<FManagedArrayCollection> ClothCollection = MakeShared<FManagedArrayCollection>();
+		FCollectionClothFacade ClothFacade(ClothCollection);
+		ClothFacade.DefineSchema();
+		OutClothCollections.Emplace(MoveTemp(ClothCollection));
+	}
+
+	// Set reference skeleton
+	constexpr bool bRebuildModels = false;  // Avoid rebuilding the asset twice
+PRAGMA_DISABLE_DEPRECATION_WARNINGS  // Temporary until moved to private
+	UpdateSkeletonFromCollection(bRebuildModels);
+PRAGMA_ENABLE_DEPRECATION_WARNINGS
+
+	// Set physics asset (note: the cloth asset's physics asset is only replaced if a collection path name is found valid)
+	PhysicsAsset = !PhysicsAssetPathName.IsEmpty() ?
+		LoadObject<UPhysicsAsset>(this, *PhysicsAssetPathName, nullptr, LOAD_None, nullptr) :
+		nullptr;
+
+	SetHasVertexColors(true);
+
+	// Rebuild the asset static data
+PRAGMA_DISABLE_DEPRECATION_WARNINGS
+	Build(InOutTransitionCache);
+PRAGMA_ENABLE_DEPRECATION_WARNINGS
+}
+
 void UChaosClothAsset::Build(TArray<FChaosClothAssetLodTransitionDataCache>* InOutTransitionCache)
 {
 	using namespace UE::Chaos::ClothAsset;
@@ -679,42 +843,21 @@ void UChaosClothAsset::SetPhysicsAsset(UPhysicsAsset* InPhysicsAsset)
 	using namespace UE::Chaos::ClothAsset;
 
 	PhysicsAsset = InPhysicsAsset;
-	for (TSharedRef<const FManagedArrayCollection>& ClothCollection : GetClothCollections())
-	{
-		TSharedRef<FManagedArrayCollection> NewClothCollection = MakeShared<FManagedArrayCollection>(*ClothCollection);
-		FCollectionClothFacade ClothFacade(NewClothCollection);
-		ClothFacade.SetPhysicsAssetPathName(PhysicsAsset ? PhysicsAsset->GetPathName() : FString());
-		ClothCollection = MoveTemp(NewClothCollection);
-	}
 }
 
-void UChaosClothAsset::SetReferenceSkeleton(const FReferenceSkeleton* ReferenceSkeleton, bool bRebuildModels, bool bRebindMeshes)
+void UChaosClothAsset::SetReferenceSkeleton(const FReferenceSkeleton* ReferenceSkeleton, bool bRebuildModels)
 {
 	using namespace UE::Chaos::ClothAsset;
 
-	if (!ReferenceSkeleton)
-	{
-		bRebindMeshes = true;  // Force the binding when a default reference skeleton is being created 
-	}
-
 	// Update the reference skeleton
 	Super::SetReferenceSkeleton(ReferenceSkeleton);
 
-	// Rebind the meshes
-	if (bRebindMeshes)
-	{
-		for (TSharedRef<const FManagedArrayCollection>& ClothCollection : GetClothCollections())
-		{
-			TSharedRef<FManagedArrayCollection> NewClothCollection = MakeShared<FManagedArrayCollection>(*ClothCollection);
-			FClothGeometryTools::BindMeshToRootBone(NewClothCollection, true, true);
-			ClothCollection = MoveTemp(NewClothCollection);
-		}
-	}
-
-	// Rebind the models
+	// Rebuild the models
 	if (bRebuildModels)
 	{
+PRAGMA_DISABLE_DEPRECATION_WARNINGS
 		Build();
+PRAGMA_ENABLE_DEPRECATION_WARNINGS
 	}
 }
 
@@ -732,8 +875,7 @@ void UChaosClothAsset::UpdateSkeletonFromCollection(bool bRebuildModels)
 
 	SetSkeleton(SkeletalMesh ? SkeletalMesh->GetSkeleton() : nullptr); // For completion only, this is not being used and might mismatch the skeletal mesh's reference skeleton
 
-	constexpr bool bRebindMeshes = false;  // The collection should contain the correct binding at the time SkeletalMeshPathName was set
-	SetReferenceSkeleton(SkeletalMesh ? &SkeletalMesh->GetRefSkeleton() : nullptr, bRebuildModels, bRebindMeshes);
+	SetReferenceSkeleton(SkeletalMesh ? &SkeletalMesh->GetRefSkeleton() : nullptr, bRebuildModels);
 }
 
 void UChaosClothAsset::CopySimMeshToRenderMesh(UMaterialInterface* Material)
@@ -766,3 +908,5 @@ void UChaosClothAsset::CopySimMeshToRenderMesh(UMaterialInterface* Material)
 		}
 	}
 }
+
+#undef LOCTEXT_NAMESPACE
diff --git a/Engine/Plugins/ChaosClothAsset/Source/ChaosClothAssetEngine/Private/ChaosClothAsset/ClothAssetBase.cpp b/Engine/Plugins/ChaosClothAsset/Source/ChaosClothAssetEngine/Private/ChaosClothAsset/ClothAssetBase.cpp
--- a/Engine/Plugins/ChaosClothAsset/Source/ChaosClothAssetEngine/Private/ChaosClothAsset/ClothAssetBase.cpp
+++ b/Engine/Plugins/ChaosClothAsset/Source/ChaosClothAssetEngine/Private/ChaosClothAsset/ClothAssetBase.cpp
@@ -452,51 +452,6 @@ void UChaosClothAssetBase::SetResourceForRendering(TUniquePtr<FSkeletalMeshRende
 	SkeletalMeshRenderData = MoveTemp(InSkeletalMeshRenderData);
 }
 
-#if WITH_EDITOR
-void UChaosClothAssetBase::ExecuteBuildInternal(FSkinnedAssetBuildContext& Context)
-{
-	TRACE_CPUPROFILER_EVENT_SCOPE(UChaosClothAsset::ExecuteBuildInternal);
-
-	// This scope allows us to use any locked properties without causing stalls
-	FSkinnedAssetAsyncBuildScope AsyncBuildScope(this);
-
-	// rebuild render data from imported model
-	CacheDerivedData(&Context);
-
-	// Build the material channel data used by the texture streamer
-	UpdateUVChannelData(true);
-}
-
-void UChaosClothAssetBase::BeginBuildInternal(FSkinnedAssetBuildContext& Context)
-{
-	TRACE_CPUPROFILER_EVENT_SCOPE(UChaosClothAsset::BeginBuildInternal);
-
-	SetInternalFlags(EInternalObjectFlags::Async);
-
-	// Unregister all instances of this component
-	Context.RecreateRenderStateContext = MakeUnique<FSkinnedMeshComponentRecreateRenderStateContext>(this, false);
-
-	// Release the render data resources.
-	ReleaseResources();
-
-	// Flush the resource release commands to the rendering thread to ensure that the build doesn't occur while a resource is still
-	// allocated, and potentially accessing the UChaosClothAsset.
-	ReleaseResourcesFence.Wait();
-
-	// Lock all properties that should not be modified/accessed during async post-load
-	AcquireAsyncProperty();
-}
-
-void UChaosClothAssetBase::FinishBuildInternal(FSkinnedAssetBuildContext& Context)
-{
-	TRACE_CPUPROFILER_EVENT_SCOPE(UChaosClothAsset::FinishBuildInternal);
-
-	ClearInternalFlags(EInternalObjectFlags::Async);
-
-	ReleaseAsyncProperty();
-}
-#endif // #if WITH_EDITOR
-
 #if WITH_EDITORONLY_DATA
 void UChaosClothAssetBase::SetPreviewSceneSkeletalMesh(USkeletalMesh* Mesh)
 {
diff --git a/Engine/Plugins/ChaosClothAsset/Source/ChaosClothAssetEngine/Public/ChaosClothAsset/ClothAsset.h b/Engine/Plugins/ChaosClothAsset/Source/ChaosClothAssetEngine/Public/ChaosClothAsset/ClothAsset.h
--- a/Engine/Plugins/ChaosClothAsset/Source/ChaosClothAssetEngine/Public/ChaosClothAsset/ClothAsset.h
+++ b/Engine/Plugins/ChaosClothAsset/Source/ChaosClothAssetEngine/Public/ChaosClothAsset/ClothAsset.h
@@ -97,7 +97,20 @@ public:
 		PRAGMA_ENABLE_DEPRECATION_WARNINGS
 	}
 
-	/** Build this asset static render and simulation data. This needs to be done every time the asset has changed. */
+	/**
+	 * Rebuild this asset from the specified managed array collection data.
+	 * @param InClothCollections The source managed array collection to build/rebuild the asset with.
+	 * @param InOutTransitionCache The LOD transition cache to speed up the build process, if available.
+	 * @param ErrorText The headline error text if any error occurs, an empty text otherwise.
+	 * @param VerboseText The verbose description of the error if any.
+	 */
+	void Build(
+		const TArray<TSharedRef<const FManagedArrayCollection>>& InClothCollections,
+		TArray<FChaosClothAssetLodTransitionDataCache>* InOutTransitionCache = nullptr,
+		FText* ErrorText = nullptr,
+		FText* VerboseText = nullptr);
+
+	UE_DEPRECATED(5.6, "Will be made private. Use Build with cloth collections instead.")
 	void Build(TArray<FChaosClothAssetLodTransitionDataCache>* InOutTransitionCache = nullptr);
 
 	/**
@@ -114,6 +127,7 @@ public:
 	 * Set the skeleton asset using the LOD0 reference skeleton of the skeletal mesh specified in the cloth collection, or set a default reference skeleton if there isn't any.
 	 * @param bRebuildModels Whether to rebuild the simulation and rendering models (editor builds only).
 	 */
+	UE_DEPRECATED(5.6, "Will be made private. Use Build with cloth collections instead.")
 	void UpdateSkeletonFromCollection(bool bRebuildModels);
 
 	/**
@@ -122,7 +136,13 @@ public:
 	 * @param bRebuildModels Whether to rebuild the mesh static models. This is always required, but could be skipped if done soon after.
 	 * @param bRebindMeshes Bindings could be invalid after a change of reference skeleton, use this option to clear the existing binding and bind the sim and render meshes to the root bone.
 	 */
-	void SetReferenceSkeleton(const FReferenceSkeleton* ReferenceSkeleton, bool bRebuildModels, bool bRebindMeshes = true);
+	void SetReferenceSkeleton(const FReferenceSkeleton* ReferenceSkeleton, bool bRebuildModels);
+
+	UE_DEPRECATED(5.6, "bRebindMeshes is deprecated. Use the other versions of SetReferenceSkeleton.")
+	void SetReferenceSkeleton(const FReferenceSkeleton* ReferenceSkeleton, bool bRebuildModels, bool bRebindMeshes)
+	{
+		SetReferenceSkeleton(ReferenceSkeleton, bRebuildModels);
+	}
 
 	UE_DEPRECATED(5.6, "No longer used.")
 	static void OnLodStrippingQualityLevelChanged(IConsoleVariable* /*Variable*/) {}
@@ -172,26 +192,26 @@ PRAGMA_ENABLE_DEPRECATION_WARNINGS
 	/** Load render data from DDC if the data is cached, otherwise generate render data and save into DDC */
 	virtual void CacheDerivedData(FSkinnedAssetCompilationContext* Context) override;
 #endif
-
 	virtual void SetReferenceSkeleton(const FReferenceSkeleton* ReferenceSkeleton) override
 	{
-		SetReferenceSkeleton(ReferenceSkeleton, true, true);
+		SetReferenceSkeleton(ReferenceSkeleton, true);
+	}
+	virtual FGuid GetAssetGuid(int32 /*ModelIndex*/) const override
+	{
+		return AssetGuid;
 	}
-
 	//~ End UChaosClothAssetBase interface
 
 private:
-#if WITH_EDITORONLY_DATA
-	/** Dataflow asset. */
-	UPROPERTY(meta = (DeprecatedProperty, DeprecationMessage = "Dataflow Asset is now stored in DataflowInstance"))
-	TObjectPtr<UDataflow> DataflowAsset_DEPRECATED;
-
-	/** Dataflow Asset terminal node. */
-	UPROPERTY(meta = (DeprecatedProperty, DeprecationMessage = "Dataflow terminal name is now stored in DataflowInstance"))
-	FString DataflowTerminal_DEPRECATED = "ClothAssetTerminal";
-#endif  // WITH_EDITORONLY_DATA
-
 	//~ Begin USkinnedAsset interface
+#if WITH_EDITOR
+	/** Initial step for the building process - Can't be done in parallel. USkinnedAsset Interface. */
+	virtual void BeginBuildInternal(FSkinnedAssetBuildContext& Context) override;
+	/** Thread-safe part. USkinnedAsset Interface. */
+	virtual void ExecuteBuildInternal(FSkinnedAssetBuildContext& Context) override;
+	/** Complete the building process - Can't be done in parallel. USkinnedAsset Interface. */
+	virtual void FinishBuildInternal(FSkinnedAssetBuildContext& Context) override;
+#endif
 	/** Initial step for the Post Load process - Can't be done in parallel. */
 	virtual void BeginPostLoadInternal(FSkinnedAssetPostLoadContext& Context) override;
 	/** Thread-safe part of the Post Load */
@@ -213,6 +233,16 @@ private:
 	/** Build the clothing simulation meshes from the Cloth Collection. */
 	void BuildClothSimulationModel(TArray<FChaosClothAssetLodTransitionDataCache>* InOutTransitionCache = nullptr);
 
+#if WITH_EDITORONLY_DATA
+	/** Dataflow asset. */
+	UPROPERTY(meta = (DeprecatedProperty, DeprecationMessage = "Dataflow Asset is now stored in DataflowInstance"))
+	TObjectPtr<UDataflow> DataflowAsset_DEPRECATED;
+
+	/** Dataflow Asset terminal node. */
+	UPROPERTY(meta = (DeprecatedProperty, DeprecationMessage = "Dataflow terminal name is now stored in DataflowInstance"))
+	FString DataflowTerminal_DEPRECATED = "ClothAssetTerminal";
+#endif  // WITH_EDITORONLY_DATA
+
 	/**
 	 * Skeleton asset used at creation time.
 	 * This is of limited use since this USkeleton's reference skeleton might not necessarily match the one created for this asset.
diff --git a/Engine/Plugins/ChaosClothAsset/Source/ChaosClothAssetEngine/Public/ChaosClothAsset/ClothAssetBase.h b/Engine/Plugins/ChaosClothAsset/Source/ChaosClothAssetEngine/Public/ChaosClothAsset/ClothAssetBase.h
--- a/Engine/Plugins/ChaosClothAsset/Source/ChaosClothAssetEngine/Public/ChaosClothAsset/ClothAssetBase.h
+++ b/Engine/Plugins/ChaosClothAsset/Source/ChaosClothAssetEngine/Public/ChaosClothAsset/ClothAssetBase.h
@@ -52,6 +52,10 @@ public:
 	virtual const UPhysicsAsset* GetPhysicsAssetForModel(int32 ModelIndex) const
 	PURE_VIRTUAL(UChaosClothAssetBase::GetPhysicsAssetForModel, return nullptr;);
 
+	/** Return the asset GUID used to match the render sections for this asset model. */
+	virtual FGuid GetAssetGuid(int32 ModelIndex) const
+	PURE_VIRTUAL(UChaosClothAssetBase::GetAssetGuid, return FGuid(););
+
 #if WITH_EDITOR
 	/** Load the platform render data from the DDC if cached, otherwise generate the data and cache it. */
 	virtual void CacheDerivedData(FSkinnedAssetCompilationContext* /*Context*/) {}
@@ -274,17 +278,6 @@ protected:
 	/** Pre-calculate refpose-to-local transforms. */
 	void CalculateInvRefMatrices();
 
-#if WITH_EDITOR
-	/** Initial step for the building process - Can't be done in parallel. USkinnedAsset Interface. */
-	virtual void BeginBuildInternal(FSkinnedAssetBuildContext& Context) override;
-
-	/** Thread-safe part. USkinnedAsset Interface. */
-	virtual void ExecuteBuildInternal(FSkinnedAssetBuildContext& Context) override;
-
-	/** Complete the building process - Can't be done in parallel. USkinnedAsset Interface. */
-	virtual void FinishBuildInternal(FSkinnedAssetBuildContext& Context) override;
-#endif
-
 	UE_DEPRECATED(5.6, "This property isn't deprecated, but the proper getter and setter should be used instead in oreder to preserve correct behavior.")
 	UPROPERTY(EditAnywhere, Category = "Dataflow")
 	FDataflowInstance DataflowInstance;
@@ -351,6 +344,9 @@ protected:
 	UPROPERTY()
 	FBoxSphereBounds Bounds;
 
+	/** A fence which is used to keep track of the rendering thread releasing the static mesh resources. */
+	FRenderCommandFence ReleaseResourcesFence;
+
 private:
 	enum class EAsyncProperties : uint32;
 	FRIEND_ENUM_CLASS_FLAGS(EAsyncProperties);
@@ -384,9 +380,6 @@ private:
 	/** Rendering data. */
 	TUniquePtr<FSkeletalMeshRenderData> SkeletalMeshRenderData;
 
-	/** A fence which is used to keep track of the rendering thread releasing the static mesh resources. */
-	FRenderCommandFence ReleaseResourcesFence;
-
 	/** Reference skeleton precomputed bases. */
 	TArray<FMatrix44f> RefBasesInvMatrix;
 
diff --git a/Engine/Plugins/ChaosClothAssetEditor/Source/ChaosClothAssetDataflowNodes/Private/ChaosClothAsset/ChaosClothAssetDataflowNodesModule.cpp b/Engine/Plugins/ChaosClothAssetEditor/Source/ChaosClothAssetDataflowNodes/Private/ChaosClothAsset/ChaosClothAssetDataflowNodesModule.cpp
--- a/Engine/Plugins/ChaosClothAssetEditor/Source/ChaosClothAssetDataflowNodes/Private/ChaosClothAsset/ChaosClothAssetDataflowNodesModule.cpp
+++ b/Engine/Plugins/ChaosClothAssetEditor/Source/ChaosClothAssetDataflowNodes/Private/ChaosClothAsset/ChaosClothAssetDataflowNodesModule.cpp
@@ -14,6 +14,7 @@
 #include "ChaosClothAsset/ImportFilePathCustomization.h"
 #include "ChaosClothAsset/ImportNode.h"
 #include "ChaosClothAsset/ImportSimulationCacheNode.h"
+#include "ChaosClothAsset/MakeClothAssetNode.h"
 #include "ChaosClothAsset/MergeClothCollectionsNode.h"
 #include "ChaosClothAsset/ProxyDeformerNode.h"
 #include "ChaosClothAsset/RemeshNode.h"
@@ -98,6 +99,7 @@ namespace UE::Chaos::ClothAsset
 			DATAFLOW_NODE_REGISTER_CREATION_FACTORY(FChaosClothAssetExtractWeightMapNode);
 			DATAFLOW_NODE_REGISTER_CREATION_FACTORY(FChaosClothAssetImportNode);
 			DATAFLOW_NODE_REGISTER_CREATION_FACTORY(FChaosClothAssetImportSimulationCacheNode);
+			DATAFLOW_NODE_REGISTER_CREATION_FACTORY(FChaosClothAssetMakeClothAssetNode);
 			DATAFLOW_NODE_REGISTER_CREATION_FACTORY(FChaosClothAssetMergeClothCollectionsNode_v2);
 			DATAFLOW_NODE_REGISTER_CREATION_FACTORY(FChaosClothAssetProxyDeformerNode_v2);
 			DATAFLOW_NODE_REGISTER_CREATION_FACTORY(FChaosClothAssetRemeshNode_v2);
diff --git a/Engine/Plugins/ChaosClothAssetEditor/Source/ChaosClothAssetDataflowNodes/Private/ChaosClothAsset/MakeClothAssetNode.cpp b/Engine/Plugins/ChaosClothAssetEditor/Source/ChaosClothAssetDataflowNodes/Private/ChaosClothAsset/MakeClothAssetNode.cpp
new file mode 100644
--- /dev/null
+++ b/Engine/Plugins/ChaosClothAssetEditor/Source/ChaosClothAssetDataflowNodes/Private/ChaosClothAsset/MakeClothAssetNode.cpp
@@ -0,0 +1,177 @@
+// Copyright Epic Games, Inc. All Rights Reserved.
+
+#include "ChaosClothAsset/MakeClothAssetNode.h"
+#include "Chaos/CollectionEmbeddedSpringConstraintFacade.h"
+#include "Chaos/CollectionPropertyFacade.h"
+#include "ChaosClothAsset/ClothAsset.h"
+#include "ChaosClothAsset/ClothCollectionGroup.h"
+#include "ChaosClothAsset/ClothDataflowTools.h"
+#include "ChaosClothAsset/ClothGeometryTools.h"
+#include "ChaosClothAsset/CollectionClothFacade.h"
+#include "ChaosClothAsset/CollectionClothSelectionFacade.h"
+#include "Dataflow/DataflowInputOutput.h"
+#include "PhysicsEngine/PhysicsAsset.h"
+
+#include UE_INLINE_GENERATED_CPP_BY_NAME(MakeClothAssetNode)
+
+#define LOCTEXT_NAMESPACE "ChaosClothAssetMakeClothAssetNode"
+
+FChaosClothAssetMakeClothAssetNode::FChaosClothAssetMakeClothAssetNode(const UE::Dataflow::FNodeParameters& InParam, FGuid InGuid)
+	: FDataflowNode(InParam, InGuid)
+{
+	RegisterOutputConnection(&ClothAsset);
+
+	for (int32 Index = 0; Index < NumInitialCollectionLods; ++Index)
+	{
+		AddPins();
+	}
+	check(GetNumInputs() == NumRequiredInputs + NumInitialCollectionLods);  // Update NumRequiredInputs when adding inputs (used by Serialize)
+}
+
+void FChaosClothAssetMakeClothAssetNode::Evaluate(UE::Dataflow::FContext& Context, const FDataflowOutput* Out) const
+{
+	using namespace UE::Chaos::ClothAsset;
+
+	if (Out->IsA<TObjectPtr<const UChaosClothAsset>>(&ClothAsset))
+	{
+		TObjectPtr<UChaosClothAsset> OutClothAsset = NewObject<UChaosClothAsset>();
+
+		TArray<TSharedRef<const FManagedArrayCollection>> InClothCollections = GetCleanedCollectionLodValues(Context);
+		TArray<TSharedRef<const FManagedArrayCollection>>& ClothCollections = OutClothAsset->GetClothCollections();
+
+		FText ErrorText;
+		FText VerboseText;
+		OutClothAsset->Build(InClothCollections, nullptr, &ErrorText, &VerboseText);
+
+		if (!ErrorText.IsEmpty())
+		{
+			// TODO: Use error reporting system
+		}
+
+		SetValue<TObjectPtr<const UChaosClothAsset>>(Context, OutClothAsset, &ClothAsset);
+	}
+}
+
+TArray<UE::Dataflow::FPin> FChaosClothAssetMakeClothAssetNode::AddPins()
+{
+	const int32 Index = CollectionLods.AddDefaulted();
+	const FDataflowInput& Input = RegisterInputArrayConnection(GetConnectionReference(Index));
+	return { { UE::Dataflow::FPin::EDirection::INPUT, Input.GetType(), Input.GetName() } };
+}
+
+TArray<UE::Dataflow::FPin> FChaosClothAssetMakeClothAssetNode::GetPinsToRemove() const
+{
+	const int32 Index = CollectionLods.Num() - 1;
+	check(CollectionLods.IsValidIndex(Index));
+	if (const FDataflowInput* const Input = FindInput(GetConnectionReference(Index)))
+	{
+		return { { UE::Dataflow::FPin::EDirection::INPUT, Input->GetType(), Input->GetName() } };
+	}
+	return Super::GetPinsToRemove();
+}
+
+void FChaosClothAssetMakeClothAssetNode::OnPinRemoved(const UE::Dataflow::FPin& Pin)
+{
+	const int32 Index = CollectionLods.Num() - 1;
+	check(CollectionLods.IsValidIndex(Index));
+#if DO_CHECK
+	const FDataflowInput* const Input = FindInput(GetConnectionReference(Index));
+	check(Input);
+	check(Input->GetName() == Pin.Name);
+	check(Input->GetType() == Pin.Type);
+#endif
+	CollectionLods.SetNum(Index);
+
+	return Super::OnPinRemoved(Pin);
+}
+
+TArray<TSharedRef<const FManagedArrayCollection>> FChaosClothAssetMakeClothAssetNode::GetCleanedCollectionLodValues(UE::Dataflow::FContext& Context) const
+{
+	using namespace UE::Chaos::ClothAsset;
+
+	TArray<TSharedRef<const FManagedArrayCollection>> CollectionLodValues;
+	CollectionLodValues.Reserve(CollectionLods.Num());
+
+	int32 LastValidLodIndex = INDEX_NONE;
+	for (int32 LodIndex = 0; LodIndex < CollectionLods.Num(); ++LodIndex)
+	{
+		TSharedRef<FManagedArrayCollection> CollectionLodValue = MakeShared<FManagedArrayCollection>(GetValue<FManagedArrayCollection>(Context, GetConnectionReference(LodIndex)));
+
+		FCollectionClothFacade ClothFacade(CollectionLodValue);
+		if (ClothFacade.HasValidRenderData())  // The cloth collection must at least have a render mesh
+		{
+			FClothGeometryTools::CleanupAndCompactMesh(CollectionLodValue);
+			LastValidLodIndex = LodIndex;
+		}
+		else if (LastValidLodIndex >= 0)
+		{
+			ClothFacade.DefineSchema();
+			ClothFacade.Initialize(FCollectionClothConstFacade(CollectionLodValues[LastValidLodIndex]));
+
+			FClothDataflowTools::LogAndToastWarning(*this,
+				LOCTEXT("InvalidInputLodNHeadline", "Invalid input LOD."),
+				FText::Format(
+					LOCTEXT("InvalidInputLodNDetails",
+						"Invalid or empty input LOD for LOD {0}.\n"
+						"Using the previous valid LOD {1} instead."),
+					LodIndex,
+					LastValidLodIndex));
+		}
+		else
+		{
+			FClothDataflowTools::LogAndToastWarning(*this,
+				LOCTEXT("InvalidInputLod0Headline", "Invalid input LOD 0."),
+				LOCTEXT("InvalidInputLod0Details",
+					"Invalid or empty input LOD for LOD 0.\n"
+					"LOD 0 cannot be empty in order to construct a valid Cloth Asset."));
+			break;
+		}
+		CollectionLodValues.Emplace(MoveTemp(CollectionLodValue));
+	}
+	return CollectionLodValues;
+}
+
+UE::Dataflow::TConnectionReference<FManagedArrayCollection> FChaosClothAssetMakeClothAssetNode::GetConnectionReference(int32 Index) const
+{
+	return { &CollectionLods[Index], Index, &CollectionLods };
+}
+
+void FChaosClothAssetMakeClothAssetNode::PostSerialize(const FArchive& Ar)
+{
+	// Added pins need to be restored when loading to make sure they get reconnected
+	if (Ar.IsLoading())
+	{
+		for (int32 Index = 0; Index < NumInitialCollectionLods; ++Index)
+		{
+			check(FindInput(GetConnectionReference(Index)));
+		}
+
+		for (int32 Index = NumInitialCollectionLods; Index < CollectionLods.Num(); ++Index)
+		{
+			FindOrRegisterInputArrayConnection(GetConnectionReference(Index));
+		}
+		if (Ar.IsTransacting())
+		{
+			const int32 OrigNumRegisteredInputs = GetNumInputs();
+			check(OrigNumRegisteredInputs >= NumRequiredInputs + NumInitialCollectionLods);
+			const int32 OrigNumCollections = CollectionLods.Num();
+			const int32 OrigNumRegisteredCollections = (OrigNumRegisteredInputs - NumRequiredInputs);
+			if (OrigNumRegisteredCollections > OrigNumCollections)
+			{
+				// Inputs have been removed, temporarily expand ClothAssets so we can get connection references
+				CollectionLods.SetNum(OrigNumRegisteredCollections);
+				for (int32 Index = OrigNumCollections; Index < CollectionLods.Num(); ++Index)
+				{
+					UnregisterInputConnection(GetConnectionReference(Index));
+				}
+				CollectionLods.SetNum(OrigNumCollections);
+			}
+		}
+		else
+		{
+			ensureAlways(CollectionLods.Num() + NumRequiredInputs == GetNumInputs());
+		}
+	}
+}
+
+#undef LOCTEXT_NAMESPACE
diff --git a/Engine/Plugins/ChaosClothAssetEditor/Source/ChaosClothAssetDataflowNodes/Private/ChaosClothAsset/MakeClothAssetNode.h b/Engine/Plugins/ChaosClothAssetEditor/Source/ChaosClothAssetDataflowNodes/Private/ChaosClothAsset/MakeClothAssetNode.h
new file mode 100644
--- /dev/null
+++ b/Engine/Plugins/ChaosClothAssetEditor/Source/ChaosClothAssetDataflowNodes/Private/ChaosClothAsset/MakeClothAssetNode.h
@@ -0,0 +1,46 @@
+// Copyright Epic Games, Inc. All Rights Reserved.
+
+#pragma once 
+
+#include "ChaosClothAsset/ClothLodTransitionDataCache.h"
+#include "Dataflow/DataflowNode.h"
+#include "GeometryCollection/ManagedArrayCollection.h"
+#include "MakeClothAssetNode.generated.h"
+
+class UChaosClothAsset;
+
+/** Cloth terminal node to generate a cloth asset from a cloth collection. */
+USTRUCT(Meta = (Experimental, DataflowCloth))
+struct FChaosClothAssetMakeClothAssetNode : public FDataflowNode
+{
+	GENERATED_USTRUCT_BODY()
+	DATAFLOW_NODE_DEFINE_INTERNAL(FChaosClothAssetMakeClothAssetNode, "MakeClothAsset", "Cloth", "Make Cloth Asset ")
+
+public:
+	FChaosClothAssetMakeClothAssetNode(const UE::Dataflow::FNodeParameters& InParam, FGuid InGuid = FGuid::NewGuid());
+
+private:
+	//~ Begin FDataflowNode interface
+	virtual void Evaluate(UE::Dataflow::FContext& Context, const FDataflowOutput* Out) const override;
+	virtual TArray<UE::Dataflow::FPin> AddPins() override;
+	virtual bool CanAddPin() const override { return true; }
+	virtual bool CanRemovePin() const override { return CollectionLods.Num() > NumInitialCollectionLods; }
+	virtual TArray<UE::Dataflow::FPin> GetPinsToRemove() const override;
+	virtual void OnPinRemoved(const UE::Dataflow::FPin& Pin) override;
+	virtual void PostSerialize(const FArchive& Ar) override;
+	//~ End FDataflowNode interface
+
+	TArray<TSharedRef<const FManagedArrayCollection>> GetCleanedCollectionLodValues(UE::Dataflow::FContext& Context) const;
+	UE::Dataflow::TConnectionReference<FManagedArrayCollection> GetConnectionReference(int32 Index) const;
+
+	/** Input cloth collection for this LOD. */
+	UPROPERTY()
+	TArray<FManagedArrayCollection> CollectionLods;
+
+	/** The cloth asset output. */
+	UPROPERTY(Meta = (DataflowOutput))
+	TObjectPtr<const UChaosClothAsset> ClothAsset;
+
+	static constexpr int32 NumRequiredInputs = 0;
+	static constexpr int32 NumInitialCollectionLods = 1;
+};
diff --git a/Engine/Plugins/ChaosClothAssetEditor/Source/ChaosClothAssetDataflowNodes/Private/ChaosClothAsset/TerminalNode.cpp b/Engine/Plugins/ChaosClothAssetEditor/Source/ChaosClothAssetDataflowNodes/Private/ChaosClothAsset/TerminalNode.cpp
--- a/Engine/Plugins/ChaosClothAssetEditor/Source/ChaosClothAssetDataflowNodes/Private/ChaosClothAsset/TerminalNode.cpp
+++ b/Engine/Plugins/ChaosClothAssetEditor/Source/ChaosClothAssetDataflowNodes/Private/ChaosClothAsset/TerminalNode.cpp
@@ -138,124 +138,15 @@ void FChaosClothAssetTerminalNode_v2::SetAssetValue(TObjectPtr<UObject> Asset, U
 			return;
 		}
 
-		// Reset the asset's collection
-		ClothCollections.Reset(InClothCollections.Num());
+		FText ErrorText;
+		FText VerboseText;
+		ClothAsset->Build(InClothCollections, &LODTransitionDataCache, &ErrorText, &VerboseText);
 
-		// Reset the asset's material list
-		ClothAsset->GetMaterials().Reset();
-
-		// Iterate through the LODs
-		FString PhysicsAssetPathName;
-
-		for (int32 LodIndex = 0; LodIndex < InClothCollections.Num(); ++LodIndex)
-		{
-			// New LOD
-			TSharedRef<FManagedArrayCollection> ClothCollection = MakeShared<FManagedArrayCollection>();
-			UE::Chaos::ClothAsset::FCollectionClothFacade ClothFacade(ClothCollection);
-			ClothFacade.DefineSchema();
-
-			const FCollectionClothConstFacade InClothFacade(InClothCollections[LodIndex]);
-			if (!InClothFacade.HasValidRenderData())  // The cloth collection must at least have a render mesh
-			{
-				FClothDataflowTools::LogAndToastWarning(*this, LOCTEXT("InvalidLODHeadline", "Invalid LOD."),
-					FText::Format(
-						LOCTEXT("InvalidLODDetails", "LOD {0} has no valid data."),
-						LodIndex));
-
-				ClothCollections.Emplace(MoveTemp(ClothCollection));
-				continue;
-			}
-
-			// Copy input LOD to current output LOD
-			ClothFacade.Initialize(InClothFacade);
-
-			// Add this LOD's materials to the asset
-			const int32 NumLodMaterials = ClothFacade.GetNumRenderPatterns();
-
-			TArray<FSkeletalMaterial>& Materials = ClothAsset->GetMaterials();
-			Materials.Reserve(Materials.Num() + NumLodMaterials);
-
-			const TConstArrayView<FString> LodRenderMaterialPathName = ClothFacade.GetRenderMaterialPathName();
-			for (int32 LodMaterialIndex = 0; LodMaterialIndex < NumLodMaterials; ++LodMaterialIndex)
-			{
-				const FString& RenderMaterialPathName = LodRenderMaterialPathName[LodMaterialIndex];
-
-				if (UMaterialInterface* const Material = LoadObject<UMaterialInterface>(ClothAsset, *RenderMaterialPathName, nullptr, LOAD_None, nullptr))
-				{
-					Materials.Emplace(Material, true, false, Material->GetFName());
-				}
-				else
-				{
-					Materials.Emplace();
-				}
-			}
-
-			// Set properties
-			constexpr bool bUpdateExistingProperties = false;
-			Chaos::Softs::FCollectionPropertyMutableFacade PropertyFacade(ClothCollection);
-			PropertyFacade.Append(InClothCollections[LodIndex].ToSharedPtr(), bUpdateExistingProperties);
-
-			// Set selections
-			FCollectionClothSelectionFacade Selection(ClothCollection);
-			FCollectionClothSelectionConstFacade InSelection(InClothCollections[LodIndex]);
-			if (InSelection.IsValid())
-			{
-				Selection.DefineSchema();
-				const TArray<FName> InSelectionNames = InSelection.GetNames();
-				for (const FName& InSelectionName : InSelectionNames)
-				{
-					const FName SelectionGroup = InSelection.GetSelectionGroup(InSelectionName);
-					if (SelectionGroup == ClothCollectionGroup::SimVertices3D ||
-						SelectionGroup == ClothCollectionGroup::SimFaces)
-					{
-						Selection.FindOrAddSelectionSet(InSelectionName, SelectionGroup) = InSelection.GetSelectionSet(InSelectionName);
-					}
-				}
-			}
-
-			// Set springs
-			Chaos::Softs::FEmbeddedSpringFacade EmbeddedSpringFacade(ClothCollection.Get(), ClothCollectionGroup::SimVertices3D);
-			const Chaos::Softs::FEmbeddedSpringFacade InEmbeddedSpringFacade(InClothCollections[LodIndex].Get(), ClothCollectionGroup::SimVertices3D);
-			if (InEmbeddedSpringFacade.IsValid())
-			{
-				EmbeddedSpringFacade.DefineSchema();
-				constexpr int32 VertexOffset = 0;
-				EmbeddedSpringFacade.Append(InEmbeddedSpringFacade, VertexOffset);
-			}
-
-			// Set physics asset only with LOD 0 at the moment
-			if (LodIndex == 0)
-			{
-				using namespace ::Chaos::Softs;
-				PhysicsAssetPathName = InClothFacade.GetPhysicsAssetPathName();
-			}
-
-			ClothCollections.Emplace(MoveTemp(ClothCollection));
-		}
-
-		// Make sure that whatever happens there is always at least one empty LOD to avoid crashing the render data
-		if (ClothCollections.Num() < 1)
+		if (!ErrorText.IsEmpty())
 		{
-			TSharedRef<FManagedArrayCollection> ClothCollection = MakeShared<FManagedArrayCollection>();
-			FCollectionClothFacade ClothFacade(ClothCollection);
-			ClothFacade.DefineSchema();
-			ClothCollections.Emplace(MoveTemp(ClothCollection));
+			FClothDataflowTools::LogAndToastWarning(*this, ErrorText, VerboseText);
 		}
 
-		// Set reference skeleton
-		constexpr bool bRebuildModels = false;  // Avoid rebuilding the asset twice
-		ClothAsset->UpdateSkeletonFromCollection(bRebuildModels);
-
-		// Set physics asset (note: the cloth asset's physics asset is only replaced if a collection path name is found valid)
-		UPhysicsAsset* const PhysicsAsset = !PhysicsAssetPathName.IsEmpty() ?
-			LoadObject<UPhysicsAsset>(ClothAsset, *PhysicsAssetPathName, nullptr, LOAD_None, nullptr) :
-			nullptr;
-		ClothAsset->SetPhysicsAsset(PhysicsAsset);
-		ClothAsset->SetHasVertexColors(true);
-
-		// Rebuild the asset static data
-		ClothAsset->Build(&LODTransitionDataCache);
-
 		// Asset must be resaved
 		ClothAsset->MarkPackageDirty();
 	}
@@ -441,112 +332,15 @@ PRAGMA_ENABLE_DEPRECATION_WARNINGS
 			return;
 		}
 
-		// Reset the asset's collection
-		ClothCollections.Reset(InClothCollections.Num());
-
-		// Reset the asset's material list
-		ClothAsset->GetMaterials().Reset();
-
-		// Iterate through the LODs
-		FString PhysicsAssetPathName;
-
-		for (int32 LodIndex = 0; LodIndex < InClothCollections.Num(); ++LodIndex)
-		{
-			// New LOD
-			TSharedRef<FManagedArrayCollection> ClothCollection = MakeShared<FManagedArrayCollection>();
-			UE::Chaos::ClothAsset::FCollectionClothFacade ClothFacade(ClothCollection);
-			ClothFacade.DefineSchema();
-
-			const FCollectionClothConstFacade InClothFacade(InClothCollections[LodIndex]);
-			if (!InClothFacade.HasValidRenderData())  // The cloth collection must at least have a render mesh
-			{
-				FClothDataflowTools::LogAndToastWarning(*this, LOCTEXT("InvalidLODHeadline", "Invalid LOD."),
-					FText::Format(
-						LOCTEXT("InvalidLODDetails", "LOD {0} has no valid data."),
-						LodIndex));
-
-				ClothCollections.Emplace(MoveTemp(ClothCollection));
-				continue;
-			}
-
-			// Copy input LOD to current output LOD
-			ClothFacade.Initialize(InClothFacade);
-
-			// Add this LOD's materials to the asset
-			const int32 NumLodMaterials = ClothFacade.GetNumRenderPatterns();
-
-			TArray<FSkeletalMaterial>& Materials = ClothAsset->GetMaterials();
-			Materials.Reserve(Materials.Num() + NumLodMaterials);
-
-			const TConstArrayView<FString> LodRenderMaterialPathName = ClothFacade.GetRenderMaterialPathName();
-			for (int32 LodMaterialIndex = 0; LodMaterialIndex < NumLodMaterials; ++LodMaterialIndex)
-			{
-				const FString& RenderMaterialPathName = LodRenderMaterialPathName[LodMaterialIndex];
-
-				if (UMaterialInterface* const Material = LoadObject<UMaterialInterface>(ClothAsset, *RenderMaterialPathName, nullptr, LOAD_None, nullptr))
-				{
-					Materials.Emplace(Material, true, false, Material->GetFName());
-				}
-				else
-				{
-					Materials.Emplace();
-				}
-			}
-
-			// Set properties
-			constexpr bool bUpdateExistingProperties = false;
-			Chaos::Softs::FCollectionPropertyMutableFacade(ClothCollection).Append(InClothCollections[LodIndex].ToSharedPtr(), bUpdateExistingProperties);
-
-			// Set selections
-			FCollectionClothSelectionFacade Selection(ClothCollection);
-			FCollectionClothSelectionConstFacade InSelection(InClothCollections[LodIndex]);
-			if (InSelection.IsValid())
-			{
-				Selection.DefineSchema();
-				const TArray<FName> InSelectionNames = InSelection.GetNames();
-				for (const FName& InSelectionName : InSelectionNames)
-				{
-					const FName SelectionGroup = InSelection.GetSelectionGroup(InSelectionName);
-					if (SelectionGroup == ClothCollectionGroup::SimVertices3D || 
-						SelectionGroup == ClothCollectionGroup::SimFaces)
-					{
-						Selection.FindOrAddSelectionSet(InSelectionName, SelectionGroup) = InSelection.GetSelectionSet(InSelectionName);
-					}
-				}
-			}
-
-			// Set physics asset only with LOD 0 at the moment
-			if (LodIndex == 0)
-			{
-				using namespace ::Chaos::Softs;
-				PhysicsAssetPathName = InClothFacade.GetPhysicsAssetPathName();
-			}
-
-			ClothCollections.Emplace(MoveTemp(ClothCollection));
-		}
+		FText ErrorText;
+		FText VerboseText;
+		ClothAsset->Build(InClothCollections, &LODTransitionDataCache, &ErrorText, &VerboseText);
 
-		// Make sure that whatever happens there is always at least one empty LOD to avoid crashing the render data
-		if (ClothCollections.Num() < 1)
+		if (!ErrorText.IsEmpty())
 		{
-			TSharedRef<FManagedArrayCollection> ClothCollection = MakeShared<FManagedArrayCollection>();
-			FCollectionClothFacade ClothFacade(ClothCollection);
-			ClothFacade.DefineSchema();
-			ClothCollections.Emplace(MoveTemp(ClothCollection));
+			FClothDataflowTools::LogAndToastWarning(*this, ErrorText, VerboseText);
 		}
 
-		// Set reference skeleton
-		constexpr bool bRebuildModels = false;  // Avoid rebuilding the asset twice
-		ClothAsset->UpdateSkeletonFromCollection(bRebuildModels);
-
-		// Set physics asset (note: the cloth asset's physics asset is only replaced if a collection path name is found valid)
-		UPhysicsAsset* const PhysicsAsset = !PhysicsAssetPathName.IsEmpty() ?
-			LoadObject<UPhysicsAsset>(ClothAsset, *PhysicsAssetPathName, nullptr, LOAD_None, nullptr) :
-			nullptr;
-		ClothAsset->SetPhysicsAsset(PhysicsAsset);
-
-		// Rebuild the asset static data
-		ClothAsset->Build(&LODTransitionDataCache);
-
 		// Asset must be resaved
 		ClothAsset->MarkPackageDirty();
 	}
diff --git a/Engine/Plugins/ChaosClothAssetEditor/Source/ChaosClothAssetTools/Private/ChaosClothAsset/ClothingAssetToClothAssetExporter.cpp b/Engine/Plugins/ChaosClothAssetEditor/Source/ChaosClothAssetTools/Private/ChaosClothAsset/ClothingAssetToClothAssetExporter.cpp
--- a/Engine/Plugins/ChaosClothAssetEditor/Source/ChaosClothAssetTools/Private/ChaosClothAsset/ClothingAssetToClothAssetExporter.cpp
+++ b/Engine/Plugins/ChaosClothAssetEditor/Source/ChaosClothAssetTools/Private/ChaosClothAsset/ClothingAssetToClothAssetExporter.cpp
@@ -14,6 +14,7 @@
 #include "GeometryCollection/ManagedArrayCollection.h"
 #include "Materials/Material.h"
 #include "Misc/MessageDialog.h"
+#include "PhysicsEngine/PhysicsAsset.h"
 #include "ClothingAsset.h"
 
 #include UE_INLINE_GENERATED_CPP_BY_NAME(ClothingAssetToClothAssetExporter)
@@ -97,6 +98,8 @@ void UClothingAssetToChaosClothAssetExporter::Export(const UClothingAssetBase* C
 	check(ClothAsset);
 	TArray<TSharedRef<const FManagedArrayCollection>>& ClothCollections = ClothAsset->GetClothCollections();
 	
+	const FString PhysicsAssetPathName = ClothingAssetCommon->PhysicsAsset ? ClothingAssetCommon->PhysicsAsset->GetPathName() : FString();
+
 	const int32 NumLods = ClothingAssetCommon->LodData.Num();
 	ClothCollections.Reset(NumLods);
 
@@ -124,6 +127,9 @@ void UClothingAssetToChaosClothAssetExporter::Export(const UClothingAssetBase* C
 		constexpr bool bAppend = false;
 		FClothGeometryTools::BuildSimMeshFromDynamicMesh(ClothCollection, DynamicMesh, UVChannelIndexNone, FVector2f(1.f), bAppend);
 
+		// Set the physics asset if any
+		ClothFacade.SetPhysicsAssetPathName(PhysicsAssetPathName);
+
 		ClothCollections.Emplace(MoveTemp(ClothCollection));
 	}
 
@@ -136,14 +142,12 @@ void UClothingAssetToChaosClothAssetExporter::Export(const UClothingAssetBase* C
 	{
 		// Make sure that at least one empty LOD is always created
 		TSharedRef<FManagedArrayCollection> ClothCollection = MakeShared<FManagedArrayCollection>();
-		FCollectionClothFacade Cloth(ClothCollection);
-		Cloth.DefineSchema();
+		FCollectionClothFacade ClothFacade(ClothCollection);
+		ClothFacade.DefineSchema();
+		ClothFacade.SetPhysicsAssetPathName(PhysicsAssetPathName);
 		ClothCollections.Emplace(MoveTemp(ClothCollection));
 	}
 
-	// Assign the physics asset if any (must be done after having added the LODs)
-	ClothAsset->SetPhysicsAsset(ClothingAssetCommon->PhysicsAsset);
-
 	// Set the reference skeleton from the skeletal mesh and rebuild the asset (must be done after having added the LODs)
 	USkeletalMesh* const SkeletalMesh = CastChecked<USkeletalMesh>(ClothingAssetCommon->GetOuter());
 	ClothAsset->SetReferenceSkeleton(&SkeletalMesh->GetRefSkeleton());
diff --git a/Engine/Plugins/Experimental/ChaosOutfitAsset/Source/ChaosOutfitAssetEngine/Private/ChaosClothAsset/OutfitAsset.cpp b/Engine/Plugins/Experimental/ChaosOutfitAsset/Source/ChaosOutfitAssetEngine/Private/ChaosClothAsset/OutfitAsset.cpp
--- a/Engine/Plugins/Experimental/ChaosOutfitAsset/Source/ChaosOutfitAssetEngine/Private/ChaosClothAsset/OutfitAsset.cpp
+++ b/Engine/Plugins/Experimental/ChaosOutfitAsset/Source/ChaosOutfitAssetEngine/Private/ChaosClothAsset/OutfitAsset.cpp
@@ -60,11 +60,12 @@ const UPhysicsAsset* UChaosOutfitAsset::GetPhysicsAssetForModel(int32 ModelIndex
 	return Pieces[ModelIndex].PhysicsAsset;
 }
 
-#if WITH_EDITOR
-void UChaosOutfitAsset::BuildLODModel(FSkeletalMeshRenderData& RenderData, const ITargetPlatform* /*TargetPlatform*/, int32 /*LODIndex*/)
+FGuid UChaosOutfitAsset::GetAssetGuid(int32 ModelIndex) const
 {
+	return FGuid(); // Pieces[ModelIndex].AssetGuid;
 }
 
+#if WITH_EDITOR
 FString UChaosOutfitAsset::BuildDerivedDataKey(const ITargetPlatform* TargetPlatform)
 {
 	const FString KeySuffix;
@@ -74,11 +75,6 @@ FString UChaosOutfitAsset::BuildDerivedDataKey(const ITargetPlatform* TargetPlat
 		*KeySuffix
 	);
 }
-
-bool UChaosOutfitAsset::IsInitialBuildDone() const
-{
-	return true;
-}
 #endif // #if WITH_EDITOR
 
 void UChaosOutfitAsset::Build(const TObjectPtr<const UChaosOutfit> Outfit)
diff --git a/Engine/Plugins/Experimental/ChaosOutfitAsset/Source/ChaosOutfitAssetEngine/Public/ChaosClothAsset/OutfitAsset.h b/Engine/Plugins/Experimental/ChaosOutfitAsset/Source/ChaosOutfitAssetEngine/Public/ChaosClothAsset/OutfitAsset.h
--- a/Engine/Plugins/Experimental/ChaosOutfitAsset/Source/ChaosOutfitAssetEngine/Public/ChaosClothAsset/OutfitAsset.h
+++ b/Engine/Plugins/Experimental/ChaosOutfitAsset/Source/ChaosOutfitAssetEngine/Public/ChaosClothAsset/OutfitAsset.h
@@ -34,6 +34,7 @@ private:
 	virtual TSharedPtr<const FChaosClothSimulationModel> GetClothSimulationModel(int32 ModelIndex) const override;
 	virtual const TArray<TSharedRef<const FManagedArrayCollection>>& GetCollections(int32 ModelIndex) const override;
 	virtual const UPhysicsAsset* GetPhysicsAssetForModel(int32 ModelIndex) const override;
+	virtual FGuid GetAssetGuid(int32 ModelIndex) const override;
 	//~ End UChaosClothAssetBase interface
 
 	//~ Begin USkinnedAsset interface
@@ -51,9 +52,11 @@ private:
 	}
 	virtual void SetSkeleton(USkeleton* InSkeleton) override {}
 #if WITH_EDITOR
-	virtual void BuildLODModel(FSkeletalMeshRenderData& RenderData, const ITargetPlatform* TargetPlatform, int32 LODIndex) override;
 	virtual FString BuildDerivedDataKey(const ITargetPlatform* TargetPlatform) override;
-	virtual bool IsInitialBuildDone() const override;
+	virtual bool IsInitialBuildDone() const override
+	{
+		return true;
+	}
 #endif
 #if WITH_EDITORONLY_DATA
 	virtual class FSkeletalMeshModel* GetImportedModel() const override
