Chaos Outfit Asset - Fixed the required bones update in the simulation context for multiple cloth pieces.

#rnx

diff --git a/Engine/Plugins/ChaosClothAsset/Source/ChaosClothAssetEngine/Private/ChaosClothAsset/ClothSimulationContext.cpp b/Engine/Plugins/ChaosClothAsset/Source/ChaosClothAssetEngine/Private/ChaosClothAsset/ClothSimulationContext.cpp
--- a/Engine/Plugins/ChaosClothAsset/Source/ChaosClothAssetEngine/Private/ChaosClothAsset/ClothSimulationContext.cpp
+++ b/Engine/Plugins/ChaosClothAsset/Source/ChaosClothAssetEngine/Private/ChaosClothAsset/ClothSimulationContext.cpp
@@ -101,24 +101,24 @@ namespace UE::Chaos::ClothAsset
 
 		// Update bone matrices
 		RefToLocalMatrices.Reset(NumBones);
+		RequiredExtraBones.Reset();
 
 		bool bSetRefToLocalMatricesToIdentity = true;
 		if (!bIsInitialization)
 		{
 			if (const FSkeletalMeshRenderData* const RenderData = Asset ? Asset->GetResourceForRendering() : nullptr)
 			{
-				if (RenderData->LODRenderData.IsValidIndex(LodIndex))
+				if (RenderData->LODRenderData.IsValidIndex(LodIndex))  // TODO Disable this conditional to drive skeletal meshes
 				{
-					const TArray<uint16>* RequiredExtraBones = nullptr;
 					for (int32 ModelIndex = 0; ModelIndex < Asset->GetNumClothSimulationModels(); ++ModelIndex)
 					{
 						const TSharedPtr<const FChaosClothSimulationModel>& ClothModel = Asset->GetClothSimulationModel(ModelIndex);
 						if (ClothModel && ClothModel->IsValidLodIndex(LodIndex))
 						{
-							RequiredExtraBones = &ClothModel->ClothSimulationLodModels[LodIndex].RequiredExtraBoneIndices;
+							RequiredExtraBones.Append(ClothModel->ClothSimulationLodModels[LodIndex].RequiredExtraBoneIndices);
 						}
 					}
-					ClothComponent.GetCurrentRefToLocalMatrices(RefToLocalMatrices, LodIndex, RequiredExtraBones);
+					ClothComponent.GetCurrentRefToLocalMatrices(RefToLocalMatrices, LodIndex, &RequiredExtraBones);
 					bSetRefToLocalMatricesToIdentity = false;
 				}
 			}
diff --git a/Engine/Plugins/ChaosClothAsset/Source/ChaosClothAssetEngine/Private/ChaosClothAsset/ClothSimulationContext.h b/Engine/Plugins/ChaosClothAsset/Source/ChaosClothAssetEngine/Private/ChaosClothAsset/ClothSimulationContext.h
--- a/Engine/Plugins/ChaosClothAsset/Source/ChaosClothAssetEngine/Private/ChaosClothAsset/ClothSimulationContext.h
+++ b/Engine/Plugins/ChaosClothAsset/Source/ChaosClothAssetEngine/Private/ChaosClothAsset/ClothSimulationContext.h
@@ -56,5 +56,8 @@ namespace UE::Chaos::ClothAsset
 
 		/** Fill the context from data collected from the specified component. */
 		void Fill(const UChaosClothComponent& ClothComponent, float InDeltaTime, float MaxDeltaTime, bool bIsInitialization = false, FClothingSimulationCacheData* CacheData = nullptr);
+
+	private:
+		TArray<uint16> RequiredExtraBones;  // Avoids reallocations while calculating RefToLocalMatrices
 	};
 }
