Chaos Cloth Asset - Fixed crash caused by broken PostLoad static render data on CL 40357808.

#rnx

diff --git a/Engine/Plugins/ChaosClothAsset/Source/ChaosClothAssetEngine/Private/ChaosClothAsset/ClothAsset.cpp b/Engine/Plugins/ChaosClothAsset/Source/ChaosClothAssetEngine/Private/ChaosClothAsset/ClothAsset.cpp
--- a/Engine/Plugins/ChaosClothAsset/Source/ChaosClothAssetEngine/Private/ChaosClothAsset/ClothAsset.cpp
+++ b/Engine/Plugins/ChaosClothAsset/Source/ChaosClothAssetEngine/Private/ChaosClothAsset/ClothAsset.cpp
@@ -157,9 +157,6 @@ namespace UE::Chaos::ClothAsset::Private
 
 UChaosClothAsset::UChaosClothAsset(const FObjectInitializer& ObjectInitializer)
 	: Super(ObjectInitializer)
-#if WITH_EDITORONLY_DATA
-	, MeshModel(MakeShareable(new FSkeletalMeshModel()))
-#endif
 {
 PRAGMA_DISABLE_DEPRECATION_WARNINGS
 	DataflowInstance.SetDataflowTerminal(TEXT("ClothAssetTerminal"));
@@ -668,8 +665,8 @@ void UChaosClothAsset::Build(TArray<FChaosClothAssetLodTransitionDataCache>* InO
 	BuildClothSimulationModel(InOutTransitionCache);
 
 #if WITH_EDITOR
-	// Rebuild LOD Model
-	BuildMeshModel();
+	// Update bounds
+	CalculateBounds();
 
 	// Load/save render data from/to DDC
 	ExecuteBuildInternal(Context);
@@ -689,29 +686,24 @@ void UChaosClothAsset::Build(TArray<FChaosClothAssetLodTransitionDataCache>* InO
 }
 
 #if WITH_EDITOR
-void UChaosClothAsset::BuildMeshModel()
+void UChaosClothAsset::PrepareMeshModel()
 {
-	using namespace UE::Chaos::ClothAsset;
-
-	const int32 NumLods = const_cast<const UChaosClothAsset*>(this)->GetClothCollections().Num();
+	WaitUntilAsyncPropertyReleased(EClothAssetAsyncProperties::ImportedModel);
+	WaitUntilAsyncPropertyReleased(EClothAssetAsyncProperties::ClothCollection);
 
-	// Clear current LOD models
-	check(GetImportedModel());  // MeshModel should always be created in the Cloth Asset constructor WITH_EDITORONLY_DATA
-	GetImportedModel()->LODModels.Reset(NumLods);
+PRAGMA_DISABLE_DEPRECATION_WARNINGS
+	const int32 NumLods = ClothCollections.Num();
 
-	// Get the running platform
-	ITargetPlatformManagerModule& TargetPlatformManager = GetTargetPlatformManagerRef();
-	ITargetPlatform* const TargetPlatform = TargetPlatformManager.GetRunningTargetPlatform();
+	// Reset current LOD models
+	MeshModel = MakeShareable(new FSkeletalMeshModel());
+	MeshModel->LODModels.Reset(NumLods);
 
 	// Rebuild each LOD models
 	for (int32 LodIndex = 0; LodIndex < NumLods; ++LodIndex)
 	{
-		GetImportedModel()->LODModels.Add(new FSkeletalMeshLODModel());
-		FBuilder::BuildLod(GetImportedModel()->LODModels[LodIndex], *this, LodIndex, TargetPlatform);
+		MeshModel->LODModels.Add(new FSkeletalMeshLODModel());
 	}
-
-	// Update bounds
-	CalculateBounds();
+PRAGMA_ENABLE_DEPRECATION_WARNINGS
 }
 #endif  // #if WITH_EDITOR
 
@@ -747,9 +739,11 @@ void UChaosClothAsset::CacheDerivedData(FSkinnedAssetCompilationContext* Context
 	// Create the render data
 	SetResourceForRendering(MakeUnique<FSkeletalMeshRenderData>());
 
+	// Prepare the LOD Model array with the number of LODs for when the cache DDC function regenerates the models
+	PrepareMeshModel();
+
 	// Load render data from DDC, or generate it and save to DDC
 	GetResourceForRendering()->Cache(RunningPlatform, this, Context);
-
 }
 
 void UChaosClothAsset::BuildLODModel(FSkeletalMeshRenderData& RenderData, const ITargetPlatform* TargetPlatform, int32 LODIndex)
diff --git a/Engine/Plugins/ChaosClothAsset/Source/ChaosClothAssetEngine/Public/ChaosClothAsset/ClothAsset.h b/Engine/Plugins/ChaosClothAsset/Source/ChaosClothAssetEngine/Public/ChaosClothAsset/ClothAsset.h
--- a/Engine/Plugins/ChaosClothAsset/Source/ChaosClothAssetEngine/Public/ChaosClothAsset/ClothAsset.h
+++ b/Engine/Plugins/ChaosClothAsset/Source/ChaosClothAssetEngine/Public/ChaosClothAsset/ClothAsset.h
@@ -226,8 +226,8 @@ private:
 	/** Re-calculate the bounds for this asset. */
 	void CalculateBounds();
 
-	/** Build the SkeletalMeshLODModel for this asset. */
-	void BuildMeshModel();
+	/** Prepare the SkeletalMeshLODModel for the CacheDerivedData call. */
+	void PrepareMeshModel();
 #endif
 
 	/** Build the clothing simulation meshes from the Cloth Collection. */
