Chaos Cloth Asset - Fixed the missing simulation model when reloading Cloth Assets without following with a Dataflow evaluation.

#rnx

diff --git a/Engine/Plugins/ChaosClothAsset/Source/ChaosClothAssetEngine/Private/ChaosClothAsset/ClothAsset.cpp b/Engine/Plugins/ChaosClothAsset/Source/ChaosClothAssetEngine/Private/ChaosClothAsset/ClothAsset.cpp
--- a/Engine/Plugins/ChaosClothAsset/Source/ChaosClothAssetEngine/Private/ChaosClothAsset/ClothAsset.cpp
+++ b/Engine/Plugins/ChaosClothAsset/Source/ChaosClothAssetEngine/Private/ChaosClothAsset/ClothAsset.cpp
@@ -413,6 +413,9 @@ void UChaosClothAsset::BeginPostLoadInternal(FSkinnedAssetPostLoadContext& Conte
 	// We're done touching the ClothCollections, so can unlock for read
 	ReleaseAsyncProperty((uint64)EClothAssetAsyncProperties::ClothCollection, ESkinnedAssetAsyncPropertyLockType::WriteOnly);
 
+	// Build the cloth simulation model
+	BuildClothSimulationModel();  // TODO: Cache ClothSimulationModel in the DDC
+
 	// Convert PerPlatForm data to PerQuality if perQuality data have not been serialized.
 	// Also test default value, since PerPlatformData can have Default !=0 and no PerPlatform data overrides.
 	const bool bConvertMinLODData = (MinQualityLevelLOD.PerQuality.Num() == 0 && MinQualityLevelLOD.Default == 0) && (MinLod.PerPlatform.Num() != 0 || MinLod.Default != 0);
