Chaos Cloth Asset:
- Fixed crash caused by the cloth collection Reset() that was adding the optional Solvers group when it wasn't supposed without even initializing it with a default element.
- Added support for TArray<int32> as a user defined type.

#rnx

diff --git a/Engine/Plugins/ChaosClothAsset/Source/ChaosClothAsset/Private/ChaosClothAsset/CollectionClothFacade.cpp b/Engine/Plugins/ChaosClothAsset/Source/ChaosClothAsset/Private/ChaosClothAsset/CollectionClothFacade.cpp
--- a/Engine/Plugins/ChaosClothAsset/Source/ChaosClothAsset/Private/ChaosClothAsset/CollectionClothFacade.cpp
+++ b/Engine/Plugins/ChaosClothAsset/Source/ChaosClothAsset/Private/ChaosClothAsset/CollectionClothFacade.cpp
@@ -174,6 +174,7 @@ namespace UE::Chaos::ClothAsset
 	template CHAOSCLOTHASSET_API uint32 FCollectionClothConstFacade::CalculateUserDefinedAttributesTypeHash<int32>(const FName& GroupName, uint32 PreviousHash) const;
 	template CHAOSCLOTHASSET_API uint32 FCollectionClothConstFacade::CalculateUserDefinedAttributesTypeHash<float>(const FName& GroupName, uint32 PreviousHash) const;
 	template CHAOSCLOTHASSET_API uint32 FCollectionClothConstFacade::CalculateUserDefinedAttributesTypeHash<FVector3f>(const FName& GroupName, uint32 PreviousHash) const;
+	template CHAOSCLOTHASSET_API uint32 FCollectionClothConstFacade::CalculateUserDefinedAttributesTypeHash<TArray<int32>>(const FName& GroupName, uint32 PreviousHash) const;
 
 	const FString& FCollectionClothConstFacade::GetPhysicsAssetPathName() const
 	{
@@ -466,6 +467,7 @@ namespace UE::Chaos::ClothAsset
 	template CHAOSCLOTHASSET_API bool FCollectionClothConstFacade::HasUserDefinedAttribute<int32>(const FName& Name, const FName& GroupName) const;
 	template CHAOSCLOTHASSET_API bool FCollectionClothConstFacade::HasUserDefinedAttribute<float>(const FName& Name, const FName& GroupName) const;
 	template CHAOSCLOTHASSET_API bool FCollectionClothConstFacade::HasUserDefinedAttribute<FVector3f>(const FName& Name, const FName& GroupName) const;
+	template CHAOSCLOTHASSET_API bool FCollectionClothConstFacade::HasUserDefinedAttribute<TArray<int32>>(const FName& Name, const FName& GroupName) const;
 
 	template<typename T UE_REQUIRES_DEFINITION(TIsUserAttributeType<T>::Value)>
 	TArray<FName> FCollectionClothConstFacade::GetUserDefinedAttributeNames(const FName& GroupName) const
@@ -475,6 +477,7 @@ namespace UE::Chaos::ClothAsset
 	template CHAOSCLOTHASSET_API TArray<FName> FCollectionClothConstFacade::GetUserDefinedAttributeNames<int32>(const FName& GroupName) const;
 	template CHAOSCLOTHASSET_API TArray<FName> FCollectionClothConstFacade::GetUserDefinedAttributeNames<float>(const FName& GroupName) const;
 	template CHAOSCLOTHASSET_API TArray<FName> FCollectionClothConstFacade::GetUserDefinedAttributeNames<FVector3f>(const FName& GroupName) const;
+	template CHAOSCLOTHASSET_API TArray<FName> FCollectionClothConstFacade::GetUserDefinedAttributeNames<TArray<int32>>(const FName& GroupName) const;
 
 	template<typename T UE_REQUIRES_DEFINITION(TIsUserAttributeType<T>::Value)>
 	TConstArrayView<T> FCollectionClothConstFacade::GetUserDefinedAttribute(const FName& Name, const FName& GroupName) const
@@ -484,6 +487,7 @@ namespace UE::Chaos::ClothAsset
 	template CHAOSCLOTHASSET_API TConstArrayView<int32> FCollectionClothConstFacade::GetUserDefinedAttribute<int32>(const FName& Name, const FName& GroupName) const;
 	template CHAOSCLOTHASSET_API TConstArrayView<float> FCollectionClothConstFacade::GetUserDefinedAttribute<float>(const FName& Name, const FName& GroupName) const;
 	template CHAOSCLOTHASSET_API TConstArrayView<FVector3f> FCollectionClothConstFacade::GetUserDefinedAttribute<FVector3f>(const FName& Name, const FName& GroupName) const;
+	template CHAOSCLOTHASSET_API TConstArrayView<TArray<int32>> FCollectionClothConstFacade::GetUserDefinedAttribute<TArray<int32>>(const FName& Name, const FName& GroupName) const;
 
 	bool FCollectionClothConstFacade::IsValidClothCollectionGroupName(const FName& GroupName)
 	{
@@ -553,15 +557,18 @@ namespace UE::Chaos::ClothAsset
 
 	void FCollectionClothFacade::Reset()
 	{
-		// Hack to reset all Lods data
-		GetClothCollection()->SetNumElements(0, ClothCollectionGroup::Lods);
-		GetClothCollection()->SetNumElements(1, ClothCollectionGroup::Lods);
+		GetClothCollection()->SetNumElements(0, ClothCollectionGroup::Lods);  // Reset all Lods data
+		GetClothCollection()->SetNumElements(1, ClothCollectionGroup::Lods);  //
 		RemoveAllSimVertices3D();
 		SetNumSimPatterns(0);
 		SetNumRenderPatterns(0);
 		SetNumSeams(0); // Do this after removing SimVertices3D and SimPatterns. Otherwise, Seams will do a bunch of unnecessary work to unseam stuff.
 		SetNumFabrics(0); 
-		GetClothCollection()->SetNumElements(0, ClothCollectionGroup::Solvers);
+		if (HasSolverElement())
+		{
+			GetClothCollection()->SetNumElements(0, ClothCollectionGroup::Solvers);  // Reset all Solvers data
+			GetClothCollection()->SetNumElements(1, ClothCollectionGroup::Solvers);  //
+		}
 	}
 
 	void FCollectionClothFacade::Initialize(const FCollectionClothConstFacade& Other)
@@ -1135,6 +1142,7 @@ namespace UE::Chaos::ClothAsset
 	template CHAOSCLOTHASSET_API bool FCollectionClothFacade::AddUserDefinedAttribute<int32>(const FName& Name, const FName& GroupName, const FName& GroupDependency);
 	template CHAOSCLOTHASSET_API bool FCollectionClothFacade::AddUserDefinedAttribute<float>(const FName& Name, const FName& GroupName, const FName& GroupDependency);
 	template CHAOSCLOTHASSET_API bool FCollectionClothFacade::AddUserDefinedAttribute<FVector3f>(const FName& Name, const FName& GroupName, const FName& GroupDependency);
+	template CHAOSCLOTHASSET_API bool FCollectionClothFacade::AddUserDefinedAttribute<TArray<int32>>(const FName& Name, const FName& GroupName, const FName& GroupDependency);
 
 	void FCollectionClothFacade::RemoveUserDefinedAttribute(const FName& Name, const FName& GroupName)
 	{
@@ -1150,6 +1158,7 @@ namespace UE::Chaos::ClothAsset
 	template CHAOSCLOTHASSET_API TArrayView<int32> FCollectionClothFacade::GetUserDefinedAttribute<int32>(const FName& Name, const FName& GroupName);
 	template CHAOSCLOTHASSET_API TArrayView<float> FCollectionClothFacade::GetUserDefinedAttribute<float>(const FName& Name, const FName& GroupName);
 	template CHAOSCLOTHASSET_API TArrayView<FVector3f> FCollectionClothFacade::GetUserDefinedAttribute<FVector3f>(const FName& Name, const FName& GroupName);
+	template CHAOSCLOTHASSET_API TArrayView<TArray<int32>> FCollectionClothFacade::GetUserDefinedAttribute<TArray<int32>>(const FName& Name, const FName& GroupName);
 
 	void FCollectionClothFacade::SetDefaults()
 	{
diff --git a/Engine/Plugins/ChaosClothAsset/Source/ChaosClothAsset/Public/ChaosClothAsset/IsUserAttributeType.h b/Engine/Plugins/ChaosClothAsset/Source/ChaosClothAsset/Public/ChaosClothAsset/IsUserAttributeType.h
--- a/Engine/Plugins/ChaosClothAsset/Source/ChaosClothAsset/Public/ChaosClothAsset/IsUserAttributeType.h
+++ b/Engine/Plugins/ChaosClothAsset/Source/ChaosClothAsset/Public/ChaosClothAsset/IsUserAttributeType.h
@@ -12,4 +12,5 @@ namespace UE::Chaos::ClothAsset
 	template<> struct TIsUserAttributeType<int32> { static constexpr bool Value = true; };
 	template<> struct TIsUserAttributeType<float> { static constexpr bool Value = true; };
 	template<> struct TIsUserAttributeType<FVector3f> { static constexpr bool Value = true; };
+	template<> struct TIsUserAttributeType<TArray<int32>> { static constexpr bool Value = true; };
 }  // End namespace UE::Chaos::ClothAsset
