Fix a bug where we weren't updating the necessary cached values for Distance-based cloth teleport when loading existing assets.

diff --git a/Engine/Plugins/ChaosClothAsset/Source/ChaosClothAssetEngine/Private/ChaosClothAsset/ClothComponent.cpp b/Engine/Plugins/ChaosClothAsset/Source/ChaosClothAssetEngine/Private/ChaosClothAsset/ClothComponent.cpp
--- a/Engine/Plugins/ChaosClothAsset/Source/ChaosClothAssetEngine/Private/ChaosClothAsset/ClothComponent.cpp
+++ b/Engine/Plugins/ChaosClothAsset/Source/ChaosClothAssetEngine/Private/ChaosClothAsset/ClothComponent.cpp
@@ -266,6 +266,10 @@ void UChaosClothComponent::OnRegister()
 
 	// Register the dataflow simulation interface
 	UE::Dataflow::RegisterSimulationInterface(this);
+
+	// Update teleport thresholds.
+	ClothTeleportDistThresholdSquared = UE::ClothingSimulation::TeleportHelpers::ComputeTeleportDistanceThresholdSquared(TeleportDistanceThreshold);
+	ClothTeleportCosineThresholdInRad = UE::ClothingSimulation::TeleportHelpers::ComputeTeleportCosineRotationThreshold(TeleportRotationThreshold);
 }
 
 void UChaosClothComponent::OnUnregister()
