Chaos Cloth Asset - Added template instanciations missing from CL 40715867.

#rnx

diff --git a/Engine/Plugins/ChaosClothAsset/Source/ChaosClothAsset/Private/ChaosClothAsset/ClothCollection.cpp b/Engine/Plugins/ChaosClothAsset/Source/ChaosClothAsset/Private/ChaosClothAsset/ClothCollection.cpp
--- a/Engine/Plugins/ChaosClothAsset/Source/ChaosClothAsset/Private/ChaosClothAsset/ClothCollection.cpp
+++ b/Engine/Plugins/ChaosClothAsset/Source/ChaosClothAsset/Private/ChaosClothAsset/ClothCollection.cpp
@@ -701,6 +701,7 @@ namespace UE::Chaos::ClothAsset
 	template CHAOSCLOTHASSET_API TArray<FName> FClothCollection::GetUserDefinedAttributeNames<int32>(const FName& GroupName) const;
 	template CHAOSCLOTHASSET_API TArray<FName> FClothCollection::GetUserDefinedAttributeNames<float>(const FName& GroupName) const;
 	template CHAOSCLOTHASSET_API TArray<FName> FClothCollection::GetUserDefinedAttributeNames<FVector3f>(const FName& GroupName) const;
+	template CHAOSCLOTHASSET_API TArray<FName> FClothCollection::GetUserDefinedAttributeNames<TArray<int32>>(const FName& GroupName) const;
 
 	template<typename T UE_REQUIRES_DEFINITION(TIsUserAttributeType<T>::Value)>
 	TManagedArray<T>* FClothCollection::FindOrAddUserDefinedAttribute(const FName& Name, const FName& GroupName, const FName& GroupDependency)
@@ -728,6 +729,7 @@ namespace UE::Chaos::ClothAsset
 	template CHAOSCLOTHASSET_API TManagedArray<int32>* FClothCollection::FindOrAddUserDefinedAttribute<int32>(const FName& Name, const FName& GroupName, const FName& GroupDependency);
 	template CHAOSCLOTHASSET_API TManagedArray<float>* FClothCollection::FindOrAddUserDefinedAttribute<float>(const FName& Name, const FName& GroupName, const FName& GroupDependency);
 	template CHAOSCLOTHASSET_API TManagedArray<FVector3f>* FClothCollection::FindOrAddUserDefinedAttribute<FVector3f>(const FName& Name, const FName& GroupName, const FName& GroupDependency);
+	template CHAOSCLOTHASSET_API TManagedArray<TArray<int32>>* FClothCollection::FindOrAddUserDefinedAttribute<TArray<int32>>(const FName& Name, const FName& GroupName, const FName& GroupDependency);
 
 	void FClothCollection::RemoveUserDefinedAttribute(const FName& Name, const FName& GroupName)
 	{
@@ -768,6 +770,7 @@ namespace UE::Chaos::ClothAsset
 	template CHAOSCLOTHASSET_API bool FClothCollection::HasUserDefinedAttribute<int32>(const FName& Name, const FName& GroupName) const;
 	template CHAOSCLOTHASSET_API bool FClothCollection::HasUserDefinedAttribute<float>(const FName& Name, const FName& GroupName) const;
 	template CHAOSCLOTHASSET_API bool FClothCollection::HasUserDefinedAttribute<FVector3f>(const FName& Name, const FName& GroupName) const;
+	template CHAOSCLOTHASSET_API bool FClothCollection::HasUserDefinedAttribute<TArray<int32>>(const FName& Name, const FName& GroupName) const;
 
 	template<typename T UE_REQUIRES_DEFINITION(TIsUserAttributeType<T>::Value)>
 	const TManagedArray<T>* FClothCollection::GetUserDefinedAttribute(const FName& Name, const FName& GroupName) const
@@ -795,6 +798,7 @@ namespace UE::Chaos::ClothAsset
 	template CHAOSCLOTHASSET_API const TManagedArray<int32>* FClothCollection::GetUserDefinedAttribute<int32>(const FName& Name, const FName& GroupName) const;
 	template CHAOSCLOTHASSET_API const TManagedArray<float>* FClothCollection::GetUserDefinedAttribute<float>(const FName& Name, const FName& GroupName) const;
 	template CHAOSCLOTHASSET_API const TManagedArray<FVector3f>* FClothCollection::GetUserDefinedAttribute<FVector3f>(const FName& Name, const FName& GroupName) const;
+	template CHAOSCLOTHASSET_API const TManagedArray<TArray<int32>>* FClothCollection::GetUserDefinedAttribute<TArray<int32>>(const FName& Name, const FName& GroupName) const;
 
 	template<typename T UE_REQUIRES_DEFINITION(TIsUserAttributeType<T>::Value)>
 	TManagedArray<T>* FClothCollection::GetUserDefinedAttribute(const FName& Name, const FName& GroupName)
@@ -822,6 +826,7 @@ namespace UE::Chaos::ClothAsset
 	template CHAOSCLOTHASSET_API TManagedArray<int32>* FClothCollection::GetUserDefinedAttribute<int32>(const FName& Name, const FName& GroupName);
 	template CHAOSCLOTHASSET_API TManagedArray<float>* FClothCollection::GetUserDefinedAttribute<float>(const FName& Name, const FName& GroupName);
 	template CHAOSCLOTHASSET_API TManagedArray<FVector3f>* FClothCollection::GetUserDefinedAttribute<FVector3f>(const FName& Name, const FName& GroupName);
+	template CHAOSCLOTHASSET_API TManagedArray<TArray<int32>>* FClothCollection::GetUserDefinedAttribute<TArray<int32>>(const FName& Name, const FName& GroupName);
 
 	TArray<FName> FClothCollection::GetValidClothCollectionGroupName()
 	{
