Chaos Cloth Asset - Fixed selection facade constructor constness that allowed modifying of a const managed array collection.

#rnx

diff --git a/Engine/Plugins/ChaosClothAsset/Source/ChaosClothAsset/Private/ChaosClothAsset/CollectionClothSelectionFacade.cpp b/Engine/Plugins/ChaosClothAsset/Source/ChaosClothAsset/Private/ChaosClothAsset/CollectionClothSelectionFacade.cpp
--- a/Engine/Plugins/ChaosClothAsset/Source/ChaosClothAsset/Private/ChaosClothAsset/CollectionClothSelectionFacade.cpp
+++ b/Engine/Plugins/ChaosClothAsset/Source/ChaosClothAsset/Private/ChaosClothAsset/CollectionClothSelectionFacade.cpp
@@ -109,6 +109,11 @@ PRAGMA_ENABLE_DEPRECATION_WARNINGS
 
 	// --------------- FCollectionClothSelectionFacade -------------------------------
 
+	FCollectionClothSelectionFacade::FCollectionClothSelectionFacade(const TSharedRef<FManagedArrayCollection>& ManagedArrayCollection) :
+		FCollectionClothSelectionConstFacade(ManagedArrayCollection)
+	{}
+
+	// UE_DEPRECATED(5.6, "Use the constructor that takes a non const ManagedArrayCollection instead.")
 	FCollectionClothSelectionFacade::FCollectionClothSelectionFacade(const TSharedRef<const FManagedArrayCollection>& ManagedArrayCollection) :
 		FCollectionClothSelectionConstFacade(ManagedArrayCollection)
 	{}
diff --git a/Engine/Plugins/ChaosClothAsset/Source/ChaosClothAsset/Public/ChaosClothAsset/CollectionClothSelectionFacade.h b/Engine/Plugins/ChaosClothAsset/Source/ChaosClothAsset/Public/ChaosClothAsset/CollectionClothSelectionFacade.h
--- a/Engine/Plugins/ChaosClothAsset/Source/ChaosClothAsset/Public/ChaosClothAsset/CollectionClothSelectionFacade.h
+++ b/Engine/Plugins/ChaosClothAsset/Source/ChaosClothAsset/Public/ChaosClothAsset/CollectionClothSelectionFacade.h
@@ -74,6 +74,9 @@ namespace UE::Chaos::ClothAsset
 	class FCollectionClothSelectionFacade final : public FCollectionClothSelectionConstFacade
 	{
 	public:
+		CHAOSCLOTHASSET_API explicit FCollectionClothSelectionFacade(const TSharedRef<FManagedArrayCollection>& ManagedArrayCollection);
+
+		UE_DEPRECATED(5.6, "Use the constructor that takes a non const ManagedArrayCollection instead.")
 		CHAOSCLOTHASSET_API explicit FCollectionClothSelectionFacade(const TSharedRef<const FManagedArrayCollection>& ManagedArrayCollection);
 
 		FCollectionClothSelectionFacade() = delete;
