Chaos Cloth Asset:
- Fixed build logic to protect the Cloth Asset from render crashes when the render data is missing the bone map.
- Deprecated the SetReferenceSkeleton functions instead of simply removing bRebindMeshes. The Skeleton must be set in the collection that is used to build the Cloth Asset, not afterwards.
- Deprecated CopySimMeshToRenderMesh. The render mesh must be created inside the collection that is used to build the Cloth Asset, not afterwards.
- Updated the Skeletal Mesh Cloth Asset exporter.

#rnx

diff --git a/Engine/Plugins/ChaosClothAsset/Source/ChaosClothAssetEngine/Private/ChaosClothAsset/ClothAsset.cpp b/Engine/Plugins/ChaosClothAsset/Source/ChaosClothAssetEngine/Private/ChaosClothAsset/ClothAsset.cpp
--- a/Engine/Plugins/ChaosClothAsset/Source/ChaosClothAssetEngine/Private/ChaosClothAsset/ClothAsset.cpp
+++ b/Engine/Plugins/ChaosClothAsset/Source/ChaosClothAssetEngine/Private/ChaosClothAsset/ClothAsset.cpp
@@ -33,7 +33,7 @@
 
 // If Chaos cloth asset derived data needs to be rebuilt (new format, serialization differences, etc.) replace the version GUID below with a new one. 
 // In case of merge conflicts with DDC versions, you MUST generate a new GUID and set this new GUID as the version.
-#define CHAOS_CLOTH_ASSET_DERIVED_DATA_VERSION TEXT("5C041E93563044A69646B5E98F152B7C")
+#define CHAOS_CLOTH_ASSET_DERIVED_DATA_VERSION TEXT("715113CDDBB84BA0985F4D3E74548FC2")
 
 #define LOCTEXT_NAMESPACE "ChaosClothAsset"
 
@@ -46,7 +46,20 @@ namespace UE::Chaos::ClothAsset::Private
 		bClothCollectionOnlyCookRequiredFacades,
 		TEXT("Default setting for culling managed arrays on the cloth collection during the cook. Default[true]"));
 
-	::Chaos::FChaosArchive& Serialize(::Chaos::FChaosArchive& Ar, TArray<TSharedRef<const FManagedArrayCollection>>& ClothCollections)
+	static bool HasValidSkinweights(const TConstArrayView<TArray<int32>> BoneIndices, TConstArrayView<TArray<float>> BoneWeights)
+	{
+		check(BoneIndices.Num() == BoneWeights.Num());
+		for (int32 Index = 0; Index < BoneIndices.Num(); ++Index)
+		{
+			if (!BoneIndices[Index].Num() || !BoneWeights[Index].Num() || BoneIndices[Index].Num() != BoneWeights[Index].Num())
+			{
+				return false;
+			}
+		}
+		return true;
+	}
+
+	static ::Chaos::FChaosArchive& Serialize(::Chaos::FChaosArchive& Ar, TArray<TSharedRef<const FManagedArrayCollection>>& ClothCollections)
 	{
 		Ar.UsingCustomVersion(FUE5MainStreamObjectVersion::GUID);
 
@@ -114,7 +127,7 @@ namespace UE::Chaos::ClothAsset::Private
 		}
 	}
 
-	TArray<TSharedRef<const FManagedArrayCollection>> TrimOnCook(const FString InAssetName, const TArray<TSharedRef<const FManagedArrayCollection>>& InClothCollections)
+	static TArray<TSharedRef<const FManagedArrayCollection>> TrimOnCook(const FString InAssetName, const TArray<TSharedRef<const FManagedArrayCollection>>& InClothCollections)
 	{
 		int32 Index = 0;
 #if WITH_EDITORONLY_DATA
@@ -180,6 +193,8 @@ PRAGMA_ENABLE_DEPRECATION_WARNINGS
 
 void UChaosClothAsset::Serialize(FArchive& Ar)
 {
+	using namespace UE::Chaos::ClothAsset;
+
 	LLM_SCOPE_BYNAME(TEXT("Physics/Cloth"));
 	Super::Serialize(Ar);
 
@@ -191,14 +206,14 @@ void UChaosClothAsset::Serialize(FArchive& Ar)
 	if (bCooked && Ar.IsSaving())
 	{
 		TArray<TSharedRef<const FManagedArrayCollection>> OutputCollections = 
-			UE::Chaos::ClothAsset::Private::TrimOnCook(GetPathName(), GetClothCollections());
+			Private::TrimOnCook(GetPathName(), GetClothCollections());
 		Chaos::FChaosArchive ChaosArchive(Ar);
-		UE::Chaos::ClothAsset::Private::Serialize(ChaosArchive, OutputCollections);
+		Private::Serialize(ChaosArchive, OutputCollections);
 	}
 	else
 	{
 		Chaos::FChaosArchive ChaosArchive(Ar);
-		UE::Chaos::ClothAsset::Private::Serialize(ChaosArchive, GetClothCollections());
+		Private::Serialize(ChaosArchive, GetClothCollections());
 	}
 
 	if (Ar.CustomVer(FFortniteMainBranchObjectVersion::GUID) < FFortniteMainBranchObjectVersion::AddClothAssetBase)
@@ -206,6 +221,29 @@ void UChaosClothAsset::Serialize(FArchive& Ar)
 		Ar << GetRefSkeleton(); // Moved to Cloth Asset Base serialization
 	}
 
+	if (Ar.CustomVer(FFortniteMainBranchObjectVersion::GUID) < FFortniteMainBranchObjectVersion::ClothAssetSkinweightsValidation)
+	{
+		// Fix the skeleton mesh binding, which can cause crashes in the render code, or make the sim mesh disappear when missing
+		for (TSharedRef<const FManagedArrayCollection>& ClothCollection : GetClothCollections())
+		{
+			const FCollectionClothConstFacade ClothConstFacade(ClothCollection);
+			if (ClothConstFacade.IsValid())
+			{
+				const bool bHasValidSimSkinweights = Private::HasValidSkinweights(ClothConstFacade.GetSimBoneIndices(), ClothConstFacade.GetSimBoneWeights());
+				const bool bHasValidRenderSkinweights = Private::HasValidSkinweights(ClothConstFacade.GetRenderBoneIndices(), ClothConstFacade.GetRenderBoneWeights());
+				if (!bHasValidSimSkinweights || !bHasValidRenderSkinweights)
+				{
+					TSharedRef<FManagedArrayCollection> NewClothCollection = MakeShared<FManagedArrayCollection>(*ClothCollection);
+					FClothGeometryTools::BindMeshToRootBone(NewClothCollection, !bHasValidSimSkinweights, !bHasValidRenderSkinweights);
+					ClothCollection = MoveTemp(NewClothCollection);
+
+					UE_CLOG(!bHasValidSimSkinweights, LogChaosClothAsset, Warning, TEXT("%s had invalid simulation mesh skin weights. This asset must be resaved."), *GetFullName());
+					UE_CLOG(!bHasValidRenderSkinweights, LogChaosClothAsset, Warning, TEXT("%s had invalid render mesh skin weights. This asset must be resaved."), *GetFullName());
+				}
+			}
+		}
+	}
+
 	if (bCooked && !IsTemplate() && !Ar.IsCountingMemory())  // Counting of these resources are done in GetResourceSizeEx, so skip these when counting memory
 	{
 		{
@@ -515,6 +553,7 @@ void UChaosClothAsset::Build(
 
 	// Iterate through the LODs
 	FString PhysicsAssetPathName;
+	FString SkeletalMeshPathName;
 
 	for (int32 LodIndex = 0; LodIndex < InClothCollections.Num(); ++LodIndex)
 	{
@@ -600,11 +639,20 @@ void UChaosClothAsset::Build(
 			EmbeddedSpringFacade.Append(InEmbeddedSpringFacade, VertexOffset);
 		}
 
-		// Set physics asset only with LOD 0 at the moment
+		// Set physics asset and skeleton source only with LOD 0 at the moment
 		if (LodIndex == 0)
 		{
 			using namespace ::Chaos::Softs;
 			PhysicsAssetPathName = InClothFacade.GetPhysicsAssetPathName();
+			SkeletalMeshPathName = InClothFacade.GetSkeletalMeshPathName();
+		}
+
+		// Fix the skeleton mesh binding if needed, which can cause crashes in the render code, or make the sim mesh disappear
+		const bool bHasValidSimSkinweights = Private::HasValidSkinweights(InClothFacade.GetSimBoneIndices(), InClothFacade.GetSimBoneWeights());
+		const bool bHasValidRenderSkinweights = Private::HasValidSkinweights(InClothFacade.GetRenderBoneIndices(), InClothFacade.GetRenderBoneWeights());
+		if (!ensureAlwaysMsgf(bHasValidSimSkinweights && bHasValidRenderSkinweights, TEXT("A Dataflow node, likely an import node, has generated missing or invalid skin weights in this collection LOD. This must be fixed ASAP!")))
+		{
+			FClothGeometryTools::BindMeshToRootBone(ClothCollection, !bHasValidSimSkinweights, !bHasValidRenderSkinweights);
 		}
 
 		OutClothCollections.Emplace(MoveTemp(ClothCollection));
@@ -620,10 +668,11 @@ void UChaosClothAsset::Build(
 	}
 
 	// Set reference skeleton
-	constexpr bool bRebuildModels = false;  // Avoid rebuilding the asset twice
-PRAGMA_DISABLE_DEPRECATION_WARNINGS  // Temporary until moved to private
-	UpdateSkeletonFromCollection(bRebuildModels);
-PRAGMA_ENABLE_DEPRECATION_WARNINGS
+	USkeletalMesh* const SkeletalMesh = SkeletalMeshPathName.IsEmpty() ? nullptr :
+		LoadObject<USkeletalMesh>(nullptr, *SkeletalMeshPathName, nullptr, LOAD_None, nullptr);
+
+	SetSkeleton(SkeletalMesh ? SkeletalMesh->GetSkeleton() : nullptr); // For completion only, this is not being used and might mismatch the skeletal mesh's reference skeleton
+	Super::SetReferenceSkeleton(SkeletalMesh ? &SkeletalMesh->GetRefSkeleton() : nullptr);
 
 	// Set physics asset (note: the cloth asset's physics asset is only replaced if a collection path name is found valid)
 	PhysicsAsset = !PhysicsAssetPathName.IsEmpty() ?
@@ -842,7 +891,7 @@ void UChaosClothAsset::SetPhysicsAsset(UPhysicsAsset* InPhysicsAsset)
 	PhysicsAsset = InPhysicsAsset;
 }
 
-void UChaosClothAsset::SetReferenceSkeleton(const FReferenceSkeleton* ReferenceSkeleton, bool bRebuildModels)
+void UChaosClothAsset::SetReferenceSkeleton(const FReferenceSkeleton* ReferenceSkeleton, bool bRebuildModels, bool bRebindMeshes)
 {
 	using namespace UE::Chaos::ClothAsset;
 
@@ -858,7 +907,7 @@ PRAGMA_ENABLE_DEPRECATION_WARNINGS
 	}
 }
 
-void UChaosClothAsset::UpdateSkeletonFromCollection(bool bRebuildModels)
+void UChaosClothAsset::UpdateSkeletonFromCollection(bool /*bRebuildModels*/)
 {
 	using namespace UE::Chaos::ClothAsset;
 
@@ -871,8 +920,7 @@ void UChaosClothAsset::UpdateSkeletonFromCollection(bool bRebuildModels)
 		LoadObject<USkeletalMesh>(nullptr, *SkeletalMeshPathName, nullptr, LOAD_None, nullptr);
 
 	SetSkeleton(SkeletalMesh ? SkeletalMesh->GetSkeleton() : nullptr); // For completion only, this is not being used and might mismatch the skeletal mesh's reference skeleton
-
-	SetReferenceSkeleton(SkeletalMesh ? &SkeletalMesh->GetRefSkeleton() : nullptr, bRebuildModels);
+	Super::SetReferenceSkeleton(SkeletalMesh ? &SkeletalMesh->GetRefSkeleton() : nullptr);
 }
 
 void UChaosClothAsset::CopySimMeshToRenderMesh(UMaterialInterface* Material)
diff --git a/Engine/Plugins/ChaosClothAsset/Source/ChaosClothAssetEngine/Public/ChaosClothAsset/ClothAsset.h b/Engine/Plugins/ChaosClothAsset/Source/ChaosClothAssetEngine/Public/ChaosClothAsset/ClothAsset.h
--- a/Engine/Plugins/ChaosClothAsset/Source/ChaosClothAssetEngine/Public/ChaosClothAsset/ClothAsset.h
+++ b/Engine/Plugins/ChaosClothAsset/Source/ChaosClothAssetEngine/Public/ChaosClothAsset/ClothAsset.h
@@ -118,31 +118,17 @@ public:
 	 * This is useful to visualize the simulation mesh, or when the simulation mesh can be used for both simulation and rendering.
 	 * @param Material The material used to render. If none is specified the default debug cloth material will be used.
 	 */
+	UE_DEPRECATED(5.6, "Use FClothGeometryTools:CopySimMeshToRenderMesh on the cloth collections instead.")
 	void CopySimMeshToRenderMesh(UMaterialInterface* Material = nullptr);
 
 	/** Set the physics asset for this cloth. */
 	void SetPhysicsAsset(UPhysicsAsset* InPhysicsAsset);
 
-	/**
-	 * Set the skeleton asset using the LOD0 reference skeleton of the skeletal mesh specified in the cloth collection, or set a default reference skeleton if there isn't any.
-	 * @param bRebuildModels Whether to rebuild the simulation and rendering models (editor builds only).
-	 */
 	UE_DEPRECATED(5.6, "Will be made private. Use Build with cloth collections instead.")
 	void UpdateSkeletonFromCollection(bool bRebuildModels);
 
-	/**
-	 * Set the specified reference skeleton for this cloth, and rebuild the models.
-	 * Use nullptr to set to a default single bone root reference skeleton.
-	 * @param bRebuildModels Whether to rebuild the mesh static models. This is always required, but could be skipped if done soon after.
-	 * @param bRebindMeshes Bindings could be invalid after a change of reference skeleton, use this option to clear the existing binding and bind the sim and render meshes to the root bone.
-	 */
-	void SetReferenceSkeleton(const FReferenceSkeleton* ReferenceSkeleton, bool bRebuildModels);
-
-	UE_DEPRECATED(5.6, "bRebindMeshes is deprecated. Use the other versions of SetReferenceSkeleton.")
-	void SetReferenceSkeleton(const FReferenceSkeleton* ReferenceSkeleton, bool bRebuildModels, bool bRebindMeshes)
-	{
-		SetReferenceSkeleton(ReferenceSkeleton, bRebuildModels);
-	}
+	UE_DEPRECATED(5.6, "SetReferenceSkeleton is deprecated. Skeletons must only been set through the cloth collections.")
+	void SetReferenceSkeleton(const FReferenceSkeleton* ReferenceSkeleton, bool bRebuildModels, bool bRebindMeshes);
 
 	UE_DEPRECATED(5.6, "No longer used.")
 	static void OnLodStrippingQualityLevelChanged(IConsoleVariable* /*Variable*/) {}
@@ -192,9 +178,10 @@ PRAGMA_ENABLE_DEPRECATION_WARNINGS
 	/** Load render data from DDC if the data is cached, otherwise generate render data and save into DDC */
 	virtual void CacheDerivedData(FSkinnedAssetCompilationContext* Context) override;
 #endif
+	UE_DEPRECATED(5.6, "SetReferenceSkeleton is deprecated. Skeletons must only been set through the cloth collections.")
 	virtual void SetReferenceSkeleton(const FReferenceSkeleton* ReferenceSkeleton) override
 	{
-		SetReferenceSkeleton(ReferenceSkeleton, true);
+		Super::SetReferenceSkeleton(ReferenceSkeleton);
 	}
 	virtual FGuid GetAssetGuid(int32 /*ModelIndex*/) const override
 	{
diff --git a/Engine/Plugins/ChaosClothAssetEditor/Source/ChaosClothAssetTools/Private/ChaosClothAsset/ClothingAssetToClothAssetExporter.cpp b/Engine/Plugins/ChaosClothAssetEditor/Source/ChaosClothAssetTools/Private/ChaosClothAsset/ClothingAssetToClothAssetExporter.cpp
--- a/Engine/Plugins/ChaosClothAssetEditor/Source/ChaosClothAssetTools/Private/ChaosClothAsset/ClothingAssetToClothAssetExporter.cpp
+++ b/Engine/Plugins/ChaosClothAssetEditor/Source/ChaosClothAssetTools/Private/ChaosClothAsset/ClothingAssetToClothAssetExporter.cpp
@@ -96,11 +96,13 @@ void UClothingAssetToChaosClothAssetExporter::Export(const UClothingAssetBase* C
 
 	UChaosClothAsset* const ClothAsset = CastChecked<UChaosClothAsset>(ExportedAsset);
 	check(ClothAsset);
-	TArray<TSharedRef<const FManagedArrayCollection>>& ClothCollections = ClothAsset->GetClothCollections();
-	
+
+	static const FString DefaultMaterialPathName = FString(TEXT("/Engine/EditorMaterials/Cloth/CameraLitDoubleSided.CameraLitDoubleSided"));
 	const FString PhysicsAssetPathName = ClothingAssetCommon->PhysicsAsset ? ClothingAssetCommon->PhysicsAsset->GetPathName() : FString();
+	const FString SkeletalMeshPathName = CastChecked<USkeletalMesh>(ClothingAssetCommon->GetOuter())->GetPathName();
 
 	const int32 NumLods = ClothingAssetCommon->LodData.Num();
+	TArray<TSharedRef<const FManagedArrayCollection>> ClothCollections;
 	ClothCollections.Reset(NumLods);
 
 	// Create the LODs
@@ -129,31 +131,34 @@ void UClothingAssetToChaosClothAssetExporter::Export(const UClothingAssetBase* C
 
 		// Set the physics asset if any
 		ClothFacade.SetPhysicsAssetPathName(PhysicsAssetPathName);
+		// Set the skeleton
+		ClothFacade.SetSkeletalMeshPathName(SkeletalMeshPathName);
 
-		ClothCollections.Emplace(MoveTemp(ClothCollection));
-	}
+		// Bind the sim mesh to root bone
+		constexpr bool bBindSimMesh = true;
+		constexpr bool bBindRenderMesh = false;  // Will get bound below in CopySimMeshToRenderMesh
+		FClothGeometryTools::BindMeshToRootBone(ClothCollection, bBindSimMesh, bBindRenderMesh);
 
-	if (NumLods)
-	{
 		// Set the render mesh to duplicate the sim mesh
-		ClothAsset->CopySimMeshToRenderMesh();
+		constexpr bool bSingleRenderPattern = true;
+		FClothGeometryTools::CopySimMeshToRenderMesh(ClothCollection, DefaultMaterialPathName, bSingleRenderPattern);
+
+		ClothCollections.Emplace(MoveTemp(ClothCollection));
 	}
-	else
+
+	if (!NumLods)
 	{
 		// Make sure that at least one empty LOD is always created
 		TSharedRef<FManagedArrayCollection> ClothCollection = MakeShared<FManagedArrayCollection>();
 		FCollectionClothFacade ClothFacade(ClothCollection);
 		ClothFacade.DefineSchema();
 		ClothFacade.SetPhysicsAssetPathName(PhysicsAssetPathName);
+		ClothFacade.SetSkeletalMeshPathName(SkeletalMeshPathName);
 		ClothCollections.Emplace(MoveTemp(ClothCollection));
 	}
 
-	// Set the reference skeleton from the skeletal mesh and rebuild the asset (must be done after having added the LODs)
-	USkeletalMesh* const SkeletalMesh = CastChecked<USkeletalMesh>(ClothingAssetCommon->GetOuter());
-	ClothAsset->SetReferenceSkeleton(&SkeletalMesh->GetRefSkeleton());
-
-	// Save the USkeleton for completion, although this is not being used and might mismatch the skeletal mesh's reference skeleton
-	ClothAsset->SetSkeleton(SkeletalMesh->GetSkeleton());
+	// Rebuild the asset
+	ClothAsset->Build(ClothCollections);
 }
 
 #undef LOCTEXT_NAMESPACE
diff --git a/Engine/Source/Runtime/Core/Public/UObject/FortniteMainBranchObjectVersions.inl b/Engine/Source/Runtime/Core/Public/UObject/FortniteMainBranchObjectVersions.inl
--- a/Engine/Source/Runtime/Core/Public/UObject/FortniteMainBranchObjectVersions.inl
+++ b/Engine/Source/Runtime/Core/Public/UObject/FortniteMainBranchObjectVersions.inl
@@ -604,5 +604,8 @@ FORTNITE_MAIN_VERSION(MeshMaterialSlotOverlayMaterialAdded, 196)
 // Convert Sustrate glint density properly
 FORTNITE_MAIN_VERSION(ConvertGlintDensity, 197)
 
+// Introduced skinweight validation to avoid render crashes and disappearing simulation meshes
+FORTNITE_MAIN_VERSION(ClothAssetSkinweightsValidation, 198)
+
 // -----<new versions can be added above this line>-------------------------------------------------
 #undef FORTNITE_MAIN_VERSION
\ No newline at end of file
