Chaos Cloth Asset - Bumped the Cloth Asset DDC key after recent skeletal rendering serialization changes in CL 40789459 that could cause potential crashes while reloading older DDC render data.

#rnx

diff --git a/Engine/Plugins/ChaosClothAsset/Source/ChaosClothAssetEngine/Private/ChaosClothAsset/ClothAsset.cpp b/Engine/Plugins/ChaosClothAsset/Source/ChaosClothAssetEngine/Private/ChaosClothAsset/ClothAsset.cpp
--- a/Engine/Plugins/ChaosClothAsset/Source/ChaosClothAssetEngine/Private/ChaosClothAsset/ClothAsset.cpp
+++ b/Engine/Plugins/ChaosClothAsset/Source/ChaosClothAssetEngine/Private/ChaosClothAsset/ClothAsset.cpp
@@ -33,7 +33,7 @@
 
 // If Chaos cloth asset derived data needs to be rebuilt (new format, serialization differences, etc.) replace the version GUID below with a new one. 
 // In case of merge conflicts with DDC versions, you MUST generate a new GUID and set this new GUID as the version.
-#define CHAOS_CLOTH_ASSET_DERIVED_DATA_VERSION TEXT("715113CDDBB84BA0985F4D3E74548FC2")
+#define CHAOS_CLOTH_ASSET_DERIVED_DATA_VERSION TEXT("479D81081F3A4A22B3C22ED4B278680E")
 
 #define LOCTEXT_NAMESPACE "ChaosClothAsset"
 
