Cloth/Animation: Remove cloth collision dependency on transforms changed delegate

This allows cloth collision to work with solely AnimNext-driven meshes

diff --git a/Engine/Plugins/ChaosClothAsset/Source/ChaosClothAssetEngine/Private/ChaosClothAsset/CollisionSources.cpp b/Engine/Plugins/ChaosClothAsset/Source/ChaosClothAssetEngine/Private/ChaosClothAsset/CollisionSources.cpp
--- a/Engine/Plugins/ChaosClothAsset/Source/ChaosClothAssetEngine/Private/ChaosClothAsset/CollisionSources.cpp
+++ b/Engine/Plugins/ChaosClothAsset/Source/ChaosClothAssetEngine/Private/ChaosClothAsset/CollisionSources.cpp
@@ -31,15 +31,8 @@ namespace UE::Chaos::ClothAsset
 
 			if (!FoundCollisionSource)
 			{
-				// Delegate called after the transform buffer flip, so that the cloths' transform gets updated when the component owning the cloth isn't moving, but the collision source is
-				const FSimpleDelegate OnBoneTransformsFinalizedDelegate = FSimpleDelegate::CreateLambda(
-					[this]()
-					{
-						++Version;
-					});
-
 				// Add the new collision source
-				CollisionSources.Emplace(SourceComponent, SourcePhysicsAsset, OnBoneTransformsFinalizedDelegate, bUseSphylsOnly);
+				CollisionSources.Emplace(SourceComponent, SourcePhysicsAsset, bUseSphylsOnly);
 
 				// Add prerequisite so we don't end up with a frame delay
 				OwnerComponent->PrimaryComponentTick.AddPrerequisite(SourceComponent, SourceComponent->PrimaryComponentTick);
@@ -110,6 +103,48 @@ namespace UE::Chaos::ClothAsset
 		}
 	}
 
+	bool FCollisionSources::IsCollisionDataUpToDate(int32 InVersion) const
+	{
+		if (InVersion != Version)
+		{
+			return false;
+		}
+
+		for (const FCollisionSource& CollisionSource : CollisionSources)
+		{
+			if (CollisionSource.ResolvedSourceComponent)
+			{
+				if (CollisionSource.BoneTransformRevisionNumber != CollisionSource.ResolvedSourceComponent->GetBoneTransformRevisionNumber())
+				{
+					return false;
+				}
+			}
+		}
+
+		return true;
+	}
+
+	void FCollisionSources::MarkCollisionDataUpToDate(int32 InVersion)
+	{
+		Version = InVersion;
+
+		for (FCollisionSource& CollisionSource : CollisionSources)
+		{
+			if (CollisionSource.ResolvedSourceComponent)
+			{
+				CollisionSource.BoneTransformRevisionNumber = CollisionSource.ResolvedSourceComponent->GetBoneTransformRevisionNumber();
+			}
+		}
+	}
+
+	void FCollisionSources::ResolveWeakComponentPtrs()
+	{
+		for (FCollisionSource& CollisionSource : CollisionSources)
+		{
+			CollisionSource.ResolvedSourceComponent = CollisionSource.SourceComponent.Get();
+		}
+	}
+
 	void FCollisionSources::ExtractEnvironmentalCollisionData(FClothCollisionData& CollisionData)
 	{
 		CollisionData.Reset();
@@ -125,10 +160,13 @@ namespace UE::Chaos::ClothAsset
 	void FCollisionSourcesProxy::ExtractCollisionData()
 	{
 		bool bCollisionDataChanged = false;
-		if (Version != CollisionSources.Version)
+
+		CollisionSources.ResolveWeakComponentPtrs();
+
+		if (!CollisionSources.IsCollisionDataUpToDate(Version))
 		{
 			CollisionSources.ExtractNonEnvironmentalCollisionData(NonEnvironmentalCollisionData);
-			Version = CollisionSources.Version;
+			CollisionSources.MarkCollisionDataUpToDate(Version);
 			bCollisionDataChanged = true;
 		}
 
@@ -148,7 +186,6 @@ namespace UE::Chaos::ClothAsset
 	FCollisionSources::FCollisionSource::FCollisionSource(
 		USkinnedMeshComponent* InSourceComponent,
 		const UPhysicsAsset* InSourcePhysicsAsset,
-		const FSimpleDelegate& InOnBoneTransformsFinalizedDelegate,
 		bool bInUseSphylsOnly)
 		: SourceComponent(InSourceComponent)
 		, SourcePhysicsAsset(InSourcePhysicsAsset)
@@ -156,24 +193,15 @@ namespace UE::Chaos::ClothAsset
 	{
 		if (InSourceComponent)
 		{
-			OnBoneTransformsFinalizedHandle = InSourceComponent->RegisterOnBoneTransformsFinalizedDelegate(InOnBoneTransformsFinalizedDelegate);
-		}
-	}
-
-	FCollisionSources::FCollisionSource::~FCollisionSource()
-	{
-		if (USkinnedMeshComponent* const SourceComponentPtr = SourceComponent.Get())
-		{
-			check(OnBoneTransformsFinalizedHandle.IsValid());
-			SourceComponentPtr->UnregisterOnBoneTransformsFinalizedDelegate(OnBoneTransformsFinalizedHandle);
+			BoneTransformRevisionNumber = InSourceComponent->GetBoneTransformRevisionNumber();
 		}
 	}
 
 	void FCollisionSources::FCollisionSource::ExtractCollisionData(const USkinnedMeshComponent& InOwnerComponent, FClothCollisionData& CollisionData)
 	{
-		if (USkinnedMeshComponent* const SourceComponentPtr = SourceComponent.Get())
+		if (ResolvedSourceComponent)
 		{
-			const USkinnedAsset* SkinnedAsset = SourceComponentPtr->GetSkinnedAsset();
+			const USkinnedAsset* SkinnedAsset = ResolvedSourceComponent->GetSkinnedAsset();
 
 			// Extract the collision data if not already cached
 			if (CachedSkinnedAsset != SkinnedAsset)
diff --git a/Engine/Plugins/ChaosClothAsset/Source/ChaosClothAssetEngine/Private/ChaosClothAsset/CollisionSources.h b/Engine/Plugins/ChaosClothAsset/Source/ChaosClothAssetEngine/Private/ChaosClothAsset/CollisionSources.h
--- a/Engine/Plugins/ChaosClothAsset/Source/ChaosClothAssetEngine/Private/ChaosClothAsset/CollisionSources.h
+++ b/Engine/Plugins/ChaosClothAsset/Source/ChaosClothAssetEngine/Private/ChaosClothAsset/CollisionSources.h
@@ -39,25 +39,28 @@ namespace UE::Chaos::ClothAsset
 		struct FCollisionSource
 		{
 			const TWeakObjectPtr<USkinnedMeshComponent> SourceComponent;
+			// Resolved copy of SourceComponent, refreshed prior to each update to avoid multiple TWeakObjectPtr resolves
+			const USkinnedMeshComponent* ResolvedSourceComponent = nullptr;
 			const TWeakObjectPtr<const UPhysicsAsset> SourcePhysicsAsset;
 			TWeakObjectPtr<const USkinnedAsset> CachedSkinnedAsset;
 			FClothCollisionData CachedCollisionData;
 			TArray<int32> CachedUsedBoneIndices;
-			FDelegateHandle OnBoneTransformsFinalizedHandle;
+			uint32 BoneTransformRevisionNumber = 0;
 			bool bUseSphylsOnly = false;
 
 			FCollisionSource(
 				USkinnedMeshComponent* InSourceComponent,
 				const UPhysicsAsset* InSourcePhysicsAsset,
-				const FSimpleDelegate& InOnBoneTransformsFinalizedDelegate,
 				bool bInUseSphylsOnly);
-			~FCollisionSource();
 
 			void ExtractCollisionData(const USkinnedMeshComponent& InOwnerComponent, FClothCollisionData& CollisionData);
 		};
 
 		void ExtractNonEnvironmentalCollisionData(FClothCollisionData& CollisionData);
 		void ExtractEnvironmentalCollisionData(FClothCollisionData& CollisionData);
+		bool IsCollisionDataUpToDate(int32 InVersion) const;
+		void MarkCollisionDataUpToDate(int32 InVersion);
+		void ResolveWeakComponentPtrs();
 
 		const TWeakObjectPtr<USkinnedMeshComponent> OwnerComponent;
 		TArray<FCollisionSource> CollisionSources;
