Change ClothComponent defaults to enable Distance-based teleport and enable Reset on teleport by default. This matches the behavior of SKM-based cloth.

diff --git a/Engine/Plugins/ChaosClothAsset/Source/ChaosClothAssetEngine/Private/ChaosClothAsset/ClothComponent.cpp b/Engine/Plugins/ChaosClothAsset/Source/ChaosClothAssetEngine/Private/ChaosClothAsset/ClothComponent.cpp
--- a/Engine/Plugins/ChaosClothAsset/Source/ChaosClothAssetEngine/Private/ChaosClothAsset/ClothComponent.cpp
+++ b/Engine/Plugins/ChaosClothAsset/Source/ChaosClothAssetEngine/Private/ChaosClothAsset/ClothComponent.cpp
@@ -44,11 +44,13 @@ UChaosClothComponent::UChaosClothComponent(const FObjectInitializer& ObjectIniti
 	, bEnableSimulation(1)
 	, bSuspendSimulation(0)
 	, bBindToLeaderComponent(0)
+	, bTeleport(0)
+	, bReset(1)
 	, bCollideWithEnvironment(0)
 #if WITH_EDITORONLY_DATA
 	, bSimulateInEditor(0)
 #endif
-	, TeleportDistanceThreshold(0.f)
+	, TeleportDistanceThreshold(300.f)
 	, TeleportRotationThreshold(0.f)
 	, ClothTeleportDistThresholdSquared(UE::ClothingSimulation::TeleportHelpers::ComputeTeleportDistanceThresholdSquared(TeleportDistanceThreshold))
 	, ClothTeleportCosineThresholdInRad(UE::ClothingSimulation::TeleportHelpers::ComputeTeleportCosineRotationThreshold(TeleportRotationThreshold))
