Cloth dataflow nodes to support sim mesh resizing:

* ApplyProxyDeformerNode will apply any existing proxy deformer data to update the Render Mesh given the current Sim Mesh
* Added the option to transfer normals from UpdateClothFromDynamicMesh node
* Added a ProceduralSelectionNode which will procedurallly generate a new selection set. Current options are by selecting all elements in a group, and by converting one selection to another type.

diff --git a/Engine/Plugins/ChaosClothAsset/Source/ChaosClothAsset/Private/ChaosClothAsset/ClothGeometryTools.cpp b/Engine/Plugins/ChaosClothAsset/Source/ChaosClothAsset/Private/ChaosClothAsset/ClothGeometryTools.cpp
--- a/Engine/Plugins/ChaosClothAsset/Source/ChaosClothAsset/Private/ChaosClothAsset/ClothGeometryTools.cpp
+++ b/Engine/Plugins/ChaosClothAsset/Source/ChaosClothAsset/Private/ChaosClothAsset/ClothGeometryTools.cpp
@@ -1548,6 +1548,21 @@ PRAGMA_ENABLE_DEPRECATION_WARNINGS
 		PRAGMA_ENABLE_DEPRECATION_WARNINGS
 	}
 
+	void FClothGeometryTools::SelectAllInGroupType(const TSharedRef<FManagedArrayCollection>& ClothCollection, const FName& SelectionName, const FName& GroupName)
+	{
+		FCollectionClothSelectionFacade SelectionFacade(ClothCollection);
+		SelectionFacade.DefineSchema();
+		if (ClothCollection->HasGroup(GroupName))
+		{
+			const int32 GroupSize = ClothCollection->NumElements(GroupName);
+			TSet<int32>& NewSelectionSet = SelectionFacade.FindOrAddSelectionSet(SelectionName, GroupName);
+			NewSelectionSet.Reserve(GroupSize);
+			for (int32 Index = 0; Index < GroupSize; ++Index)
+			{
+				NewSelectionSet.Add(Index);
+			}
+		}
+	}
 
 	void FClothGeometryTools::TransferWeightMap(
 		const TConstArrayView<FVector3f>& SourcePositions,
@@ -1652,4 +1667,127 @@ PRAGMA_ENABLE_DEPRECATION_WARNINGS
 		}
 		return KinematicVertices;
 	}
+
+	void FClothGeometryTools::ApplyProxyDeformer(const TSharedRef<FManagedArrayCollection>& ClothCollection, bool bIgnoreSkinningBlend)
+	{
+		FCollectionClothFacade ClothFacade(ClothCollection);
+		if (ClothFacade.IsValid(EClothCollectionOptionalSchemas::RenderDeformer))
+		{
+			return ApplyProxyDeformer(ClothFacade, bIgnoreSkinningBlend, ClothFacade.GetSimPosition3D(), ClothFacade.GetSimNormal());
+		}
+	}
+
+	void FClothGeometryTools::ApplyProxyDeformer(FCollectionClothFacade& ClothFacade, bool bIgnoreSkinningBlend, const TConstArrayView<FVector3f>& SimPositions, const TConstArrayView<FVector3f>& SimNormals)
+	{
+		if (ClothFacade.IsValid(EClothCollectionOptionalSchemas::RenderDeformer))
+		{
+			// This follows GpuSkinVertexFactor.ush
+			for (int32 SectionIndex = 0; SectionIndex < ClothFacade.GetNumRenderPatterns(); ++SectionIndex)
+			{
+				FCollectionClothRenderPatternFacade RenderPatternFacade = ClothFacade.GetRenderPattern(SectionIndex);
+				const int32 RenderDeformerNumInfluences = RenderPatternFacade.GetRenderDeformerNumInfluences();
+				if (RenderDeformerNumInfluences > 0)
+				{
+					TArrayView<FVector3f> RenderPositions = RenderPatternFacade.GetRenderPosition();
+					TArrayView<FVector3f> RenderNormal = RenderPatternFacade.GetRenderNormal();
+					TArrayView<FVector3f> RenderTangentU = RenderPatternFacade.GetRenderTangentU();
+					TArrayView<FVector3f> RenderTangentV = RenderPatternFacade.GetRenderTangentV();
+
+					TConstArrayView<TArray<FVector4f>> PositionBaryCoordsAndDist = RenderPatternFacade.GetRenderDeformerPositionBaryCoordsAndDist();
+					TConstArrayView<TArray<FVector4f>> NormalBaryCoordsAndDist = RenderPatternFacade.GetRenderDeformerNormalBaryCoordsAndDist();
+					TConstArrayView<TArray<FVector4f>> TangentBaryCoordsAndDist = RenderPatternFacade.GetRenderDeformerTangentBaryCoordsAndDist();
+					TConstArrayView<TArray<FIntVector3>> SimIndices3D = RenderPatternFacade.GetRenderDeformerSimIndices3D();
+					TConstArrayView<TArray<float>> DeformerWeight = RenderPatternFacade.GetRenderDeformerWeight();
+					TConstArrayView<float> SkinningBlend = RenderPatternFacade.GetRenderDeformerSkinningBlend();
+
+					for (int32 Index = 0; Index < RenderPatternFacade.GetNumRenderVertices(); ++Index)
+					{
+						const float SkinningBlendValue = bIgnoreSkinningBlend ? 0.f : SkinningBlend[Index];
+						if (SkinningBlendValue < 1)
+						{
+							FVector3f AveragedSimPosition(0.f);
+							FVector3f NormalPosition(0.f);
+							FVector3f TangentPosition(0.f);
+
+							float SumWeights = 0.f;
+							float SimulWeight = 0.f;
+							for (int32 Influence = 0; Influence < RenderDeformerNumInfluences; ++Influence)
+							{
+								const FVector3f& SimPosA = SimPositions[SimIndices3D[Index][Influence].X];
+								const FVector3f& SimPosB = SimPositions[SimIndices3D[Index][Influence].Y];
+								const FVector3f& SimPosC = SimPositions[SimIndices3D[Index][Influence].Z];
+
+								const FVector3f& SimNormalA = SimNormals[SimIndices3D[Index][Influence].X];
+								const FVector3f& SimNormalB = SimNormals[SimIndices3D[Index][Influence].Y];
+								const FVector3f& SimNormalC = SimNormals[SimIndices3D[Index][Influence].Z];
+
+								SimulWeight += 1.f - SkinningBlendValue;
+
+								float Weight = 1.f;
+								if (RenderDeformerNumInfluences > 1)
+								{
+									Weight = DeformerWeight[Index][Influence];
+									SumWeights += Weight;
+								}
+								else
+								{
+									Weight = 1.f;
+									SumWeights = 1.f;
+								}
+
+								auto InterpolatePosition = [Weight, &SimPosA, &SimPosB, &SimPosC, &SimNormalA, &SimNormalB, &SimNormalC](const FVector4f& BaryCoord)->FVector3f
+									{
+										// Note: coordinates are calculated with inverted normals, so subtract them here.
+										return Weight * (
+											BaryCoord.X * (SimPosA - SimNormalA * BaryCoord.W) +
+											BaryCoord.Y * (SimPosB - SimNormalB * BaryCoord.W) +
+											BaryCoord.Z * (SimPosC - SimNormalC * BaryCoord.W));
+									};
+
+								const FVector4f& BaryCoordPos = PositionBaryCoordsAndDist[Index][Influence];
+								AveragedSimPosition += InterpolatePosition(BaryCoordPos);
+
+								const FVector4f& BaryCoordNormal = NormalBaryCoordsAndDist[Index][Influence];
+								NormalPosition += InterpolatePosition(BaryCoordNormal);
+
+								const FVector4f& BaryCoordTangent = TangentBaryCoordsAndDist[Index][Influence];
+								TangentPosition += InterpolatePosition(BaryCoordTangent);
+							}
+
+
+							FVector3f Normal(0.f);
+							FVector3f TangentU(0.f);
+							FVector3f TangentV(0.f);
+							if (SumWeights > UE_KINDA_SMALL_NUMBER)
+							{
+								float InvWeight = 1.f / SumWeights;
+								AveragedSimPosition *= InvWeight;
+								TangentPosition *= InvWeight;
+								NormalPosition *= InvWeight;
+
+								TangentU = (TangentPosition - AveragedSimPosition).GetSafeNormal();
+								Normal = (NormalPosition - AveragedSimPosition).GetSafeNormal();
+
+								TangentV = FVector3f::CrossProduct(TangentU, Normal).GetSafeNormal();
+							}
+							else
+							{
+								SimulWeight = 0.f; // Fallback to skinned pos
+							}
+
+							if (RenderDeformerNumInfluences > 1)
+							{
+								SimulWeight /= RenderDeformerNumInfluences;
+							}
+
+							RenderPositions[Index] = FMath::Lerp(RenderPositions[Index], AveragedSimPosition, SimulWeight);
+							RenderNormal[Index] = FMath::Lerp(RenderNormal[Index], Normal, SimulWeight).GetSafeNormal();
+							RenderTangentU[Index] = FMath::Lerp(RenderTangentU[Index], TangentU, SimulWeight).GetSafeNormal();
+							RenderTangentV[Index] = FMath::Lerp(RenderTangentV[Index], TangentV, SimulWeight).GetSafeNormal();
+						}
+					}
+				}
+			}
+		}
+	}
 }  // End namespace UE::Chaos::ClothAsset
diff --git a/Engine/Plugins/ChaosClothAsset/Source/ChaosClothAsset/Public/ChaosClothAsset/ClothGeometryTools.h b/Engine/Plugins/ChaosClothAsset/Source/ChaosClothAsset/Public/ChaosClothAsset/ClothGeometryTools.h
--- a/Engine/Plugins/ChaosClothAsset/Source/ChaosClothAsset/Public/ChaosClothAsset/ClothGeometryTools.h
+++ b/Engine/Plugins/ChaosClothAsset/Source/ChaosClothAsset/Public/ChaosClothAsset/ClothGeometryTools.h
@@ -9,6 +9,8 @@ namespace UE::Geometry { class FDynamicMesh3; }
 
 namespace UE::Chaos::ClothAsset
 {
+	class FCollectionClothFacade;
+
 	/**
 	 * Geometry tools operating on cloth collections.
 	 */
@@ -145,6 +147,11 @@ namespace UE::Chaos::ClothAsset
 		 */
 		static bool ConvertSelectionToNewGroupType(const TSharedRef<const FManagedArrayCollection>& ClothCollection, const FName& SelectionName, const FName& GroupName, TSet<int32>& OutSelectionSet);
 
+		/**
+		* Create a selection set that selects all members of a group. 
+		*/
+		static void SelectAllInGroupType(const TSharedRef<FManagedArrayCollection>& ClothCollection, const FName& SelectionName, const FName& GroupName);
+
 		/**
 		 * Transfer a vertex weight map from a source to target mesh
 		 */
@@ -161,5 +168,11 @@ namespace UE::Chaos::ClothAsset
 		 * Generate KinematicVertices3D set from the given MaxDistance weight map, MaxDistance values, and any additional kinematic vertices.
 		 */
 		static TSet<int32> GenerateKinematicVertices3D(const TSharedRef<FManagedArrayCollection>& ClothCollection, const FName& MaxDistanceMap, const FVector2f& MaxDistanceValues, const FName& InputKinematicVertices, float KinematicDistanceThreshold = 0.1f);
+
+		/**
+		 * Update the render mesh by applying the proxy deformer.
+		 */
+		static void ApplyProxyDeformer(const TSharedRef<FManagedArrayCollection>& ClothCollection, bool bIgnoreSkinningBlend);
+		static void ApplyProxyDeformer(FCollectionClothFacade& ClothFacade, bool bIgnoreSkinningBlend, const TConstArrayView<FVector3f>& SimPositions, const TConstArrayView<FVector3f>& SimNormals);
 	};
 }  // End namespace UE::Chaos::ClothAsset
\ No newline at end of file
diff --git a/Engine/Plugins/ChaosClothAssetEditor/Source/ChaosClothAssetDataflowNodes/Private/ChaosClothAsset/ApplyProxyDeformerNode.cpp b/Engine/Plugins/ChaosClothAssetEditor/Source/ChaosClothAssetDataflowNodes/Private/ChaosClothAsset/ApplyProxyDeformerNode.cpp
new file mode 100644
--- /dev/null
+++ b/Engine/Plugins/ChaosClothAssetEditor/Source/ChaosClothAssetDataflowNodes/Private/ChaosClothAsset/ApplyProxyDeformerNode.cpp
@@ -0,0 +1,33 @@
+// Copyright Epic Games, Inc. All Rights Reserved.
+
+#include "ChaosClothAsset/ApplyProxyDeformerNode.h"
+#include "ChaosClothAsset/ClothGeometryTools.h"
+#include "Dataflow/DataflowInputOutput.h"
+
+#include UE_INLINE_GENERATED_CPP_BY_NAME(ApplyProxyDeformerNode)
+
+#define LOCTEXT_NAMESPACE "ChaosClothAssetApplyProxyDeformerNode"
+
+FChaosClothAssetApplyProxyDeformerNode::FChaosClothAssetApplyProxyDeformerNode(const UE::Dataflow::FNodeParameters& InParam, FGuid InGuid)
+	: FDataflowNode(InParam, InGuid)
+{
+	RegisterInputConnection(&Collection);
+	RegisterOutputConnection(&Collection, &Collection);
+}
+
+void FChaosClothAssetApplyProxyDeformerNode::Evaluate(UE::Dataflow::FContext& Context, const FDataflowOutput* Out) const
+{
+	if (Out->IsA<FManagedArrayCollection>(&Collection))
+	{
+		using namespace UE::Chaos::ClothAsset;
+		// Evaluate in collection
+		FManagedArrayCollection InCollection = GetValue<FManagedArrayCollection>(Context, &Collection);
+		const TSharedRef<FManagedArrayCollection> ClothCollection = MakeShared<FManagedArrayCollection>(MoveTemp(InCollection));
+
+		FClothGeometryTools::ApplyProxyDeformer(ClothCollection, bIgnoreSkinningBlendWeights);
+
+		SetValue(Context, MoveTemp(*ClothCollection), &Collection);
+	}
+}
+
+#undef LOCTEXT_NAMESPACE
diff --git a/Engine/Plugins/ChaosClothAssetEditor/Source/ChaosClothAssetDataflowNodes/Private/ChaosClothAsset/ChaosClothAssetDataflowNodesModule.cpp b/Engine/Plugins/ChaosClothAssetEditor/Source/ChaosClothAssetDataflowNodes/Private/ChaosClothAsset/ChaosClothAssetDataflowNodesModule.cpp
--- a/Engine/Plugins/ChaosClothAssetEditor/Source/ChaosClothAssetDataflowNodes/Private/ChaosClothAsset/ChaosClothAssetDataflowNodesModule.cpp
+++ b/Engine/Plugins/ChaosClothAssetEditor/Source/ChaosClothAssetDataflowNodes/Private/ChaosClothAsset/ChaosClothAssetDataflowNodesModule.cpp
@@ -2,11 +2,13 @@
 
 #include "ChaosClothAsset/AddStitchNode.h"
 #include "ChaosClothAsset/AddWeightMapNode.h"
+#include "ChaosClothAsset/ApplyProxyDeformerNode.h"
 #include "ChaosClothAsset/AttributeNode.h"
 #include "ChaosClothAsset/BindToRootBoneNode.h"
 #include "ChaosClothAsset/BlendVerticesNode.h"
 #include "ChaosClothAsset/ClothCollectionGroup.h"
 #include "ChaosClothAsset/ClothCollectionToDynamicMeshNode.h"
+#include "ChaosClothAsset/ClothCollectionQueryNode.h"
 #include "ChaosClothAsset/ColorScheme.h"
 #include "ChaosClothAsset/CopySimulationToRenderMeshNode.h"
 #include "ChaosClothAsset/DatasmithImportNode.h"
@@ -16,6 +18,7 @@
 #include "ChaosClothAsset/ImportSimulationCacheNode.h"
 #include "ChaosClothAsset/MakeClothAssetNode.h"
 #include "ChaosClothAsset/MergeClothCollectionsNode.h"
+#include "ChaosClothAsset/ProceduralSelectionNode.h"
 #include "ChaosClothAsset/ProxyDeformerNode.h"
 #include "ChaosClothAsset/RemeshNode.h"
 #include "ChaosClothAsset/ReverseNormalsNode.h"
@@ -90,9 +93,11 @@ namespace UE::Chaos::ClothAsset
 			DATAFLOW_NODE_REGISTER_CREATION_FACTORY(FChaosClothAssetAddStitchNode);
 			DATAFLOW_NODE_REGISTER_CREATION_FACTORY(FChaosClothAssetWeightMapNode);
 			DATAFLOW_NODE_REGISTER_CREATION_FACTORY(FChaosClothAssetAttributeNode_v2);
+			DATAFLOW_NODE_REGISTER_CREATION_FACTORY(FChaosClothAssetApplyProxyDeformerNode);
 			DATAFLOW_NODE_REGISTER_CREATION_FACTORY(FChaosClothAssetBindToRootBoneNode);
 			DATAFLOW_NODE_REGISTER_CREATION_FACTORY(FChaosClothAssetBlendVerticesNode);
 			DATAFLOW_NODE_REGISTER_CREATION_FACTORY(FChaosClothAssetCollectionToDynamicMeshNode);
+			DATAFLOW_NODE_REGISTER_CREATION_FACTORY(FChaosClothAssetCollectionQueryNode);
 			DATAFLOW_NODE_REGISTER_CREATION_FACTORY(FChaosClothAssetCopySimulationToRenderMeshNode);
 			DATAFLOW_NODE_REGISTER_CREATION_FACTORY(FChaosClothAssetDatasmithImportNode);
 			DATAFLOW_NODE_REGISTER_CREATION_FACTORY(FChaosClothAssetDeleteElementNode);
@@ -101,6 +106,7 @@ namespace UE::Chaos::ClothAsset
 			DATAFLOW_NODE_REGISTER_CREATION_FACTORY(FChaosClothAssetImportSimulationCacheNode);
 			DATAFLOW_NODE_REGISTER_CREATION_FACTORY(FChaosClothAssetMakeClothAssetNode);
 			DATAFLOW_NODE_REGISTER_CREATION_FACTORY(FChaosClothAssetMergeClothCollectionsNode_v2);
+			DATAFLOW_NODE_REGISTER_CREATION_FACTORY(FChaosClothAssetProceduralSelectionNode);
 			DATAFLOW_NODE_REGISTER_CREATION_FACTORY(FChaosClothAssetProxyDeformerNode_v2);
 			DATAFLOW_NODE_REGISTER_CREATION_FACTORY(FChaosClothAssetRemeshNode_v2);
 			DATAFLOW_NODE_REGISTER_CREATION_FACTORY(FChaosClothAssetReverseNormalsNode);
diff --git a/Engine/Plugins/ChaosClothAssetEditor/Source/ChaosClothAssetDataflowNodes/Private/ChaosClothAsset/ClothCollectionQueryNode.cpp b/Engine/Plugins/ChaosClothAssetEditor/Source/ChaosClothAssetDataflowNodes/Private/ChaosClothAsset/ClothCollectionQueryNode.cpp
new file mode 100644
--- /dev/null
+++ b/Engine/Plugins/ChaosClothAssetEditor/Source/ChaosClothAssetDataflowNodes/Private/ChaosClothAsset/ClothCollectionQueryNode.cpp
@@ -0,0 +1,46 @@
+// Copyright Epic Games, Inc. All Rights Reserved.
+
+#include "ChaosClothAsset/ClothCollectionQueryNode.h"
+#include "ChaosClothAsset/CollectionClothFacade.h"
+#include "ChaosClothAsset/ClothGeometryTools.h"
+
+#include UE_INLINE_GENERATED_CPP_BY_NAME(ClothCollectionQueryNode)
+
+#define LOCTEXT_NAMESPACE "ChaosClothAssetCollectionQueryNode"
+
+FChaosClothAssetCollectionQueryNode::FChaosClothAssetCollectionQueryNode(const UE::Dataflow::FNodeParameters& InParam, FGuid InGuid)
+	: FDataflowNode(InParam, InGuid)
+{
+	RegisterInputConnection(&Collection);
+	RegisterOutputConnection(&Collection, &Collection);
+	RegisterOutputConnection(&bIsClothCollection);
+	RegisterOutputConnection(&bHasClothSimMesh);
+	RegisterOutputConnection(&bHasClothRenderMesh);
+	RegisterOutputConnection(&bHasClothProxyDeformer);
+}
+
+void FChaosClothAssetCollectionQueryNode::Evaluate(UE::Dataflow::FContext& Context, const FDataflowOutput* Out) const
+{
+	using namespace UE::Chaos::ClothAsset;
+	if (Out->IsA<FManagedArrayCollection>(&Collection))
+	{
+		SafeForwardInput(Context, &Collection, &Collection);
+	}
+	else if (Out->IsA<bool>(&bIsClothCollection)
+		|| Out->IsA<bool>(&bHasClothSimMesh)
+		|| Out->IsA<bool>(&bHasClothRenderMesh)
+		|| Out->IsA<bool>(&bHasClothProxyDeformer))
+	{
+		// Setting all of these outputs once since it involves making a temporary copy of the whole collection.
+		FManagedArrayCollection InCollection = GetValue<FManagedArrayCollection>(Context, &Collection);
+		const TSharedRef<FManagedArrayCollection> ClothCollection = MakeShared<FManagedArrayCollection>(MoveTemp(InCollection));
+		FCollectionClothConstFacade ClothFacade(ClothCollection);
+
+		SetValue(Context, ClothFacade.IsValid(), &bIsClothCollection);
+		SetValue(Context, FClothGeometryTools::HasSimMesh(ClothCollection), &bHasClothSimMesh);
+		SetValue(Context, FClothGeometryTools::HasRenderMesh(ClothCollection), &bHasClothRenderMesh);
+		SetValue(Context, ClothFacade.IsValid(EClothCollectionOptionalSchemas::RenderDeformer), &bHasClothProxyDeformer);
+	}
+}
+
+#undef LOCTEXT_NAMESPACE
diff --git a/Engine/Plugins/ChaosClothAssetEditor/Source/ChaosClothAssetDataflowNodes/Private/ChaosClothAsset/ClothCollectionToDynamicMeshNode.cpp b/Engine/Plugins/ChaosClothAssetEditor/Source/ChaosClothAssetDataflowNodes/Private/ChaosClothAsset/ClothCollectionToDynamicMeshNode.cpp
--- a/Engine/Plugins/ChaosClothAssetEditor/Source/ChaosClothAssetDataflowNodes/Private/ChaosClothAsset/ClothCollectionToDynamicMeshNode.cpp
+++ b/Engine/Plugins/ChaosClothAssetEditor/Source/ChaosClothAssetDataflowNodes/Private/ChaosClothAsset/ClothCollectionToDynamicMeshNode.cpp
@@ -101,6 +101,32 @@ void FChaosClothAssetUpdateClothFromDynamicMeshNode::Evaluate(UE::Dataflow::FCon
 				const UE::Geometry::FDynamicMesh3& DynMesh = InMesh->GetMeshRef();
 				FClothPatternToDynamicMeshMappingSupport ClothMapping(DynMesh);
 
+				auto CopyNormals = [&DynMesh, &ClothMapping](const UE::Geometry::FDynamicMeshNormalOverlay* const NormalOverlay, const TConstArrayView<FIntVector>& TriangleIndices, TArrayView<FVector3f> Normals)
+					{
+						for (int32 VertexID : DynMesh.VertexIndicesItr())
+						{
+							const int32 ClothVertexID = ClothMapping.GetOriginalVertexID(VertexID);
+							if (Normals.IsValidIndex(ClothVertexID))
+							{
+								// Get the UV values corresponding to VertexID in the dynamic mesh
+								NormalOverlay->EnumerateVertexElements(VertexID,
+									[&Normals, &ClothMapping, &TriangleIndices, ClothVertexID](int32 TriangleID, int32 ElementID, const FVector3f& NormalValue)->bool
+									{
+										const int32 ClothMeshTriID = ClothMapping.GetOriginalTriangleID(TriangleID);
+										const FIntVector& ClothTri = TriangleIndices[ClothVertexID];
+										for (int32 LocalIndex = 0; LocalIndex < 3; ++LocalIndex)
+										{
+											if (ClothTri[LocalIndex] == ClothVertexID)
+											{
+												Normals[ClothTri[LocalIndex]] = NormalValue;
+											}
+										}
+										return true;
+									});
+							}
+						}
+					};
+
 				if (bCopyToRenderPositions)
 				{
 					TArrayView<FVector3f> RenderPositions = ClothFacade.GetRenderPosition();
@@ -113,6 +139,26 @@ void FChaosClothAssetUpdateClothFromDynamicMeshNode::Evaluate(UE::Dataflow::FCon
 						}
 					}
 				}
+				if (bCopyToRendeNormalsAndTangents)
+				{
+					if (const UE::Geometry::FDynamicMeshAttributeSet* const AttributeSet = DynMesh.Attributes())
+					{
+						TConstArrayView<FIntVector> RenderTriangles = ClothFacade.GetRenderIndices();
+
+						if (const UE::Geometry::FDynamicMeshNormalOverlay* const NormalOverlay = AttributeSet->PrimaryNormals())
+						{
+							CopyNormals(NormalOverlay, RenderTriangles, ClothFacade.GetRenderNormal());
+						}
+						if (const UE::Geometry::FDynamicMeshNormalOverlay* const TangentOverlay = AttributeSet->PrimaryTangents())
+						{
+							CopyNormals(TangentOverlay, RenderTriangles, ClothFacade.GetRenderTangentU());
+						}
+						if (const UE::Geometry::FDynamicMeshNormalOverlay* const TangentOverlay = AttributeSet->PrimaryBiTangents())
+						{
+							CopyNormals(TangentOverlay, RenderTriangles, ClothFacade.GetRenderTangentV());
+						}
+					}
+				}
 				if (bCopyUVsToRenderUVs)
 				{
 					if (const UE::Geometry::FDynamicMeshAttributeSet* const AttributeSet = DynMesh.Attributes())
@@ -158,6 +204,16 @@ void FChaosClothAssetUpdateClothFromDynamicMeshNode::Evaluate(UE::Dataflow::FCon
 						}
 					}
 				}
+				if (bCopyToSimNormals)
+				{
+					if (const UE::Geometry::FDynamicMeshAttributeSet* const AttributeSet = DynMesh.Attributes())
+					{
+						if (const UE::Geometry::FDynamicMeshNormalOverlay* const NormalOverlay = AttributeSet->PrimaryNormals())
+						{
+							CopyNormals(NormalOverlay, ClothFacade.GetSimIndices3D(), ClothFacade.GetSimNormal());
+						}
+					}
+				}
 				if (bCopyUVsToSim2DPositions)
 				{
 					if (const UE::Geometry::FDynamicMeshAttributeSet* const AttributeSet = DynMesh.Attributes())
diff --git a/Engine/Plugins/ChaosClothAssetEditor/Source/ChaosClothAssetDataflowNodes/Private/ChaosClothAsset/ImportSimulationCacheNode.cpp b/Engine/Plugins/ChaosClothAssetEditor/Source/ChaosClothAssetDataflowNodes/Private/ChaosClothAsset/ImportSimulationCacheNode.cpp
--- a/Engine/Plugins/ChaosClothAssetEditor/Source/ChaosClothAssetDataflowNodes/Private/ChaosClothAsset/ImportSimulationCacheNode.cpp
+++ b/Engine/Plugins/ChaosClothAssetEditor/Source/ChaosClothAssetDataflowNodes/Private/ChaosClothAsset/ImportSimulationCacheNode.cpp
@@ -2,6 +2,7 @@
 
 #include "ChaosClothAsset/ImportSimulationCacheNode.h"
 #include "ChaosClothAsset/ClothDataflowTools.h"
+#include "ChaosClothAsset/ClothGeometryTools.h"
 #include "ChaosClothAsset/CollectionClothFacade.h"
 #include "Chaos/CacheCollection.h"
 #include "Chaos/ChaosCache.h"
@@ -57,116 +58,6 @@ namespace UE::Chaos::ClothAsset::Private
 			Normal = Normal.GetSafeNormal(UE_SMALL_NUMBER, FVector3f::XAxisVector);
 		}
 	}
-
-	// This follows GpuSkinVertexFactor.ush
-	static void ApplyRenderDeformer(FCollectionClothFacade& ClothFacade, const TConstArrayView<FVector3f>& SimPositions, const TConstArrayView<FVector3f>& SimNormals)
-	{
-		for (int32 SectionIndex = 0; SectionIndex < ClothFacade.GetNumRenderPatterns(); ++SectionIndex)
-		{
-			FCollectionClothRenderPatternFacade RenderPatternFacade = ClothFacade.GetRenderPattern(SectionIndex);
-			const int32 RenderDeformerNumInfluences = RenderPatternFacade.GetRenderDeformerNumInfluences();
-			if (RenderDeformerNumInfluences > 0)
-			{
-				TArrayView<FVector3f> RenderPositions = RenderPatternFacade.GetRenderPosition();
-				TArrayView<FVector3f> RenderNormal = RenderPatternFacade.GetRenderNormal();
-				TArrayView<FVector3f> RenderTangentU = RenderPatternFacade.GetRenderTangentU();
-				TArrayView<FVector3f> RenderTangentV = RenderPatternFacade.GetRenderTangentV();
-
-				TConstArrayView<TArray<FVector4f>> PositionBaryCoordsAndDist = RenderPatternFacade.GetRenderDeformerPositionBaryCoordsAndDist();
-				TConstArrayView<TArray<FVector4f>> NormalBaryCoordsAndDist = RenderPatternFacade.GetRenderDeformerNormalBaryCoordsAndDist();
-				TConstArrayView<TArray<FVector4f>> TangentBaryCoordsAndDist = RenderPatternFacade.GetRenderDeformerTangentBaryCoordsAndDist();
-				TConstArrayView<TArray<FIntVector3>> SimIndices3D = RenderPatternFacade.GetRenderDeformerSimIndices3D();
-				TConstArrayView<TArray<float>> DeformerWeight = RenderPatternFacade.GetRenderDeformerWeight();
-				TConstArrayView<float> SkinningBlend = RenderPatternFacade.GetRenderDeformerSkinningBlend();
-
-				for (int32 Index = 0; Index < RenderPatternFacade.GetNumRenderVertices(); ++Index)
-				{
-					if (SkinningBlend[Index] < 1)
-					{
-						FVector3f AveragedSimPosition(0.f);
-						FVector3f NormalPosition(0.f);
-						FVector3f TangentPosition(0.f);
-
-						float SumWeights = 0.f;
-						float SimulWeight = 0.f;
-						for (int32 Influence = 0; Influence < RenderDeformerNumInfluences; ++Influence)
-						{
-							const FVector3f& SimPosA = SimPositions[SimIndices3D[Index][Influence].X];
-							const FVector3f& SimPosB = SimPositions[SimIndices3D[Index][Influence].Y];
-							const FVector3f& SimPosC = SimPositions[SimIndices3D[Index][Influence].Z];
-
-							const FVector3f& SimNormalA = SimNormals[SimIndices3D[Index][Influence].X];
-							const FVector3f& SimNormalB = SimNormals[SimIndices3D[Index][Influence].Y];
-							const FVector3f& SimNormalC = SimNormals[SimIndices3D[Index][Influence].Z];
-
-							SimulWeight += 1.f - SkinningBlend[Index];
-
-							float Weight = 1.f;
-							if (RenderDeformerNumInfluences > 1)
-							{
-								Weight = DeformerWeight[Index][Influence];
-								SumWeights += Weight;
-							}
-							else
-							{
-								Weight = 1.f;
-								SumWeights = 1.f;
-							}
-
-							auto InterpolatePosition = [Weight, &SimPosA, &SimPosB, &SimPosC, &SimNormalA, &SimNormalB, &SimNormalC](const FVector4f& BaryCoord)->FVector3f
-								{
-									// Note: coordinates are calculated with inverted normals, so subtract them here.
-									return Weight * (
-										BaryCoord.X * (SimPosA - SimNormalA * BaryCoord.W) +
-										BaryCoord.Y * (SimPosB - SimNormalB * BaryCoord.W) +
-										BaryCoord.Z * (SimPosC - SimNormalC * BaryCoord.W));
-								};
-
-							const FVector4f& BaryCoordPos = PositionBaryCoordsAndDist[Index][Influence];
-							AveragedSimPosition += InterpolatePosition(BaryCoordPos); 
-
-							const FVector4f& BaryCoordNormal = NormalBaryCoordsAndDist[Index][Influence];
-							NormalPosition += InterpolatePosition(BaryCoordNormal);
-
-							const FVector4f& BaryCoordTangent = TangentBaryCoordsAndDist[Index][Influence];
-							TangentPosition += InterpolatePosition(BaryCoordTangent);
-						}
-
-
-						FVector3f Normal(0.f);
-						FVector3f TangentU(0.f);
-						FVector3f TangentV(0.f);
-						if (SumWeights > UE_KINDA_SMALL_NUMBER)
-						{
-							float InvWeight = 1.f / SumWeights;
-							AveragedSimPosition *= InvWeight;
-							TangentPosition *= InvWeight;
-							NormalPosition *= InvWeight;
-
-							TangentU = (TangentPosition - AveragedSimPosition).GetSafeNormal();
-							Normal = (NormalPosition - AveragedSimPosition).GetSafeNormal();
-
-							TangentV = FVector3f::CrossProduct(TangentU, Normal).GetSafeNormal();
-						}
-						else
-						{
-							SimulWeight = 0.f; // Fallback to skinned pos
-						}
-
-						if (RenderDeformerNumInfluences > 1)
-						{
-							SimulWeight /= RenderDeformerNumInfluences;
-						}
-
-						RenderPositions[Index] = FMath::Lerp(RenderPositions[Index], AveragedSimPosition, SimulWeight);
-						RenderNormal[Index] = FMath::Lerp(RenderNormal[Index], Normal, SimulWeight).GetSafeNormal();
-						RenderTangentU[Index] = FMath::Lerp(RenderTangentU[Index], TangentU, SimulWeight).GetSafeNormal();
-						RenderTangentV[Index] = FMath::Lerp(RenderTangentV[Index], TangentV, SimulWeight).GetSafeNormal();
-					}
-				}
-			}
-		}
-	}
 }
 
 FChaosClothAssetImportSimulationCacheNode::FChaosClothAssetImportSimulationCacheNode(const UE::Dataflow::FNodeParameters& InParam, FGuid InGuid)
@@ -257,7 +148,8 @@ void FChaosClothAssetImportSimulationCacheNode::Evaluate(UE::Dataflow::FContext&
 						Private::RecalculateSimulationNormals(SimPositions, ClothFacade, SimNormals);
 					}
 
-					Private::ApplyRenderDeformer(ClothFacade, SimPositions, SimNormals);
+					constexpr bool bIgnoreSkinningBlend = false;
+					FClothGeometryTools::ApplyProxyDeformer(ClothFacade, bIgnoreSkinningBlend, SimPositions, SimNormals);
 				}
 			}
 
diff --git a/Engine/Plugins/ChaosClothAssetEditor/Source/ChaosClothAssetDataflowNodes/Private/ChaosClothAsset/ProceduralSelectionNode.cpp b/Engine/Plugins/ChaosClothAssetEditor/Source/ChaosClothAssetDataflowNodes/Private/ChaosClothAsset/ProceduralSelectionNode.cpp
new file mode 100644
--- /dev/null
+++ b/Engine/Plugins/ChaosClothAssetEditor/Source/ChaosClothAssetDataflowNodes/Private/ChaosClothAsset/ProceduralSelectionNode.cpp
@@ -0,0 +1,91 @@
+// Copyright Epic Games, Inc. All Rights Reserved.
+
+#include "ChaosClothAsset/ProceduralSelectionNode.h"
+#include "ChaosClothAsset/ClothCollectionGroup.h"
+#include "ChaosClothAsset/ClothDataflowTools.h"
+#include "ChaosClothAsset/ClothGeometryTools.h"
+#include "ChaosClothAsset/CollectionClothFacade.h"
+#include "ChaosClothAsset/CollectionClothSelectionFacade.h"
+#include "Dataflow/DataflowInputOutput.h"
+
+#include UE_INLINE_GENERATED_CPP_BY_NAME(ProceduralSelectionNode)
+#define LOCTEXT_NAMESPACE "FChaosClothAssetProceduralSelectionNode"
+
+
+
+FChaosClothAssetProceduralSelectionNode::FChaosClothAssetProceduralSelectionNode(const UE::Dataflow::FNodeParameters& InParam, FGuid InGuid)
+	: FDataflowNode(InParam, InGuid)
+{
+	RegisterInputConnection(&Collection);
+	RegisterOutputConnection(&Collection, &Collection);
+	RegisterOutputConnection(&OutputName);
+	RegisterInputConnection(&ConversionInputName.StringValue, GET_MEMBER_NAME_CHECKED(FChaosClothAssetConnectableIStringValue, StringValue))
+		.SetCanHidePin(true)
+		.SetPinIsHidden(true);
+}
+
+void FChaosClothAssetProceduralSelectionNode::Evaluate(UE::Dataflow::FContext& Context, const FDataflowOutput* Out) const
+{
+	using namespace UE::Chaos::ClothAsset;
+
+	if (Out->IsA<FManagedArrayCollection>(&Collection))
+	{
+		const FName SelectionName(*OutputName);
+		if (SelectionName == NAME_None || Group.Name.IsEmpty())
+		{
+			SafeForwardInput(Context, &Collection, &Collection);
+			return;
+		}
+
+		const FName SelectionGroupName(*Group.Name);
+
+		FManagedArrayCollection InCollection = GetValue<FManagedArrayCollection>(Context, &Collection);
+		const TSharedRef<FManagedArrayCollection> ClothCollection = MakeShared<FManagedArrayCollection>(MoveTemp(InCollection));
+		FCollectionClothConstFacade ClothFacade(ClothCollection);
+
+		if (ClothFacade.IsValid())
+		{
+			FCollectionClothSelectionFacade SelectionFacade(ClothCollection);
+			SelectionFacade.DefineSchema();
+
+			switch (SelectionType)
+			{
+			case EChaosClothAssetProceduralSelectionType::SelectAll:
+				FClothGeometryTools::SelectAllInGroupType(ClothCollection, SelectionName, SelectionGroupName);
+				SetValue(Context, MoveTemp(*ClothCollection), &Collection);
+				return;
+			case EChaosClothAssetProceduralSelectionType::Conversion:
+			{
+				const FName InConversionName(*GetValue(Context, &ConversionInputName.StringValue, ConversionInputName.StringValue));
+				if (InConversionName != NAME_None)
+				{
+					TSet<int32> ConvertedSet;
+					if (FClothGeometryTools::ConvertSelectionToNewGroupType(ClothCollection, SelectionName, SelectionGroupName, ConvertedSet))
+					{
+						SelectionFacade.FindOrAddSelectionSet(InConversionName, SelectionGroupName) = ConvertedSet;
+						SetValue(Context, MoveTemp(*ClothCollection), &Collection);
+						return;
+					}
+					else
+					{
+						FClothDataflowTools::LogAndToastWarning(*this, LOCTEXT("ConversionFailureHeadline", "Conversion Failure"),
+							FText::Format(LOCTEXT("ConversionFailureDetails", "Failed to convert selection '{0}' to group type '{1}' either because '{0}' does not exist or the conversion type is unsupported."),
+								FText::FromName(InConversionName), FText::FromName(SelectionGroupName)));
+					}
+				}
+			}
+			break;
+			default:
+				checkNoEntry();
+			}
+		}
+		SafeForwardInput(Context, &Collection, &Collection);
+	}
+	else if (Out->IsA<FString>(&OutputName))
+	{
+		SetValue(Context, OutputName, &OutputName);
+	}
+}
+
+
+#undef LOCTEXT_NAMESPACE
diff --git a/Engine/Plugins/ChaosClothAssetEditor/Source/ChaosClothAssetDataflowNodes/Public/ChaosClothAsset/ApplyProxyDeformerNode.h b/Engine/Plugins/ChaosClothAssetEditor/Source/ChaosClothAssetDataflowNodes/Public/ChaosClothAsset/ApplyProxyDeformerNode.h
new file mode 100644
--- /dev/null
+++ b/Engine/Plugins/ChaosClothAssetEditor/Source/ChaosClothAssetDataflowNodes/Public/ChaosClothAsset/ApplyProxyDeformerNode.h
@@ -0,0 +1,32 @@
+// Copyright Epic Games, Inc. All Rights Reserved.
+
+#pragma once
+
+#include "Dataflow/DataflowNode.h"
+#include "GeometryCollection/ManagedArrayCollection.h"
+#include "ApplyProxyDeformerNode.generated.h"
+
+/** Update the Render Mesh by applying any existing proxy deformer data. 
+  * This node can be used to deform the render mesh after the sim mesh has changed after the proxy deformer data was calculated.
+  * It has no effect if there is no existing proxy deformer data, or if the sim or render mesh have not deformed since that data was calculated.*/
+USTRUCT(Meta = (DataflowCloth))
+struct FChaosClothAssetApplyProxyDeformerNode final : public FDataflowNode 
+{
+	GENERATED_USTRUCT_BODY()
+	DATAFLOW_NODE_DEFINE_INTERNAL(FChaosClothAssetApplyProxyDeformerNode, "ApplyProxyDeformer", "Cloth", "Cloth Apply Proxy Deformer")
+	DATAFLOW_NODE_RENDER_TYPE("SurfaceRender", FName("FClothCollection"), "Collection")
+
+public:
+	FChaosClothAssetApplyProxyDeformerNode(const UE::Dataflow::FNodeParameters& InParam, FGuid InGuid = FGuid::NewGuid());
+
+private:
+	virtual void Evaluate(UE::Dataflow::FContext& Context, const FDataflowOutput* Out) const override;
+
+	/** Input/output collection */
+	UPROPERTY(Meta = (DataflowInput, DataflowOutput, DataflowPassthrough = "Collection"))
+	FManagedArrayCollection Collection;
+
+	/** Ignore Skinning Blend Weights (apply proxy deformer to all points)*/
+	UPROPERTY(EditAnywhere, Category = "Apply Proxy Deformer")
+	bool bIgnoreSkinningBlendWeights = false;
+};
diff --git a/Engine/Plugins/ChaosClothAssetEditor/Source/ChaosClothAssetDataflowNodes/Public/ChaosClothAsset/ClothCollectionQueryNode.h b/Engine/Plugins/ChaosClothAssetEditor/Source/ChaosClothAssetDataflowNodes/Public/ChaosClothAsset/ClothCollectionQueryNode.h
new file mode 100644
--- /dev/null
+++ b/Engine/Plugins/ChaosClothAssetEditor/Source/ChaosClothAssetDataflowNodes/Public/ChaosClothAsset/ClothCollectionQueryNode.h
@@ -0,0 +1,49 @@
+// Copyright Epic Games, Inc. All Rights Reserved.
+
+#pragma once 
+
+#include "Dataflow/DataflowNode.h"
+#include "GeometryCollection/ManagedArrayCollection.h"
+#include "ClothCollectionQueryNode.generated.h"
+
+
+class UDynamicMesh;
+class UMaterialInterface;
+
+/** Query a Managed Array Collection about its Cloth Collection properties.
+ */
+USTRUCT(Meta = (DataflowCloth))
+struct FChaosClothAssetCollectionQueryNode : public FDataflowNode
+{
+	GENERATED_USTRUCT_BODY()
+	DATAFLOW_NODE_DEFINE_INTERNAL(FChaosClothAssetCollectionQueryNode, "ClothCollectionQuery", "Cloth", "Cloth Collection Query")
+	DATAFLOW_NODE_RENDER_TYPE("SurfaceRender", FName("FClothCollection"), "Collection")
+
+public:
+
+	FChaosClothAssetCollectionQueryNode(const UE::Dataflow::FNodeParameters& InParam, FGuid InGuid = FGuid::NewGuid());
+
+	virtual void Evaluate(UE::Dataflow::FContext& Context, const FDataflowOutput* Out) const override;
+
+private:
+
+	/** Input/output collection (Output is always a passthrough)*/
+	UPROPERTY(Meta = (DataflowInput, DataflowOutput, DataflowPassthrough = "Collection"))
+	FManagedArrayCollection Collection;
+
+	/** Is this collection a valid cloth collection */
+	UPROPERTY(Meta = (DataflowOutput))
+	bool bIsClothCollection;
+
+	/** Does this collection have a cloth sim mesh */
+	UPROPERTY(Meta = (DataflowOutput))
+	bool bHasClothSimMesh;
+
+	/** Does this collection have a cloth render mesh */
+	UPROPERTY(Meta = (DataflowOutput))
+	bool bHasClothRenderMesh;
+
+	/** Does this collection have proxy deformer data */
+	UPROPERTY(Meta = (DataflowOutput))
+	bool bHasClothProxyDeformer;
+};
diff --git a/Engine/Plugins/ChaosClothAssetEditor/Source/ChaosClothAssetDataflowNodes/Public/ChaosClothAsset/ClothCollectionToDynamicMeshNode.h b/Engine/Plugins/ChaosClothAssetEditor/Source/ChaosClothAssetDataflowNodes/Public/ChaosClothAsset/ClothCollectionToDynamicMeshNode.h
--- a/Engine/Plugins/ChaosClothAssetEditor/Source/ChaosClothAssetDataflowNodes/Public/ChaosClothAsset/ClothCollectionToDynamicMeshNode.h
+++ b/Engine/Plugins/ChaosClothAssetEditor/Source/ChaosClothAssetDataflowNodes/Public/ChaosClothAsset/ClothCollectionToDynamicMeshNode.h
@@ -77,6 +77,10 @@ private:
 	UPROPERTY(EditAnywhere, Category = "DynamicMesh To Cloth Collection")
 	bool bCopyToRenderPositions = false;
 
+	/** Copy DynamicMesh Vertex Normals and Tangents to Render Normals and Tangents */
+	UPROPERTY(EditAnywhere, Category = "DynamicMesh To Cloth Collection")
+	bool bCopyToRendeNormalsAndTangents = false;
+
 	/** Copy DynamicMesh UVs to Render UVs */
 	UPROPERTY(EditAnywhere, Category = "DynamicMesh To Cloth Collection")
 	bool bCopyUVsToRenderUVs = false;
@@ -85,6 +89,10 @@ private:
 	UPROPERTY(EditAnywhere, Category = "DynamicMesh To Cloth Collection")
 	bool bCopyToSim3DPositions = false;
 
+	/** Copy DynamicMesh Vertex Normals to Sim Normals */
+	UPROPERTY(EditAnywhere, Category = "DynamicMesh To Cloth Collection")
+	bool bCopyToSimNormals = false;
+
 	/** Copy DynamicMesh UVs to Sim2D Positions */
 	UPROPERTY(EditAnywhere, Category = "DynamicMesh To Cloth Collection")
 	bool bCopyUVsToSim2DPositions = false;
diff --git a/Engine/Plugins/ChaosClothAssetEditor/Source/ChaosClothAssetDataflowNodes/Public/ChaosClothAsset/ProceduralSelectionNode.h b/Engine/Plugins/ChaosClothAssetEditor/Source/ChaosClothAssetDataflowNodes/Public/ChaosClothAsset/ProceduralSelectionNode.h
new file mode 100644
--- /dev/null
+++ b/Engine/Plugins/ChaosClothAssetEditor/Source/ChaosClothAssetDataflowNodes/Public/ChaosClothAsset/ProceduralSelectionNode.h
@@ -0,0 +1,52 @@
+// Copyright Epic Games, Inc. All Rights Reserved.
+
+#pragma once 
+
+#include "Dataflow/DataflowNode.h"
+#include "GeometryCollection/ManagedArrayCollection.h"
+#include  "SelectionNode.h"
+#include "ProceduralSelectionNode.generated.h"
+
+/** Type of procedural selection */
+UENUM()
+enum class EChaosClothAssetProceduralSelectionType : uint8
+{
+	SelectAll /*Select all elements within the selected Group*/,
+	Conversion /*Convert an existing selection set to Group (if supported)*/
+};
+
+/** Procedurally generate a Cloth Selection set. */
+USTRUCT(Meta = (DataflowCloth))
+struct FChaosClothAssetProceduralSelectionNode : public FDataflowNode
+{
+	GENERATED_USTRUCT_BODY()
+	DATAFLOW_NODE_DEFINE_INTERNAL(FChaosClothAssetProceduralSelectionNode, "ProceduralSelection", "Cloth", "Cloth Procedural Selection")
+	DATAFLOW_NODE_RENDER_TYPE("SurfaceRender", FName("FClothCollection"), "Collection")
+
+public:
+
+	FChaosClothAssetProceduralSelectionNode(const UE::Dataflow::FNodeParameters& InParam, FGuid InGuid = FGuid::NewGuid());
+
+private:
+
+	UPROPERTY(Meta = (Dataflowinput, DataflowOutput, DataflowPassthrough = "Collection"))
+	FManagedArrayCollection Collection;
+
+	/** The name to be use as a selection. */
+	UPROPERTY(EditAnywhere, Category = "Procedural Selection", Meta = (DataflowOutput))
+	FString OutputName;
+
+	/** The type of element the selection refers to */
+	UPROPERTY(EditAnywhere, Category = "Procedural Selection")
+	FChaosClothAssetNodeSelectionGroup Group;
+
+	/** The procedural selection method */
+	UPROPERTY(EditAnywhere, Category = "Procedural Selection")
+	EChaosClothAssetProceduralSelectionType SelectionType;
+
+	/** Selection set to convert from when using Conversion selection type. */
+	UPROPERTY(EditAnywhere, Category = "Procedural Selection", Meta = (Dataflowinput, EditCondition = "SelectionType == EChaosClothAssetProceduralSelectionType::Conversion"))
+	FChaosClothAssetConnectableIStringValue ConversionInputName;
+
+	virtual void Evaluate(UE::Dataflow::FContext& Context, const FDataflowOutput* Out) const override;
+};
