Chaos Outfit Asset - Added SkeletalMesh export.

#rnx

diff --git a/Engine/Plugins/ChaosClothAsset/Source/ChaosClothAssetEngine/ChaosClothAssetEngine.Build.cs b/Engine/Plugins/ChaosClothAsset/Source/ChaosClothAssetEngine/ChaosClothAssetEngine.Build.cs
--- a/Engine/Plugins/ChaosClothAsset/Source/ChaosClothAssetEngine/ChaosClothAssetEngine.Build.cs
+++ b/Engine/Plugins/ChaosClothAsset/Source/ChaosClothAssetEngine/ChaosClothAssetEngine.Build.cs
@@ -21,16 +21,21 @@ public class ChaosClothAssetEngine : ModuleRules
 			new string[]
 			{
 				"Core",
-				"Engine",
 				"CoreUObject",
-				"RenderCore",
-				"RHI",
 				"Chaos",
-				"ChaosClothAsset",
-				"ChaosCloth",
 				"ChaosCaching",
+				"ChaosCloth",
+				"ChaosClothAsset",
 				"ClothingSystemRuntimeCommon",
-				"ClothingSystemRuntimeInterface"
+				"ClothingSystemRuntimeInterface",
+				"Engine",
+				"GeometryCore",
+				"MeshConversion",
+				"MeshDescription",
+				"RenderCore",
+				"RHI",
+				"SkeletalMeshDescription",
+				"SkeletalMeshUtilitiesCommon",
 			}
 		);
 
diff --git a/Engine/Plugins/ChaosClothAsset/Source/ChaosClothAssetEngine/Private/ChaosClothAsset/ClothAssetBase.cpp b/Engine/Plugins/ChaosClothAsset/Source/ChaosClothAssetEngine/Private/ChaosClothAsset/ClothAssetBase.cpp
--- a/Engine/Plugins/ChaosClothAsset/Source/ChaosClothAssetEngine/Private/ChaosClothAsset/ClothAssetBase.cpp
+++ b/Engine/Plugins/ChaosClothAsset/Source/ChaosClothAssetEngine/Private/ChaosClothAsset/ClothAssetBase.cpp
@@ -6,15 +6,27 @@
 #endif
 #include "ChaosClothAsset/ClothAssetPrivate.h"
 #include "ChaosClothAsset/ClothComponent.h"
+#if WITH_EDITOR
+#include "ChaosClothAsset/ClothEngineTools.h"
+#include "ChaosClothAsset/CollectionClothFacade.h"
+#include "DynamicMesh/DynamicMesh3.h"
+#endif
 #include "Engine/Engine.h"
 #include "Engine/SkeletalMesh.h"
 #include "Engine/SkinnedAssetAsyncCompileUtils.h"
-#include "Rendering/SkeletalMeshRenderData.h"
 #if WITH_EDITOR
 #include "Interfaces/ITargetPlatform.h"
+#include "Materials/Material.h"
+#include "Rendering/SkeletalMeshModel.h"
 #endif
+#include "Rendering/SkeletalMeshRenderData.h"
 #include "UObject/FortniteMainBranchObjectVersion.h"
-
+#if WITH_EDITOR
+#include "DynamicMeshToMeshDescription.h"
+#include "MaterialDomain.h"
+#include "SkeletalMeshAttributes.h"
+#include "StaticToSkeletalMeshConverter.h"
+#endif
 
 #include UE_INLINE_GENERATED_CPP_BY_NAME(ClothAssetBase)
 namespace UE::Chaos::ClothAsset::Private
@@ -484,6 +496,106 @@ PRAGMA_DISABLE_DEPRECATION_WARNINGS
 PRAGMA_ENABLE_DEPRECATION_WARNINGS
 }
 
+#if WITH_EDITOR
+bool UChaosClothAssetBase::ExportToSkeletalMesh(USkeletalMesh& SkeletalMesh) const
+{
+	using namespace UE::Chaos::ClothAsset;
+	using namespace UE::Geometry;
+
+	// Create one dynamic mesh from the multiple collections
+	TArray<UMaterialInterface*> RenderMaterials;
+	TArray<FDynamicMesh3> DynamicMeshes;
+	DynamicMeshes.AddDefaulted();  // TODO: Do the other LODs, one per Dynamic mesh
+
+	int32 MaterialOffset = 0;
+
+	const int32 NumModels = GetNumClothSimulationModels();
+
+	for (int32 ModelIndex = 0; ModelIndex < NumModels; ++ModelIndex)
+	{
+		const TArray<TSharedRef<const FManagedArrayCollection>>& ClothCollections = GetCollections(ModelIndex);
+
+		//for (const TSharedRef<const FManagedArrayCollection>& ClothCollection : ClothCollections)
+		if (ClothCollections.Num())
+		{
+			const TSharedRef<const FManagedArrayCollection>& ClothCollection = ClothCollections[0];  // TODO: Do other LODs
+			FCollectionClothConstFacade ClothFacade(ClothCollection);
+
+			if (ClothFacade.IsValid())
+			{
+				FDynamicMesh3 DynamicMeshPiece;
+				FClothEngineTools::BuildDynamicMeshFromRenderMesh(DynamicMeshPiece, ClothCollection, MaterialOffset);
+
+				if (ModelIndex == 0)
+				{
+					DynamicMeshes[0] = MoveTemp(DynamicMeshPiece);
+				}
+				else
+				{
+					DynamicMeshes[0].AppendWithOffsets(DynamicMeshPiece);
+				}
+
+				const TArrayView<const FString> MaterialPaths = ClothFacade.GetRenderMaterialPathName();
+				RenderMaterials.Reserve(RenderMaterials.Num() + MaterialPaths.Num());
+				for (int32 MaterialIndex = 0; MaterialIndex < MaterialPaths.Num(); ++MaterialIndex)
+				{
+					RenderMaterials.Emplace(LoadObject<UMaterialInterface>(nullptr, *MaterialPaths[MaterialIndex]));
+				}
+				MaterialOffset += MaterialPaths.Num();
+			}
+		}
+	}
+
+	// Create the skeletal mesh from the dynamic mesh
+	TArray<const FMeshDescription*> MeshDescriptions;
+
+	SkeletalMesh.PreEditChange(nullptr);
+	SkeletalMesh.GetImportedModel()->LODModels.Empty();
+	SkeletalMesh.ResetLODInfo();
+
+	TArray<FMeshDescription> LocalDescriptions;
+	for (const FDynamicMesh3& DynamicMesh : DynamicMeshes)
+	{
+		// Create a mesh description
+		FMeshDescription& MeshDescription = LocalDescriptions.AddDefaulted_GetRef();
+
+		// Add skeletal mesh attributes to the mesh description
+		FSkeletalMeshAttributes Attributes(MeshDescription);
+		Attributes.Register();
+
+		// Convert dynamic mesh to the mesh description
+		FConversionToMeshDescriptionOptions ConverterOptions;
+		FDynamicMeshToMeshDescription Converter(ConverterOptions);
+		constexpr bool bCopyTangents = true;
+		Converter.Convert(&DynamicMesh, MeshDescription, bCopyTangents);
+
+		// Add the created description to the list
+		MeshDescriptions.Add(&MeshDescription);
+	}
+
+	TArray<FSkeletalMaterial> SkeletalMaterials;
+	SkeletalMaterials.Reserve(Materials.Num());
+	for (UMaterialInterface* RenderMaterial : RenderMaterials)
+	{
+		SkeletalMaterials.Emplace(RenderMaterial, RenderMaterial->GetFName());
+	}
+	if (!SkeletalMaterials.Num())
+	{
+		SkeletalMaterials.Emplace(UMaterial::GetDefaultMaterial(MD_Surface));
+	}
+
+	constexpr bool bRecomputeNormals = false;
+	constexpr bool bRecomputeTangents = false;
+	return FStaticToSkeletalMeshConverter::InitializeSkeletalMeshFromMeshDescriptions(
+		&SkeletalMesh,
+		MeshDescriptions,
+		SkeletalMaterials,
+		GetRefSkeleton(),
+		bRecomputeNormals,
+		bRecomputeTangents);
+}
+#endif  // #if WITH_EDITOR
+
 void UChaosClothAssetBase::ReregisterComponents()
 {
 	// Recreate the simulation proxies with the updated physics asset
diff --git a/Engine/Plugins/ChaosClothAsset/Source/ChaosClothAssetEngine/Private/ChaosClothAsset/ClothCollectionToDynamicMesh.cpp b/Engine/Plugins/ChaosClothAsset/Source/ChaosClothAssetEngine/Private/ChaosClothAsset/ClothCollectionToDynamicMesh.cpp
new file mode 100644
--- /dev/null
+++ b/Engine/Plugins/ChaosClothAsset/Source/ChaosClothAssetEngine/Private/ChaosClothAsset/ClothCollectionToDynamicMesh.cpp
@@ -0,0 +1,738 @@
+// Copyright Epic Games, Inc. All Rights Reserved. 
+
+#if WITH_EDITOR
+
+#include "ChaosClothAsset/ClothCollectionToDynamicMesh.h"
+#include "ChaosClothAsset/CollectionClothFacade.h"
+#include "ChaosClothAsset/ClothCollectionGroup.h"
+#include "DynamicMesh/MeshNormals.h"
+#include "Engine/SkeletalMesh.h"
+#include "SkeletalMeshAttributes.h"
+#include "ToDynamicMesh.h"
+
+namespace UE::Chaos::ClothAsset
+{
+	namespace Private
+	{
+		//
+		// Wrapper for accessing a Cloth Pattern. Implements the interface expected by TToDynamicMesh<>.
+		//
+		class FClothPatternWrapper
+		{
+		public:
+			using TriIDType = int32;
+			using VertIDType = int32;
+			using WedgeIDType = int32;
+			using UVIDType = int32;
+			using NormalIDType = int32;
+			using ColorIDType = int32;
+			using ESourceMeshType = FClothCollectionToDynamicMesh::ESourceMeshType;
+
+			FClothPatternWrapper(const FCollectionClothConstFacade& ClothFacade, int32 PatternIndex, ESourceMeshType InSourceMeshType)
+				: SourceMeshType(InSourceMeshType)
+				, Cloth(ClothFacade)
+			{
+				NumTexCoords = 0;
+				if (PatternIndex == INDEX_NONE)
+				{
+					// All patterns in one dynamic mesh
+					switch (SourceMeshType)
+					{
+					case ESourceMeshType::Render:
+					{
+						const int32 NumVertices = Cloth.GetNumRenderVertices();
+						VertIDs.SetNum(NumVertices);
+						TConstArrayView<TArray<FVector2f>> UVs = Cloth.GetRenderUVs();
+						for (int32 VtxIndex = 0; VtxIndex < NumVertices; ++VtxIndex)
+						{
+							VertIDs[VtxIndex] = VtxIndex;
+							NumTexCoords = FMath::Max(NumTexCoords, UVs[VtxIndex].Num());
+						}
+						NormalIDs = VertIDs;
+
+						const TConstArrayView<FIntVector3> Indices = Cloth.GetRenderIndices();
+						const int32 NumFaces = Cloth.GetNumRenderFaces();
+						TriIDs.Reserve(NumFaces);
+						for (int32 TriIndex = 0; TriIndex < NumFaces; ++TriIndex)
+						{
+							if (Indices[TriIndex][0] != INDEX_NONE &&
+								Indices[TriIndex][1] != INDEX_NONE &&
+								Indices[TriIndex][2] != INDEX_NONE)
+							{
+								TriIDs.Add(TriIndex);
+							}
+						}
+					}
+					break;
+					case ESourceMeshType::Sim2D:
+					{
+						const int32 NumVertices = Cloth.GetNumSimVertices2D();
+						const TConstArrayView<int32> SimVertex3DLookup = Cloth.GetSimVertex3DLookup();
+						VertIDs.Reserve(NumVertices);
+						NormalIDs.Reserve(NumVertices);
+						for (int32 VtxIndex = 0; VtxIndex < NumVertices; ++VtxIndex)
+						{
+							if (SimVertex3DLookup[VtxIndex] != INDEX_NONE)
+							{
+								VertIDs.Add(VtxIndex);
+								NormalIDs.Add(SimVertex3DLookup[VtxIndex]);
+							}
+						}
+
+						const TConstArrayView<FIntVector3> Indices = Cloth.GetSimIndices2D();
+						const int32 NumFaces = Cloth.GetNumSimFaces();
+						TriIDs.Reserve(NumFaces);
+						for (int32 TriIndex = 0; TriIndex < NumFaces; ++TriIndex)
+						{
+							if (Indices[TriIndex][0] != INDEX_NONE &&
+								Indices[TriIndex][1] != INDEX_NONE &&
+								Indices[TriIndex][2] != INDEX_NONE &&
+								SimVertex3DLookup[Indices[TriIndex][0]] != INDEX_NONE &&
+								SimVertex3DLookup[Indices[TriIndex][1]] != INDEX_NONE &&
+								SimVertex3DLookup[Indices[TriIndex][2]] != INDEX_NONE)
+							{
+								TriIDs.Add(TriIndex);
+							}
+						}
+					}
+					break;
+					case ESourceMeshType::Sim3D:
+					{
+						const int32 NumVertices = Cloth.GetNumSimVertices3D();
+						VertIDs.SetNum(NumVertices);
+						for (int32 VtxIndex = 0; VtxIndex < NumVertices; ++VtxIndex)
+						{
+							VertIDs[VtxIndex] = VtxIndex;
+						}
+						NormalIDs = VertIDs;
+
+						const TConstArrayView<FIntVector3> Indices = Cloth.GetSimIndices3D();
+						const int32 NumFaces = Cloth.GetNumSimFaces();
+						TriIDs.Reserve(NumFaces);
+						for (int32 TriIndex = 0; TriIndex < NumFaces; ++TriIndex)
+						{
+							if (Indices[TriIndex][0] != INDEX_NONE &&
+								Indices[TriIndex][1] != INDEX_NONE &&
+								Indices[TriIndex][2] != INDEX_NONE)
+							{
+								TriIDs.Add(TriIndex);
+							}
+						}
+						NumTexCoords = 1; // SimPositions2D will be stored as UVs
+					}
+					break;
+					default:
+						checkNoEntry();
+					}
+				}
+				else
+				{
+					switch (SourceMeshType)
+					{
+					case ESourceMeshType::Render:
+					{
+						FCollectionClothRenderPatternConstFacade Pattern = Cloth.GetRenderPattern(PatternIndex);
+						const int32 NumVertices = Pattern.GetNumRenderVertices();
+						const int32 VertexOffset = Pattern.GetRenderVerticesOffset();
+						VertIDs.SetNum(NumVertices);
+						TConstArrayView<TArray<FVector2f>> UVs = Pattern.GetRenderUVs();
+						for (int32 VtxIndex = 0; VtxIndex < NumVertices; ++VtxIndex)
+						{
+							VertIDs[VtxIndex] = VtxIndex + VertexOffset;
+							NumTexCoords = FMath::Max(NumTexCoords, UVs[VtxIndex].Num());
+						}
+						NormalIDs = VertIDs;
+
+						const TConstArrayView<FIntVector3> Indices = Pattern.GetRenderIndices();
+						const int32 NumFaces = Pattern.GetNumRenderFaces();
+						const int32 FaceOffset = Pattern.GetRenderFacesOffset();
+						TriIDs.Reserve(NumFaces);
+						for (int32 TriIndex = 0; TriIndex < NumFaces; ++TriIndex)
+						{
+							if (Indices[TriIndex][0] != INDEX_NONE &&
+								Indices[TriIndex][1] != INDEX_NONE &&
+								Indices[TriIndex][2] != INDEX_NONE)
+							{
+								TriIDs.Add(TriIndex + FaceOffset);
+							}
+						}
+					}
+					break;
+					case ESourceMeshType::Sim2D:
+					{
+						FCollectionClothSimPatternConstFacade Pattern = Cloth.GetSimPattern(PatternIndex);
+						const int32 NumVertices = Pattern.GetNumSimVertices2D();
+						const TConstArrayView<int32> SimVertex3DLookup = Pattern.GetSimVertex3DLookup();
+						const int32 VertexOffset = Pattern.GetSimVertices2DOffset();
+						VertIDs.Reserve(NumVertices);
+						NormalIDs.Reserve(NumVertices);
+						for (int32 VtxIndex = 0; VtxIndex < NumVertices; ++VtxIndex)
+						{
+							if (SimVertex3DLookup[VtxIndex] != INDEX_NONE)
+							{
+								VertIDs.Add(VtxIndex + VertexOffset);
+								NormalIDs.Add(SimVertex3DLookup[VtxIndex]);
+							}
+						}
+
+						const TConstArrayView<FIntVector3> Indices = Pattern.GetSimIndices2D();
+						const int32 NumFaces = Pattern.GetNumSimFaces();
+						const int32 FaceOffset = Pattern.GetSimFacesOffset();
+						TriIDs.Reserve(NumFaces);
+						for (int32 TriIndex = 0; TriIndex < NumFaces; ++TriIndex)
+						{
+							if (Indices[TriIndex][0] != INDEX_NONE &&
+								Indices[TriIndex][1] != INDEX_NONE &&
+								Indices[TriIndex][2] != INDEX_NONE &&
+								SimVertex3DLookup[Indices[TriIndex][0] - VertexOffset] != INDEX_NONE &&
+								SimVertex3DLookup[Indices[TriIndex][1] - VertexOffset] != INDEX_NONE &&
+								SimVertex3DLookup[Indices[TriIndex][2] - VertexOffset] != INDEX_NONE)
+							{
+								TriIDs.Add(TriIndex + FaceOffset);
+							}
+						}
+					}
+					break;
+					case ESourceMeshType::Sim3D:
+					{
+						FCollectionClothSimPatternConstFacade Pattern = Cloth.GetSimPattern(PatternIndex);
+						VertIDs = Pattern.GetSimVertex3DLookup();
+						NormalIDs = Pattern.GetSimVertex3DLookup();
+
+						const TConstArrayView<FIntVector3> Indices = Pattern.GetSimIndices3D();
+						const int32 NumFaces = Pattern.GetNumSimFaces();
+						const int32 FaceOffset = Pattern.GetSimFacesOffset();
+						TriIDs.Reserve(NumFaces);
+						for (int32 TriIndex = 0; TriIndex < NumFaces; ++TriIndex)
+						{
+							if (Indices[TriIndex][0] != INDEX_NONE &&
+								Indices[TriIndex][1] != INDEX_NONE &&
+								Indices[TriIndex][2] != INDEX_NONE)
+							{
+								TriIDs.Add(TriIndex + FaceOffset);
+							}
+						}
+						NumTexCoords = 1; // SimPositions2D will be stored as UVs
+					}
+					break;
+					default:
+						checkNoEntry();
+					}
+				}
+
+				//
+				// Weight map layers precomputation
+				//
+				if (SourceMeshType != ESourceMeshType::Render)
+				{
+					WeightMapNames = Cloth.GetWeightMapNames();
+				}
+				else
+				{
+					WeightMapNames = Cloth.GetUserDefinedAttributeNames<float>(ClothCollectionGroup::RenderVertices);
+				}
+
+				// Set the reference skeleton if available
+				const FString& SkeletalMeshPathName = ClothFacade.GetSkeletalMeshPathName();
+				const USkeletalMesh* const SkeletalMesh = !SkeletalMeshPathName.IsEmpty() ?
+					LoadObject<USkeletalMesh>(nullptr, *SkeletalMeshPathName, nullptr, LOAD_None, nullptr) :
+					nullptr;
+				RefSkeleton = SkeletalMesh ? &SkeletalMesh->GetRefSkeleton() : nullptr;
+			}
+
+			int32 NumTris() const
+			{
+				return TriIDs.Num();
+			}
+
+			int32 NumVerts() const
+			{
+				return VertIDs.Num();
+			}
+
+			int32 NumUVLayers() const
+			{
+				return NumTexCoords;
+			}
+
+			int32 NumWeightMapLayers() const
+			{
+				return WeightMapNames.Num();
+			}
+
+			FName GetWeightMapName(int32 LayerIndex) const
+			{
+				return WeightMapNames[LayerIndex];
+			}
+
+			float GetVertexWeight(int32 LayerIndex, VertIDType VertexIndex) const
+			{
+				checkSlow(LayerIndex < WeightMapNames.Num());
+
+				// All weights live on 3D indices.
+				const int32 VertexWeightIndex = (SourceMeshType == ESourceMeshType::Sim2D) ? Cloth.GetSimVertex3DLookup()[VertexIndex] : VertexIndex;
+
+				TConstArrayView<float> WeightMap;
+
+				if (SourceMeshType != ESourceMeshType::Render)
+				{
+					WeightMap = Cloth.GetWeightMap(WeightMapNames[LayerIndex]);
+				}
+				else
+				{
+					WeightMap = Cloth.GetUserDefinedAttribute<float>(WeightMapNames[LayerIndex], ClothCollectionGroup::RenderVertices);
+				}
+
+				checkSlow(VertexWeightIndex < WeightMap.Num());
+				if (VertexWeightIndex < WeightMap.Num())
+				{
+					return WeightMap[VertexWeightIndex];
+				}
+				else
+				{
+					return 0;
+				}
+			}
+
+			// --"Vertex Buffer" info
+			const TArray<VertIDType>& GetVertIDs() const
+			{
+				return VertIDs;
+			}
+
+			FVector3d GetPosition(VertIDType VtxID) const
+			{
+				switch (SourceMeshType)
+				{
+				case ESourceMeshType::Render:
+				{
+					const TConstArrayView<FVector3f> RenderPositions = Cloth.GetRenderPosition();
+					return FVector3d(RenderPositions[VtxID]);
+				}
+				case ESourceMeshType::Sim2D:
+				{
+					const TConstArrayView<FVector2f> SimPositions = Cloth.GetSimPosition2D();
+					const FVector2f& Pos = SimPositions[VtxID];
+					return FVector3d(Pos[0], Pos[1], 0.0);
+				}
+				case ESourceMeshType::Sim3D:
+				{
+					const TConstArrayView<FVector3f> SimPositions = Cloth.GetSimPosition3D();
+					const FVector3f& Pos = SimPositions[VtxID];
+					return FVector3d(Pos[0], Pos[1], Pos[2]);
+				}
+				default:
+					checkNoEntry();
+					return FVector3d();
+				};
+			}
+
+			// --"Index Buffer" info
+			const TArray<TriIDType>& GetTriIDs() const
+			{
+				return TriIDs;
+			}
+
+			bool GetTri(TriIDType TriID, VertIDType& VID0, VertIDType& VID1, VertIDType& VID2) const
+			{
+				FIntVector Face;
+
+				switch (SourceMeshType)
+				{
+				case ESourceMeshType::Render:
+					Face = Cloth.GetRenderIndices()[TriID];
+					break;
+				case ESourceMeshType::Sim2D:
+					Face = Cloth.GetSimIndices2D()[TriID];
+					break;
+				case ESourceMeshType::Sim3D:
+					Face = Cloth.GetSimIndices3D()[TriID];
+					break;
+				default:
+					checkNoEntry();
+				}
+
+				VID0 = Face[0];
+				VID1 = Face[1];
+				VID2 = Face[2];
+
+				return true;
+			}
+
+			bool HasNormals() const
+			{
+				return true;
+			}
+
+			bool HasTangents() const
+			{
+				return SourceMeshType == ESourceMeshType::Render;
+			}
+
+			bool HasBiTangents() const
+			{
+				return SourceMeshType == ESourceMeshType::Render;
+			}
+
+			bool HasColors() const
+			{
+				return SourceMeshType == ESourceMeshType::Render;
+			}
+
+			// -- Access to per-wedge attributes -- //
+			void GetWedgeIDs(const TriIDType& TriID, WedgeIDType& WID0, WedgeIDType& WID1, WedgeIDType& WID2) const
+			{
+				check(SourceMeshType == ESourceMeshType::Sim3D);
+				const FIntVector& Face = Cloth.GetSimIndices2D()[TriID];
+
+				WID0 = Face[0];
+				WID1 = Face[1];
+				WID2 = Face[2];
+			}
+
+			FVector2f GetWedgeUV(int32 UVLayerIndex, WedgeIDType WID) const
+			{
+				check(SourceMeshType == ESourceMeshType::Sim3D);
+				return Cloth.GetSimPosition2D()[WID];
+			}
+
+			FVector3f GetWedgeNormal(WedgeIDType WID) const
+			{
+				checkf(false, TEXT("FClothPatternWrapper: ClothPatterns are not expected to use Wedges"));
+				return FVector3f();
+			}
+
+			FVector3f GetWedgeTangent(WedgeIDType WID) const
+			{
+				checkf(false, TEXT("FClothPatternWrapper: ClothPatterns are not expected to use Wedges"));
+				return FVector3f();
+			}
+
+			FVector3f GetWedgeBiTangent(WedgeIDType WID) const
+			{
+				checkf(false, TEXT("FClothPatternWrapper: ClothPatterns are not expected to use Wedges"));
+				return FVector3f();
+			}
+
+			FVector4f GetWedgeColor(WedgeIDType WID) const
+			{
+				checkf(false, TEXT("FClothPatternWrapper: ClothPatterns are not expected to use Wedges"));
+				return FVector4f();
+			}
+			// -- End of per-wedge attribute access -- //
+
+			int32 GetMaterialIndex(TriIDType TriID) const
+			{
+				checkf(false, TEXT("FClothPatternWrapper: Material indexing should be accomplished by passing a function into Convert"));
+				return 0;
+			}
+
+			int32 NumSkinWeightAttributes() const
+			{
+				return 1;
+			}
+
+			UE::AnimationCore::FBoneWeights GetVertexSkinWeight(int32 SkinWeightAttributeIndex, VertIDType InVertexID) const
+			{
+				using namespace UE::AnimationCore;
+
+				checkfSlow(SkinWeightAttributeIndex == 0, TEXT("Cloth assets should only have one skin weight profile"));
+
+				const bool bGetRenderMeshData = (SourceMeshType == ESourceMeshType::Render);
+				const TConstArrayView<TArray<int32>> BoneIndices = bGetRenderMeshData ? Cloth.GetRenderBoneIndices() : Cloth.GetSimBoneIndices();
+				const TConstArrayView<TArray<float>> BoneWeights = bGetRenderMeshData ? Cloth.GetRenderBoneWeights() : Cloth.GetSimBoneWeights();
+
+				// All weights live on 3D indices. Need to convert 2D index to 3D index.
+				const VertIDType VertexID = (SourceMeshType == ESourceMeshType::Sim2D) ? Cloth.GetSimVertex3DLookup()[InVertexID] : InVertexID;
+
+				if (ensure(VertexID >= 0 && VertexID < BoneIndices.Num()))
+				{
+					const TArray<int32> Indices = BoneIndices[VertexID];
+					const TArray<float> Weights = BoneWeights[VertexID];
+					const int32 NumInfluences = Indices.Num();
+					check(Weights.Num() == NumInfluences);
+
+					TArray<FBoneWeight> BoneWeightArray;
+					BoneWeightArray.SetNumUninitialized(NumInfluences);
+
+					for (int32 Idx = 0; Idx < NumInfluences; ++Idx)
+					{
+						BoneWeightArray[Idx] = FBoneWeight(static_cast<FBoneIndexType>(Indices[Idx]), Weights[Idx]);
+					}
+
+					return FBoneWeights::Create(BoneWeightArray, FBoneWeightsSettings());
+				}
+				else
+				{
+					return FBoneWeights();
+				}
+			}
+
+			FName GetSkinWeightAttributeName(int32 SkinWeightAttributeIndex) const
+			{
+				checkfSlow(SkinWeightAttributeIndex == 0, TEXT("Cloth assets should only have one skin weight profile"));
+
+				return FSkeletalMeshAttributes::DefaultSkinWeightProfileName;
+			}
+
+			int32 GetNumBones() const
+			{
+				return RefSkeleton ? RefSkeleton->GetRawBoneNum() : 0;
+			}
+
+			FName GetBoneName(int32 BoneIdx) const
+			{
+				if (ensure(BoneIdx >= 0 && BoneIdx < GetNumBones()) && RefSkeleton)
+				{
+					return RefSkeleton->GetRawRefBoneInfo()[BoneIdx].Name;
+				}
+
+				return NAME_None;
+			}
+
+			int32 GetBoneParentIndex(int32 BoneIdx) const
+			{
+				if (ensure(BoneIdx >= 0 && BoneIdx < GetNumBones()) && RefSkeleton)
+				{
+					return RefSkeleton->GetRawRefBoneInfo()[BoneIdx].ParentIndex;
+				}
+
+				return INDEX_NONE;
+			}
+
+			FTransform GetBonePose(int32 BoneIdx) const
+			{
+				if (ensure(BoneIdx >= 0 && BoneIdx < GetNumBones()) && RefSkeleton)
+				{
+					return RefSkeleton->GetRawRefBonePose()[BoneIdx];
+				}
+
+				return FTransform::Identity;
+			}
+
+			FVector4f GetBoneColor(int32 BoneIdx) const
+			{
+				return FVector4f::One();
+			}
+
+			const TArray<int32>& GetNormalIDs() const
+			{
+				return NormalIDs;
+			}
+
+			FVector3f GetNormal(NormalIDType ID) const
+			{
+				return (SourceMeshType == ESourceMeshType::Render) ? Cloth.GetRenderNormal()[ID] : Cloth.GetSimNormal()[ID];
+			}
+
+			bool GetNormalTri(const TriIDType& TriID, NormalIDType& ID0, NormalIDType& ID1, NormalIDType& ID2) const
+			{
+				if (SourceMeshType == ESourceMeshType::Sim2D)
+				{
+					// All normal data lives on 3D indices.
+					const FIntVector Face = Cloth.GetSimIndices3D()[TriID];
+
+					ID0 = Face[0];
+					ID1 = Face[1];
+					ID2 = Face[2];
+
+					return true;
+				}
+				else
+				{
+					return GetTri(TriID, ID0, ID1, ID2);
+				}
+			}
+
+			const TArray<int32>& GetUVIDs(int32 LayerID) const
+			{
+				return (SourceMeshType == ESourceMeshType::Render) ? NormalIDs : EmptyArray;
+			}
+
+			FVector2f GetUV(int32 LayerID, UVIDType UVID) const
+			{
+				checkf(SourceMeshType == ESourceMeshType::Render, TEXT("Requested UVs from a Sim mesh"));
+				const TConstArrayView<FVector2f> VertexUVs = Cloth.GetRenderUVs()[UVID];
+				return VertexUVs.IsValidIndex(LayerID) ? VertexUVs[LayerID] : FVector2f(0.f);
+			}
+
+			bool GetUVTri(int32 LayerID, const TriIDType& TriID, UVIDType& ID0, UVIDType& ID1, UVIDType& ID2) const
+			{
+				return GetTri(TriID, ID0, ID1, ID2);
+			}
+
+			const TArray<int32>& GetTangentIDs() const
+			{
+				return (SourceMeshType == ESourceMeshType::Render) ? NormalIDs : EmptyArray;
+			}
+
+			FVector3f GetTangent(NormalIDType ID) const
+			{
+				checkf(SourceMeshType == ESourceMeshType::Render, TEXT("Requested Tangent from a Sim mesh"));
+				return Cloth.GetRenderTangentU()[ID];
+			}
+
+			bool GetTangentTri(const TriIDType& TriID, NormalIDType& ID0, NormalIDType& ID1, NormalIDType& ID2) const
+			{
+				return GetNormalTri(TriID, ID0, ID1, ID2);
+			}
+
+			const TArray<int32>& GetBiTangentIDs() const
+			{
+				return (SourceMeshType == ESourceMeshType::Render) ? NormalIDs : EmptyArray;
+			}
+
+			FVector3f GetBiTangent(NormalIDType ID) const
+			{
+				checkf(SourceMeshType == ESourceMeshType::Render, TEXT("Requested Bitangent from a Sim mesh"));
+				return Cloth.GetRenderTangentV()[ID];
+			}
+
+			bool GetBiTangentTri(const TriIDType& TriID, NormalIDType& ID0, NormalIDType& ID1, NormalIDType& ID2) const
+			{
+				return GetNormalTri(TriID, ID0, ID1, ID2);
+			}
+
+			const TArray<int32>& GetColorIDs() const
+			{
+				return (SourceMeshType == ESourceMeshType::Render) ? NormalIDs : EmptyArray;
+			}
+
+			FVector4f GetColor(ColorIDType VID) const
+			{
+				checkf(SourceMeshType == ESourceMeshType::Render, TEXT("Requested color from a Sim mesh"));
+				return Cloth.GetRenderColor()[VID];
+			}
+
+			bool GetColorTri(const TriIDType& TriID, ColorIDType& ID0, ColorIDType& ID1, ColorIDType& ID2) const
+			{
+				return GetNormalTri(TriID, ID0, ID1, ID2);
+			}
+
+		private:
+
+			const ESourceMeshType SourceMeshType;
+			const FCollectionClothConstFacade& Cloth;
+
+			TArray<TriIDType> TriIDs;		// indices into Indices
+			TArray<VertIDType> VertIDs;		// indices into Positions2D or Positions3D
+			TArray<NormalIDType> NormalIDs; // indices into Normals (and all per vertex data other than positions)
+			int32 NumTexCoords;
+
+			TArray<FName> WeightMapNames;
+
+			TArray<int32> EmptyArray;
+
+			const FReferenceSkeleton* RefSkeleton = nullptr;
+		};
+	}  // namespace Private
+
+	void FClothCollectionToDynamicMesh::Convert(
+		const TSharedRef<const FManagedArrayCollection> ClothCollection,
+		int32 PatternIndex,
+		ESourceMeshType SourceMeshType,
+		UE::Geometry::FDynamicMesh3& MeshOut,
+		bool bDisableAttributes,
+		int32 MaterialOffset)
+	{
+		const FCollectionClothConstFacade ClothFacade(ClothCollection);
+
+		// Actual conversion
+		UE::Geometry::TToDynamicMesh<Private::FClothPatternWrapper> PatternToDynamicMesh;
+		Private::FClothPatternWrapper PatternWrapper(ClothFacade, PatternIndex, SourceMeshType);
+
+		auto TriangleToGroupFunction = [](Private::FClothPatternWrapper::TriIDType) { return 0; };
+
+		if (bDisableAttributes)
+		{
+			PatternToDynamicMesh.ConvertWOAttributes(MeshOut, PatternWrapper, TriangleToGroupFunction);
+		}
+		else
+		{
+			constexpr bool bCopyTangents = true;
+			const bool bIsRenderType = SourceMeshType == ESourceMeshType::Render;
+			auto TriangleToMaterialFunction = [PatternIndex, bIsRenderType, &ClothFacade](Private::FClothPatternWrapper::TriIDType TriID)->int32
+			{
+				if (bIsRenderType)
+				{
+					if (PatternIndex != INDEX_NONE)
+					{
+						return PatternIndex;
+					}
+					const int32 FoundPattern = ClothFacade.FindRenderPatternByFaceIndex(TriID);
+					check(FoundPattern != INDEX_NONE);
+					return FoundPattern;
+				}
+
+				// Sim meshes will have a default material with MaterialID zero applied
+				return 0;
+			};
+
+			PatternToDynamicMesh.Convert(MeshOut, PatternWrapper, TriangleToGroupFunction, TriangleToMaterialFunction, bCopyTangents);
+		}
+
+		// Add non-manifold mapping data if DynamicMesh indices don't match orig indices
+		const int32 NumVertices = SourceMeshType == ESourceMeshType::Render ? ClothFacade.GetNumRenderVertices() :
+			SourceMeshType == ESourceMeshType::Sim2D ? ClothFacade.GetNumSimVertices2D() : ClothFacade.GetNumSimVertices3D();
+		bool bVertexIndicesMatch = false;
+		if (PatternToDynamicMesh.ToSrcVertIDMap.Num() == NumVertices)
+		{
+			bVertexIndicesMatch = true;
+			for (int32 VertId = 0; VertId < NumVertices; ++VertId)
+			{
+				if (PatternToDynamicMesh.ToSrcVertIDMap[VertId] != VertId)
+				{
+					bVertexIndicesMatch = false;
+					break;
+				}
+			}
+		}
+		if (!bVertexIndicesMatch)
+		{
+			MeshOut.EnableAttributes();
+			FMappingSupport::AttachVertexMappingData(PatternToDynamicMesh.ToSrcVertIDMap, MeshOut);
+		}
+		const int32 NumTriangles = SourceMeshType == ESourceMeshType::Render ? ClothFacade.GetNumRenderFaces() : ClothFacade.GetNumSimFaces();
+		bool bTriangleIndicesMatch = false;
+		if (PatternToDynamicMesh.ToSrcTriIDMap.Num() == NumTriangles)
+		{
+			bTriangleIndicesMatch = true;
+			for (int32 TriId = 0; TriId < NumTriangles; ++TriId)
+			{
+				if (PatternToDynamicMesh.ToSrcTriIDMap[TriId] != TriId)
+				{
+					bTriangleIndicesMatch = false;
+					break;
+				}
+			}
+		}
+		if (!bTriangleIndicesMatch)
+		{
+			MeshOut.EnableAttributes();
+			FMappingSupport::AttachTriangleMappingData(PatternToDynamicMesh.ToSrcTriIDMap, MeshOut);
+		}
+
+		// Fix up any triangles without valid material IDs
+		if (SourceMeshType == ESourceMeshType::Render)
+		{
+			const TConstArrayView<FString> MaterialPaths = ClothFacade.GetRenderMaterialPathName();
+			int32 DefaultMaterialID = 0;
+			for (const int32 TriID : MeshOut.TriangleIndicesItr())
+			{
+				int32 MaterialID;
+				MeshOut.Attributes()->GetMaterialID()->GetValue(TriID, &MaterialID);
+				//const int32 MaterialID = MeshOut.Attributes()->GetMaterialID()->GetValue(TriID);
+				if (!MaterialPaths.IsValidIndex(MaterialID))
+				{
+					MeshOut.Attributes()->GetMaterialID()->SetValue(TriID, DefaultMaterialID);
+				}
+				else
+				{
+					MeshOut.Attributes()->GetMaterialID()->SetValue(TriID, MaterialID + MaterialOffset);
+				}
+			}
+		}
+	}
+}	// namespace UE::Chaos::ClothAsset
+
+#endif  // #if WITH_EDITOR
diff --git a/Engine/Plugins/ChaosClothAsset/Source/ChaosClothAssetEngine/Private/ChaosClothAsset/ClothCollectionToDynamicMesh.h b/Engine/Plugins/ChaosClothAsset/Source/ChaosClothAssetEngine/Private/ChaosClothAsset/ClothCollectionToDynamicMesh.h
new file mode 100644
--- /dev/null
+++ b/Engine/Plugins/ChaosClothAsset/Source/ChaosClothAssetEngine/Private/ChaosClothAsset/ClothCollectionToDynamicMesh.h
@@ -0,0 +1,133 @@
+// Copyright Epic Games, Inc. All Rights Reserved.
+
+#pragma once
+
+#if WITH_EDITOR
+
+#include "Templates/SharedPointer.h"
+#include "DynamicMesh/NonManifoldMappingSupport.h"
+
+namespace UE::Geometry { class FDynamicMesh3; }
+struct FManagedArrayCollection;
+
+namespace UE::Chaos::ClothAsset
+{
+	namespace Private
+	{
+		class FDynamicMeshSourceTriangleIdAttribute;
+	}
+
+	/**
+	 * Helper class to convert cloth collections into dynamic meshes.
+	 */
+	class FClothCollectionToDynamicMesh
+	{
+	public:
+		/** The type of the source meshg being converted. */
+		enum class ESourceMeshType : uint8
+		{
+			Sim2D = 0,
+			Sim3D = 1,
+			Render = 2
+		};
+
+		/**
+		 * Convert a single pattern from a ClothAsset to a FDynamicMesh3. 
+		 * When PatternIndex = INDEX_NONE, convert the entire cloth asset to an FDynamicMesh3.
+		 * When SourceMeshType = Sim3D, this generates a welded mesh (with native welded indexing) when PatternIndex = INDEX_NONE,
+		 * and an unwelded (with native 2D unwelded indexing) when PatternIndex is a valid SimPattern Index.
+		 */
+		static void Convert(
+			const TSharedRef<const FManagedArrayCollection> ClothCollection,
+			int32 PatternIndex,
+			ESourceMeshType SourceMeshType,
+			UE::Geometry::FDynamicMesh3& MeshOut,
+			bool bDisableAttributes = false,
+			int32 MaterialOffset = 0);
+
+	private:
+		// This uses the NonManifoldMapping attribute for vertices, and a new attribute for triangles.
+		class FMappingSupport : public UE::Geometry::FNonManifoldMappingSupport
+		{
+		public:
+			// Name used to identify triangle attribute data generated during conversion to a DynamicMesh in the case that the source was non-manifold.
+			static const FName ClothMeshTIDsAttrName;
+
+			// Attaches or replaces vertex mapping data to the provided mesh.
+			// Note: Failure occurs if the DynamicMesh does not have attributes enabled or if the provided array is not long enough to provide a mapping value for each DynamicMesh vertex id.
+			static bool AttachVertexMappingData(const TArray<int32>& VertexToOriginalVertexIDMap, UE::Geometry::FDynamicMesh3& InOutMesh)
+			{
+				return AttachNonManifoldVertexMappingData(VertexToOriginalVertexIDMap, InOutMesh);
+			}
+
+			// Attaches or replaces  triangle mapping data to the provided mesh.
+			// Note: Failure occurs if the DynamicMesh does not have attributes enabled or if the provided array is not long enough to provide a mapping value for each DynamicMesh triangle id.
+			static bool AttachTriangleMappingData(const TArray<int32>& TriangleToOriginalVertexIDMap, UE::Geometry::FDynamicMesh3& InOutMesh);
+
+			// Removes triangle mapping data.
+			// Note: this will invalidate any NonManifoldMappingSupport/ClothPatternToDynamicMeshMappingSupport object associated with this DynamicMesh,
+			// and subsequent use of such object will produce unexpected results.
+			static void RemoveVertexMappingData(UE::Geometry::FDynamicMesh3& InOutMesh)
+			{
+				RemoveNonManifoldVertexMappingData(InOutMesh);
+			}
+
+			static void RemoveTriangleMappingData(UE::Geometry::FDynamicMesh3& InOutMesh);
+
+			//  Removes all mapping data.
+			//  Note, this will invalidate any  NonManifoldMappingSupport/ClothPatternToDynamicMeshMappingSupport object associated with this DynamicMesh,
+			//  and subsequent use of such object will produce unexpected results.
+			static void RemoveAllMappingData(UE::Geometry::FDynamicMesh3& InOutMesh)
+			{
+				RemoveAllNonManifoldMappingData(InOutMesh);
+				RemoveTriangleMappingData(InOutMesh);
+			}
+
+			explicit FMappingSupport(const UE::Geometry::FDynamicMesh3& Mesh);
+
+			// Update the support for a new DynamicMesh.
+			void Reset(const UE::Geometry::FDynamicMesh3& Mesh);
+
+			// Return true if attribute data indicates that the source data that was converted to this DynamicMesh contained mapped vertices.
+			bool IsMappedVertexInSource() const
+			{
+				return IsNonManifoldVertexInSource();
+			}
+
+			// Return true if attribute data indicates that the source data that was converted to this DynamicMesh contained mapped triangles.
+			bool IsMappedTriangleInSource() const;
+
+			// Return true if the provided DynamicMesh vertex id resulted from a remapped vertex in the source data.
+			// Note: the code assumes but does not check that vid is a valid vertex id
+			bool IsMappedVertexID(const int32 vid) const
+			{
+				return !(vid == GetOriginalNonManifoldVertexID(vid));
+			}
+
+			// Return true if the provided DynamicMesh triangle id resulted from a remapped triangle in the source data.
+			// Note: the code assumes but does not check that Tid is a valid triangle id
+			bool IsMappedTriangleID(const int32 Tid) const
+			{
+				return !(Tid == GetOriginalTriangleID(Tid));
+			}
+
+			// Return the vertex ID in the original cloth data used to generate this DynamicMesh associated with the provided vertex id.
+			// In the case that the source data was actually manifold the returned vertex id will be identical to the DynamicMesh vertex id.
+			// Note: the code assumes but does not check that vid is a valid vertex vid.
+			int32 GetOriginalVertexID(const int32 vid) const
+			{
+				return GetOriginalNonManifoldVertexID(vid);
+			}
+
+			// Return the triangle ID in the original cloth data used to generate this DynamicMesh associated with the provided triangle id.
+			// In the case that the source data was actually manifold the returned triangle id will be identical to the DynamicMesh triangle id.
+			// Note: the code assumes but does not check that Tid is a valid triangle id
+			int32 GetOriginalTriangleID(const int32 Tid) const;
+
+		private:
+			const UE::Chaos::ClothAsset::Private::FDynamicMeshSourceTriangleIdAttribute* SrcTIDsAttribute = nullptr;
+		};
+	};
+}  // namespace UE::Chaos::ClothAsset
+
+#endif  // #if WITH_EDITOR
diff --git a/Engine/Plugins/ChaosClothAsset/Source/ChaosClothAssetEngine/Private/ChaosClothAsset/ClothCollectionToDynamicMeshMappingSupport.cpp b/Engine/Plugins/ChaosClothAsset/Source/ChaosClothAssetEngine/Private/ChaosClothAsset/ClothCollectionToDynamicMeshMappingSupport.cpp
new file mode 100644
--- /dev/null
+++ b/Engine/Plugins/ChaosClothAsset/Source/ChaosClothAssetEngine/Private/ChaosClothAsset/ClothCollectionToDynamicMeshMappingSupport.cpp
@@ -0,0 +1,379 @@
+// Copyright Epic Games, Inc. All Rights Reserved. 
+
+#if WITH_EDITOR
+
+#include "ChaosClothAsset/ClothCollectionToDynamicMesh.h"
+#include "DynamicMesh/DynamicMeshAttributeSet.h"
+
+namespace UE::Chaos::ClothAsset
+{
+	namespace Private
+	{
+		class FDynamicMeshSourceTriangleIdAttributeChange final :
+			public UE::Geometry::FDynamicMeshAttributeChangeBase
+		{
+			friend class FDynamicMeshSourceTriangleIdAttribute;
+			TArray<TPair<int32, int32>> OldSourceTriangleIds, NewSourceTriangleIds;
+
+			bool bOldInvalidState = false;
+			bool bNewInvalidState = false;
+
+		public:
+			FDynamicMeshSourceTriangleIdAttributeChange() = default;
+
+			virtual ~FDynamicMeshSourceTriangleIdAttributeChange() override = default;
+
+			void SaveInitialTriangle(const UE::Geometry::FDynamicMeshAttributeBase* Attribute, int TriangleID) override;
+
+			void StoreAllFinalTriangles(const UE::Geometry::FDynamicMeshAttributeBase* Attribute, const TArray<int>& TriangleIDs) override;
+
+			bool Apply(UE::Geometry::FDynamicMeshAttributeBase* Attribute, bool bRevert) const override;
+		};
+
+		// TODO: Can this class inherit from UE::Geometry::TDynamicMeshScalarTriangleAttribute<int32>
+		// and then only need to implement a smaller subset of the methods? (On* and IsValid?)
+		class FDynamicMeshSourceTriangleIdAttribute final :
+			public UE::Geometry::FDynamicMeshAttributeBase
+		{
+			friend class FDynamicMeshSourceTriangleIdAttributeChange;
+
+			UE::Geometry::FDynamicMesh3* Parent = nullptr;
+			UE::Geometry::TDynamicVector<int32> SourceTriangleIds;
+			bool bValid = false;
+
+		public:
+			FDynamicMeshSourceTriangleIdAttribute() = default;
+
+			FDynamicMeshSourceTriangleIdAttribute(UE::Geometry::FDynamicMesh3* InParent) : Parent(InParent)
+			{
+			}
+
+			virtual ~FDynamicMeshSourceTriangleIdAttribute() override = default;
+
+			const UE::Geometry::FDynamicMesh3* GetParent() const
+			{
+				return Parent;
+			}
+			UE::Geometry::FDynamicMesh3* GetParent()
+			{
+				return Parent;
+			}
+
+			bool IsValid() const
+			{
+				return bValid;
+			}
+
+			UE::Geometry::FDynamicMeshAttributeBase* MakeCopy(UE::Geometry::FDynamicMesh3* ParentIn) const override
+			{
+				FDynamicMeshSourceTriangleIdAttribute* Attribute = new FDynamicMeshSourceTriangleIdAttribute(ParentIn);
+				Attribute->Copy(*this);
+				return Attribute;
+			}
+
+			UE::Geometry::FDynamicMeshAttributeBase* MakeNew(UE::Geometry::FDynamicMesh3* ParentIn) const override
+			{
+				FDynamicMeshSourceTriangleIdAttribute* Attribute = new FDynamicMeshSourceTriangleIdAttribute(ParentIn);
+				Attribute->Initialize();
+				return Attribute;
+			}
+
+			void CompactInPlace(const UE::Geometry::FCompactMaps& CompactMaps) override
+			{
+				for (int32 FromTID = 0, NumTID = CompactMaps.NumTriangleMappings(); FromTID < NumTID; FromTID++)
+				{
+					const int32 ToTID = CompactMaps.GetTriangleMapping(FromTID);
+					if (ToTID == UE::Geometry::FCompactMaps::InvalidID)
+					{
+						continue;
+					}
+					if (ensure(ToTID <= FromTID))
+					{
+						SourceTriangleIds[ToTID] = SourceTriangleIds[FromTID];
+					}
+				}
+				SourceTriangleIds.Resize(Parent->MaxTriangleID());
+			}
+
+			void Reparent(UE::Geometry::FDynamicMesh3* NewParent) override
+			{
+				Parent = NewParent;
+			}
+
+			bool CopyThroughMapping(const UE::Geometry::TDynamicAttributeBase<UE::Geometry::FDynamicMesh3>* Source, const UE::Geometry::FMeshIndexMappings& Mapping) override
+			{
+				for (const TPair<int32, int32>& MapTID : Mapping.GetTriangleMap().GetForwardMap())
+				{
+					int32 SourceTId;
+					if (!ensure(Source->CopyOut(MapTID.Key, &SourceTId, sizeof(SourceTId))))
+					{
+						return false;
+					}
+					SetValue(MapTID.Value, SourceTId);
+				}
+				return true;
+			}
+
+			virtual bool Append(const TDynamicAttributeBase& Source, const UE::Geometry::FDynamicMesh3::FAppendInfo& Info) override
+			{
+				int32 NewMaxID = Info.NumTriangle + Info.TriangleOffset;
+				if (NewMaxID > SourceTriangleIds.Num())
+				{
+					SourceTriangleIds.SetNum(NewMaxID);
+				}
+				for (int32 Idx = 0; Idx < Info.NumTriangle; ++Idx)
+				{
+					int32 TargetID = Idx + Info.TriangleOffset;
+					if (!Parent->IsTriangle(TargetID))
+					{
+						continue;
+					}
+					int32 SourceTID;
+					if (!ensure(Source.CopyOut(Idx, &SourceTID, sizeof(SourceTID))))
+					{
+						return false;
+					}
+					SetValue(Idx + Info.TriangleOffset, SourceTID);
+				}
+				return true;
+			}
+
+			virtual void AppendDefaulted(const UE::Geometry::FDynamicMesh3::FAppendInfo& Info) override
+			{
+				int32 NewMaxID = Info.NumTriangle + Info.TriangleOffset;
+				if (NewMaxID > SourceTriangleIds.Num())
+				{
+					SourceTriangleIds.Resize(NewMaxID, INDEX_NONE);
+				}
+			}
+
+			bool CopyOut(int RawID, void* Buffer, int BufferSize) const override
+			{
+				if (BufferSize != sizeof(int32) || !Parent->IsTriangle(RawID))
+				{
+					return false;
+				}
+
+				*static_cast<int32*>(Buffer) = SourceTriangleIds[RawID];
+				return true;
+			}
+
+			bool CopyIn(int RawID, void* Buffer, int BufferSize) override
+			{
+				if (BufferSize != sizeof(int32) || !Parent->IsTriangle(RawID))
+				{
+					return false;
+				}
+
+				SourceTriangleIds[RawID] = *static_cast<int32*>(Buffer);
+				return true;
+			}
+
+			TUniquePtr<UE::Geometry::TDynamicAttributeChangeBase<UE::Geometry::FDynamicMesh3>> NewBlankChange() const override
+			{
+				return MakeUnique<FDynamicMeshSourceTriangleIdAttributeChange>();
+			}
+
+			void Initialize()
+			{
+				SourceTriangleIds.Resize(Parent->MaxTriangleID());
+				SourceTriangleIds.Fill(INDEX_NONE);
+			}
+
+			void InitializeFromArray(TConstArrayView<int32> TriangleToNonManifoldVertexIDMap)
+			{
+				check(TriangleToNonManifoldVertexIDMap.Num() == Parent->MaxTriangleID());
+				SourceTriangleIds.Resize(Parent->MaxTriangleID());
+
+				for (int32 TriangleID : Parent->TriangleIndicesItr())
+				{
+					SetValue(TriangleID, TriangleToNonManifoldVertexIDMap[TriangleID]);
+				}
+				bValid = true;
+			}
+
+			void Copy(const FDynamicMeshSourceTriangleIdAttribute& Copy)
+			{
+				UE::Geometry::FDynamicMeshAttributeBase::CopyParentClassData(Copy);
+				SourceTriangleIds = Copy.SourceTriangleIds;
+				bValid = Copy.bValid;
+			}
+
+			int32 GetValue(int32 InTriangleId) const
+			{
+				return SourceTriangleIds[InTriangleId];
+			}
+
+			void SetValue(int32 InTriangleId, int32 InSourceTriangled)
+			{
+				SourceTriangleIds[InTriangleId] = InSourceTriangled;
+			}
+
+			// Any topo operation on the mesh will invalidate the non-manifold information.
+			virtual void OnSplitEdge(const DynamicMeshInfo::FEdgeSplitInfo&) override
+			{
+				bValid = false;
+			}
+
+			virtual void OnFlipEdge(const DynamicMeshInfo::FEdgeFlipInfo& FlipInfo) override
+			{
+				bValid = false;
+			}
+
+			virtual void OnCollapseEdge(const DynamicMeshInfo::FEdgeCollapseInfo&) override
+			{
+				bValid = false;
+			}
+
+			virtual void OnPokeTriangle(const DynamicMeshInfo::FPokeTriangleInfo&) override
+			{
+				bValid = false;
+			}
+
+			virtual void OnMergeEdges(const DynamicMeshInfo::FMergeEdgesInfo&) override
+			{
+				bValid = false;
+			}
+
+			virtual void OnMergeVertices(const DynamicMeshInfo::FMergeVerticesInfo&) override
+			{
+				bValid = false;
+			}
+
+			virtual void OnSplitVertex(const DynamicMeshInfo::FVertexSplitInfo&, const TArrayView<const int>&) override
+			{
+				bValid = false;
+			}
+
+			virtual void OnNewVertex(int VertexID, bool bInserted) override
+			{
+				bValid = false;
+			}
+
+			virtual void OnRemoveVertex(int VertexID) override
+			{
+				bValid = false;
+			}
+		};
+
+		void FDynamicMeshSourceTriangleIdAttributeChange::SaveInitialTriangle(const UE::Geometry::FDynamicMeshAttributeBase* Attribute, int TriangleID)
+		{
+			const FDynamicMeshSourceTriangleIdAttribute* const SourceAttribute = static_cast<const FDynamicMeshSourceTriangleIdAttribute*>(Attribute);
+			if (OldSourceTriangleIds.IsEmpty())
+			{
+				bOldInvalidState = SourceAttribute->IsValid();
+			}
+
+			OldSourceTriangleIds.Emplace(TriangleID, SourceAttribute->GetValue(TriangleID));
+		}
+
+		void FDynamicMeshSourceTriangleIdAttributeChange::StoreAllFinalTriangles(const UE::Geometry::FDynamicMeshAttributeBase* Attribute, const TArray<int>& TriangleIDs)
+		{
+			NewSourceTriangleIds.Reserve(NewSourceTriangleIds.Num() + TriangleIDs.Num());
+			const FDynamicMeshSourceTriangleIdAttribute* const SourceAttribute = static_cast<const FDynamicMeshSourceTriangleIdAttribute*>(Attribute);
+			for (const int32 TriangleID : TriangleIDs)
+			{
+				NewSourceTriangleIds.Emplace(TriangleID, SourceAttribute->GetValue(TriangleID));
+			}
+
+			// Store the last known valid state.  
+			bNewInvalidState = SourceAttribute->IsValid();
+		}
+
+		bool FDynamicMeshSourceTriangleIdAttributeChange::Apply(UE::Geometry::FDynamicMeshAttributeBase* Attribute, bool bRevert) const
+		{
+			const TArray<TPair<int32, int32>>& Changes = bRevert ? OldSourceTriangleIds : NewSourceTriangleIds;
+			FDynamicMeshSourceTriangleIdAttribute* const SourceAttribute = static_cast<FDynamicMeshSourceTriangleIdAttribute*>(Attribute);
+			for (const TPair<int32, int32>& Item : Changes)
+			{
+				if (ensure(SourceAttribute->GetParent()->IsTriangle(Item.Key)))
+				{
+					SourceAttribute->SetValue(Item.Key, Item.Value);
+				}
+			}
+
+			// Restore the valid state as well.
+			SourceAttribute->bValid = bRevert ? bOldInvalidState : bNewInvalidState;
+
+			return true;
+		}
+	}  // namepsace Private
+
+	const FName FClothCollectionToDynamicMesh::FMappingSupport::ClothMeshTIDsAttrName = FName(TEXT("ClothMeshTIDsAttr"));
+
+	FClothCollectionToDynamicMesh::FMappingSupport::FMappingSupport(const UE::Geometry::FDynamicMesh3& Mesh)
+		: UE::Geometry::FNonManifoldMappingSupport(Mesh)
+	{
+		Reset(Mesh);
+	}
+
+	void FClothCollectionToDynamicMesh::FMappingSupport::Reset(const UE::Geometry::FDynamicMesh3& Mesh)
+	{
+		UE::Geometry::FNonManifoldMappingSupport::Reset(Mesh);
+
+		if (const UE::Geometry::FDynamicMeshAttributeSet* const Attributes = DynamicMesh->Attributes())
+		{
+			if (Attributes->HasAttachedAttribute(ClothMeshTIDsAttrName))
+			{
+				SrcTIDsAttribute = static_cast<const Private::FDynamicMeshSourceTriangleIdAttribute*>(Attributes->GetAttachedAttribute(ClothMeshTIDsAttrName));
+			}
+		}
+	}
+
+	bool FClothCollectionToDynamicMesh::FMappingSupport::IsMappedTriangleInSource() const
+	{
+		return SrcTIDsAttribute != nullptr && SrcTIDsAttribute->IsValid();
+	}
+
+	int32 FClothCollectionToDynamicMesh::FMappingSupport::GetOriginalTriangleID(const int32 Tid) const
+	{
+		checkSlow(DynamicMesh->IsTriangle(Tid));
+
+		int32 Result = Tid;
+		if (SrcTIDsAttribute)
+		{
+			Result = SrcTIDsAttribute->GetValue(Tid);
+		}
+		return Result;
+	}
+
+	bool FClothCollectionToDynamicMesh::FMappingSupport::AttachTriangleMappingData(const TArray<int32>& TriangleToOriginalVertexIDMap, UE::Geometry::FDynamicMesh3& InOutMesh)
+	{
+		const int32 MaxTID = InOutMesh.MaxTriangleID();
+		UE::Geometry::FDynamicMeshAttributeSet* const Attributes = InOutMesh.Attributes();
+
+		if (TriangleToOriginalVertexIDMap.Num() < MaxTID || !Attributes)
+		{
+			return false;
+		}
+
+		// create and attach a triangle ID buffer, removing any pre-exising one.
+		Private::FDynamicMeshSourceTriangleIdAttribute* TIDAttr = [&Attributes, &InOutMesh]
+			{
+				// replace existing buffer.  we *should* be able to re-use it
+				if (Attributes->HasAttachedAttribute(ClothMeshTIDsAttrName))
+				{
+					Attributes->RemoveAttribute(ClothMeshTIDsAttrName);
+				}
+				Private::FDynamicMeshSourceTriangleIdAttribute* const SrcMeshTIDAttr = new Private::FDynamicMeshSourceTriangleIdAttribute(&InOutMesh);
+				SrcMeshTIDAttr->SetName(ClothMeshTIDsAttrName);
+				InOutMesh.Attributes()->AttachAttribute(ClothMeshTIDsAttrName, SrcMeshTIDAttr);
+				return SrcMeshTIDAttr;
+			}();
+
+		TIDAttr->InitializeFromArray(TriangleToOriginalVertexIDMap);
+		return true;
+	}
+
+	void FClothCollectionToDynamicMesh::FMappingSupport::RemoveTriangleMappingData(UE::Geometry::FDynamicMesh3& InOutMesh)
+	{
+		UE::Geometry::FDynamicMeshAttributeSet* const Attributes = InOutMesh.Attributes();
+
+		if (Attributes && Attributes->HasAttachedAttribute(ClothMeshTIDsAttrName))
+		{
+			Attributes->RemoveAttribute(ClothMeshTIDsAttrName);
+		}
+	}
+}  // namespace UE::Chaos::ClothAsset
+
+#endif  // #if WITH_EDITOR
diff --git a/Engine/Plugins/ChaosClothAsset/Source/ChaosClothAssetEngine/Private/ChaosClothAsset/ClothEngineTools.cpp b/Engine/Plugins/ChaosClothAsset/Source/ChaosClothAssetEngine/Private/ChaosClothAsset/ClothEngineTools.cpp
--- a/Engine/Plugins/ChaosClothAsset/Source/ChaosClothAssetEngine/Private/ChaosClothAsset/ClothEngineTools.cpp
+++ b/Engine/Plugins/ChaosClothAsset/Source/ChaosClothAssetEngine/Private/ChaosClothAsset/ClothEngineTools.cpp
@@ -1,125 +1,149 @@
 // Copyright Epic Games, Inc. All Rights Reserved.
 
 #include "ChaosClothAsset/ClothEngineTools.h"
+#if WITH_EDITOR
+#include "ChaosClothAsset/ClothCollectionToDynamicMesh.h"
+#endif
 #include "ChaosClothAsset/CollectionClothFacade.h"
 #include "ChaosClothAsset/CollectionClothSelectionFacade.h"
 #include "ChaosClothAsset/ClothCollectionGroup.h"
 #include "ChaosClothAsset/ClothGeometryTools.h"
 #include "ClothTetherData.h"
+#if WITH_EDITOR
+#include "DynamicMesh/DynamicMesh3.h"
+#include "DynamicMesh/DynamicMeshAttributeSet.h"
+#endif
 
 namespace UE::Chaos::ClothAsset
 {
-
-namespace Private
-{
-static void AppendTetherData(FCollectionClothFacade& ClothFacade, const FClothTetherData& TetherData)
-{
-	// Append new tethers
-	TArrayView<TArray<int32>> TetherKinematicIndex = ClothFacade.GetTetherKinematicIndex();
-	TArrayView<TArray<float>> TetherReferenceLength = ClothFacade.GetTetherReferenceLength();
-	for (const TArray<TTuple<int32, int32, float>>& TetherBatch : TetherData.Tethers)
+	namespace Private
 	{
-		for (const TTuple<int32, int32, float>& Tether : TetherBatch)
+		static void AppendTetherData(FCollectionClothFacade& ClothFacade, const FClothTetherData& TetherData)
 		{
-			// Tuple is Kinematic, Dynamic, RefLength
-			const int32 DynamicIndex = Tether.Get<1>();
-			TArray<int32>& KinematicIndex = TetherKinematicIndex[DynamicIndex];
-			TArray<float>& ReferenceLength = TetherReferenceLength[DynamicIndex];
-			check(KinematicIndex.Num() == ReferenceLength.Num());
-			checkSlow(KinematicIndex.Find(Tether.Get<0>()) == INDEX_NONE);
-			KinematicIndex.Add(Tether.Get<0>());
-			ReferenceLength.Add(Tether.Get<2>());
+			// Append new tethers
+			TArrayView<TArray<int32>> TetherKinematicIndex = ClothFacade.GetTetherKinematicIndex();
+			TArrayView<TArray<float>> TetherReferenceLength = ClothFacade.GetTetherReferenceLength();
+			for (const TArray<TTuple<int32, int32, float>>& TetherBatch : TetherData.Tethers)
+			{
+				for (const TTuple<int32, int32, float>& Tether : TetherBatch)
+				{
+					// Tuple is Kinematic, Dynamic, RefLength
+					const int32 DynamicIndex = Tether.Get<1>();
+					TArray<int32>& KinematicIndex = TetherKinematicIndex[DynamicIndex];
+					TArray<float>& ReferenceLength = TetherReferenceLength[DynamicIndex];
+					check(KinematicIndex.Num() == ReferenceLength.Num());
+					checkSlow(KinematicIndex.Find(Tether.Get<0>()) == INDEX_NONE);
+					KinematicIndex.Add(Tether.Get<0>());
+					ReferenceLength.Add(Tether.Get<2>());
+				}
+			}
 		}
 	}
-}
-}
 
-void FClothEngineTools::GenerateTethers(const TSharedRef<FManagedArrayCollection>& ClothCollection, const FName& WeightMapName, const bool bGenerateGeodesicTethers, const FVector2f& MaxDistanceValue)
-{
-	FCollectionClothFacade ClothFacade(ClothCollection);
-	FClothGeometryTools::DeleteTethers(ClothCollection);
-	if (ClothFacade.HasWeightMap(WeightMapName))
+	void FClothEngineTools::GenerateTethers(const TSharedRef<FManagedArrayCollection>& ClothCollection, const FName& WeightMapName, const bool bGenerateGeodesicTethers, const FVector2f& MaxDistanceValue)
 	{
-		FClothTetherData TetherData;
-		TArray<uint32> SimIndices;
-		SimIndices.Reserve(ClothFacade.GetNumSimFaces() * 3);
-		for (const FIntVector3& Face : ClothFacade.GetSimIndices3D())
+		FCollectionClothFacade ClothFacade(ClothCollection);
+		FClothGeometryTools::DeleteTethers(ClothCollection);
+		if (ClothFacade.HasWeightMap(WeightMapName))
 		{
-			SimIndices.Add(Face[0]);
-			SimIndices.Add(Face[1]);
-			SimIndices.Add(Face[2]);
-		}
+			FClothTetherData TetherData;
+			TArray<uint32> SimIndices;
+			SimIndices.Reserve(ClothFacade.GetNumSimFaces() * 3);
+			for (const FIntVector3& Face : ClothFacade.GetSimIndices3D())
+			{
+				SimIndices.Add(Face[0]);
+				SimIndices.Add(Face[1]);
+				SimIndices.Add(Face[2]);
+			}
 
-		if (MaxDistanceValue.Equals(FVector2f(0.f, 1.f)))
-		{
-			TetherData.GenerateTethers(ClothFacade.GetSimPosition3D(), TConstArrayView<uint32>(SimIndices), ClothFacade.GetWeightMap(WeightMapName), bGenerateGeodesicTethers);
-		}
-		else
-		{
+			if (MaxDistanceValue.Equals(FVector2f(0.f, 1.f)))
+			{
+				TetherData.GenerateTethers(ClothFacade.GetSimPosition3D(), TConstArrayView<uint32>(SimIndices), ClothFacade.GetWeightMap(WeightMapName), bGenerateGeodesicTethers);
+			}
+			else
+			{
 
-			const TSet<int32> KinematicVertices = FClothGeometryTools::GenerateKinematicVertices3D(ClothCollection, WeightMapName, MaxDistanceValue, NAME_None);
-			TetherData.GenerateTethers(ClothFacade.GetSimPosition3D(), TConstArrayView<uint32>(SimIndices), KinematicVertices, bGenerateGeodesicTethers);
-		}
+				const TSet<int32> KinematicVertices = FClothGeometryTools::GenerateKinematicVertices3D(ClothCollection, WeightMapName, MaxDistanceValue, NAME_None);
+				TetherData.GenerateTethers(ClothFacade.GetSimPosition3D(), TConstArrayView<uint32>(SimIndices), KinematicVertices, bGenerateGeodesicTethers);
+			}
 
-		Private::AppendTetherData(ClothFacade, TetherData);
+			Private::AppendTetherData(ClothFacade, TetherData);
+		}
 	}
-}
-
-void FClothEngineTools::GenerateTethersFromSelectionSet(const TSharedRef<FManagedArrayCollection>& ClothCollection, const FName& FixedEndSet, const bool bGeodesicTethers)
-{
-	FClothGeometryTools::DeleteTethers(ClothCollection);
-	FCollectionClothSelectionConstFacade SelectionFacade(ClothCollection);
 
-	if (SelectionFacade.HasSelection(FixedEndSet) && SelectionFacade.GetSelectionGroup(FixedEndSet) == UE::Chaos::ClothAsset::ClothCollectionGroup::SimVertices3D)
+	void FClothEngineTools::GenerateTethersFromSelectionSet(const TSharedRef<FManagedArrayCollection>& ClothCollection, const FName& FixedEndSet, const bool bGeodesicTethers)
 	{
-		FCollectionClothFacade ClothFacade(ClothCollection);
-		FClothTetherData TetherData;
-		TArray<uint32> SimIndices;
-		SimIndices.Reserve(ClothFacade.GetNumSimFaces() * 3);
-		for (const FIntVector3& Face : ClothFacade.GetSimIndices3D())
+		FClothGeometryTools::DeleteTethers(ClothCollection);
+		FCollectionClothSelectionConstFacade SelectionFacade(ClothCollection);
+
+		if (SelectionFacade.HasSelection(FixedEndSet) && SelectionFacade.GetSelectionGroup(FixedEndSet) == UE::Chaos::ClothAsset::ClothCollectionGroup::SimVertices3D)
 		{
-			SimIndices.Add(Face[0]);
-			SimIndices.Add(Face[1]);
-			SimIndices.Add(Face[2]);
-		}
+			FCollectionClothFacade ClothFacade(ClothCollection);
+			FClothTetherData TetherData;
+			TArray<uint32> SimIndices;
+			SimIndices.Reserve(ClothFacade.GetNumSimFaces() * 3);
+			for (const FIntVector3& Face : ClothFacade.GetSimIndices3D())
+			{
+				SimIndices.Add(Face[0]);
+				SimIndices.Add(Face[1]);
+				SimIndices.Add(Face[2]);
+			}
 
-		TetherData.GenerateTethers(ClothFacade.GetSimPosition3D(), TConstArrayView<uint32>(SimIndices), SelectionFacade.GetSelectionSet(FixedEndSet), bGeodesicTethers);
-		Private::AppendTetherData(ClothFacade, TetherData);
+			TetherData.GenerateTethers(ClothFacade.GetSimPosition3D(), TConstArrayView<uint32>(SimIndices), SelectionFacade.GetSelectionSet(FixedEndSet), bGeodesicTethers);
+			Private::AppendTetherData(ClothFacade, TetherData);
+		}
 	}
-}
 
-void FClothEngineTools::GenerateTethersFromCustomSelectionSets(const TSharedRef<FManagedArrayCollection>& ClothCollection, const FName& InFixedEndSet, const TArray<TPair<FName, FName>>& CustomTetherEndSets, const bool bGeodesicTethers)
-{
-	FClothGeometryTools::DeleteTethers(ClothCollection);
-	FCollectionClothSelectionConstFacade SelectionFacade(ClothCollection);
-	if (SelectionFacade.HasSelection(InFixedEndSet) && SelectionFacade.GetSelectionGroup(InFixedEndSet) == UE::Chaos::ClothAsset::ClothCollectionGroup::SimVertices3D)
+	void FClothEngineTools::GenerateTethersFromCustomSelectionSets(const TSharedRef<FManagedArrayCollection>& ClothCollection, const FName& InFixedEndSet, const TArray<TPair<FName, FName>>& CustomTetherEndSets, const bool bGeodesicTethers)
 	{
-
-		FCollectionClothFacade ClothFacade(ClothCollection);
-		TArray<uint32> SimIndices;
-		SimIndices.Reserve(ClothFacade.GetNumSimFaces() * 3);
-		for (const FIntVector3& Face : ClothFacade.GetSimIndices3D())
+		FClothGeometryTools::DeleteTethers(ClothCollection);
+		FCollectionClothSelectionConstFacade SelectionFacade(ClothCollection);
+		if (SelectionFacade.HasSelection(InFixedEndSet) && SelectionFacade.GetSelectionGroup(InFixedEndSet) == UE::Chaos::ClothAsset::ClothCollectionGroup::SimVertices3D)
 		{
-			SimIndices.Add(Face[0]);
-			SimIndices.Add(Face[1]);
-			SimIndices.Add(Face[2]);
-		}
 
-		const TSet<int32>& FixedEndSet = SelectionFacade.GetSelectionSet(InFixedEndSet);
+			FCollectionClothFacade ClothFacade(ClothCollection);
+			TArray<uint32> SimIndices;
+			SimIndices.Reserve(ClothFacade.GetNumSimFaces() * 3);
+			for (const FIntVector3& Face : ClothFacade.GetSimIndices3D())
+			{
+				SimIndices.Add(Face[0]);
+				SimIndices.Add(Face[1]);
+				SimIndices.Add(Face[2]);
+			}
 
-		for (const TPair<FName, FName>& TetherEnds : CustomTetherEndSets)
-		{
-			const FName& CustomDynamicEndSet = TetherEnds.Get<0>();
-			const FName& CustomFixedEndSet = TetherEnds.Get<1>();
-			if (SelectionFacade.HasSelection(CustomFixedEndSet) && SelectionFacade.GetSelectionGroup(CustomFixedEndSet) == UE::Chaos::ClothAsset::ClothCollectionGroup::SimVertices3D &&
-				SelectionFacade.HasSelection(CustomDynamicEndSet) && SelectionFacade.GetSelectionGroup(CustomDynamicEndSet) == UE::Chaos::ClothAsset::ClothCollectionGroup::SimVertices3D)
+			const TSet<int32>& FixedEndSet = SelectionFacade.GetSelectionSet(InFixedEndSet);
+
+			for (const TPair<FName, FName>& TetherEnds : CustomTetherEndSets)
 			{
-				FClothTetherData TetherData;
-				TetherData.GenerateTethers(ClothFacade.GetSimPosition3D(), TConstArrayView<uint32>(SimIndices), FixedEndSet, SelectionFacade.GetSelectionSet(CustomDynamicEndSet), SelectionFacade.GetSelectionSet(CustomFixedEndSet), bGeodesicTethers);
-				Private::AppendTetherData(ClothFacade, TetherData);
+				const FName& CustomDynamicEndSet = TetherEnds.Get<0>();
+				const FName& CustomFixedEndSet = TetherEnds.Get<1>();
+				if (SelectionFacade.HasSelection(CustomFixedEndSet) && SelectionFacade.GetSelectionGroup(CustomFixedEndSet) == UE::Chaos::ClothAsset::ClothCollectionGroup::SimVertices3D &&
+					SelectionFacade.HasSelection(CustomDynamicEndSet) && SelectionFacade.GetSelectionGroup(CustomDynamicEndSet) == UE::Chaos::ClothAsset::ClothCollectionGroup::SimVertices3D)
+				{
+					FClothTetherData TetherData;
+					TetherData.GenerateTethers(ClothFacade.GetSimPosition3D(), TConstArrayView<uint32>(SimIndices), FixedEndSet, SelectionFacade.GetSelectionSet(CustomDynamicEndSet), SelectionFacade.GetSelectionSet(CustomFixedEndSet), bGeodesicTethers);
+					Private::AppendTetherData(ClothFacade, TetherData);
+				}
 			}
 		}
 	}
-}
 
+#if WITH_EDITOR
+	void FClothEngineTools::BuildDynamicMeshFromSimMesh(UE::Geometry::FDynamicMesh3& DynamicMesh, const TSharedRef<const FManagedArrayCollection>& ClothCollection)
+	{
+		constexpr bool bDisableAttributes = false;
+		const int32 PatternIndex = INDEX_NONE;
+		DynamicMesh.EnableAttributes();
+		FClothCollectionToDynamicMesh::Convert(ClothCollection, PatternIndex, FClothCollectionToDynamicMesh::ESourceMeshType::Sim3D, DynamicMesh, bDisableAttributes);
+	}
+
+	void FClothEngineTools::BuildDynamicMeshFromRenderMesh(UE::Geometry::FDynamicMesh3& DynamicMesh, const TSharedRef<const FManagedArrayCollection>& ClothCollection, int32 MaterialOffset)
+	{
+		constexpr bool bDisableAttributes = false;
+		const int32 PatternIndex = INDEX_NONE;
+		DynamicMesh.EnableAttributes();
+		DynamicMesh.Attributes()->EnableMaterialID();
+		FClothCollectionToDynamicMesh::Convert(ClothCollection, PatternIndex, FClothCollectionToDynamicMesh::ESourceMeshType::Render, DynamicMesh, bDisableAttributes, MaterialOffset);
+	}
+#endif  // #if WITH_EDITOR
 }  // End namespace UE::Chaos::ClothAsset
diff --git a/Engine/Plugins/ChaosClothAsset/Source/ChaosClothAssetEngine/Public/ChaosClothAsset/ClothAssetBase.h b/Engine/Plugins/ChaosClothAsset/Source/ChaosClothAssetEngine/Public/ChaosClothAsset/ClothAssetBase.h
--- a/Engine/Plugins/ChaosClothAsset/Source/ChaosClothAssetEngine/Public/ChaosClothAsset/ClothAssetBase.h
+++ b/Engine/Plugins/ChaosClothAsset/Source/ChaosClothAssetEngine/Public/ChaosClothAsset/ClothAssetBase.h
@@ -260,6 +260,16 @@ public:
 
 	void SetHasVertexColors(bool InbHasVertexColors);
 
+#if WITH_EDITOR
+	/**
+	 * Export the graphical representation of this asset to a SkeletalMesh asset.
+	 * Includes skinning, but excludes all clothing simulation data since it isn't compatible with the Cloth Asset.
+	 * @param SkeletalMesh The empty skeletal mesh asset to be written to.
+	 * @return Whether the initialization of the skeletal mesh has succeeded or not.
+	 */
+	bool ExportToSkeletalMesh(USkeletalMesh& SkeletalMesh) const;
+#endif
+
 protected:
 	/** Reregister all components using this asset to reset the simulation in case anything has changed. */
 	void ReregisterComponents();
diff --git a/Engine/Plugins/ChaosClothAsset/Source/ChaosClothAssetEngine/Public/ChaosClothAsset/ClothEngineTools.h b/Engine/Plugins/ChaosClothAsset/Source/ChaosClothAssetEngine/Public/ChaosClothAsset/ClothEngineTools.h
--- a/Engine/Plugins/ChaosClothAsset/Source/ChaosClothAssetEngine/Public/ChaosClothAsset/ClothEngineTools.h
+++ b/Engine/Plugins/ChaosClothAsset/Source/ChaosClothAssetEngine/Public/ChaosClothAsset/ClothEngineTools.h
@@ -8,6 +8,9 @@
 struct FManagedArrayCollection;
 class FName;
 class FSkeletalMeshLODModel;
+#if WITH_EDITOR
+namespace UE::Geometry { class FDynamicMesh3; }
+#endif
 
 namespace UE::Chaos::ClothAsset
 {
@@ -21,5 +24,16 @@ namespace UE::Chaos::ClothAsset
 		static void GenerateTethersFromSelectionSet(const TSharedRef<FManagedArrayCollection>& ClothCollection, const FName& FixedEndSet, const bool bGeodesicTethers);
 		/** @param CustomTetherEndSets: First element of each pair is DynamicSet, second is FixedSet */
 		static void GenerateTethersFromCustomSelectionSets(const TSharedRef<FManagedArrayCollection>& ClothCollection, const FName& FixedEndSet, const TArray<TPair<FName, FName>>& CustomTetherEndSets, const bool bGeodesicTethers);
+
+#if WITH_EDITOR
+		/** Build a DynamicMesh from a sim mesh. */
+		static void BuildDynamicMeshFromSimMesh(UE::Geometry::FDynamicMesh3& DynamicMesh, const TSharedRef<const FManagedArrayCollection>& ClothCollection);
+
+		/**
+		 * Build a DynamicMesh from a render mesh.
+		 * Use the material offset to add to the final material group id if the material list doesn't start at 0.
+		 */
+		static void BuildDynamicMeshFromRenderMesh(UE::Geometry::FDynamicMesh3& DynamicMesh, const TSharedRef<const FManagedArrayCollection>& ClothCollection, int32 MaterialOffset = 0);
+#endif  // #if WITH_EDITOR
 	};
 }  // End namespace UE::Chaos::ClothAsset
\ No newline at end of file
diff --git a/Engine/Plugins/Experimental/ChaosOutfitAsset/Source/ChaosOutfitAssetDataflowNodes/ChaosOutfitAssetDataflowNodes.Build.cs b/Engine/Plugins/Experimental/ChaosOutfitAsset/Source/ChaosOutfitAssetDataflowNodes/ChaosOutfitAssetDataflowNodes.Build.cs
--- a/Engine/Plugins/Experimental/ChaosOutfitAsset/Source/ChaosOutfitAssetDataflowNodes/ChaosOutfitAssetDataflowNodes.Build.cs
+++ b/Engine/Plugins/Experimental/ChaosOutfitAsset/Source/ChaosOutfitAssetDataflowNodes/ChaosOutfitAssetDataflowNodes.Build.cs
@@ -14,12 +14,18 @@ public class ChaosOutfitAssetDataflowNodes : ModuleRules
 				"Core",
 				"CoreUObject",
 				"Chaos",
+				"ChaosClothAsset",
 				"ChaosClothAssetEngine",
 				"ChaosOutfitAssetEngine",
 				"DataflowCore",
 				"DataflowNodes",
 				"Engine",
+				"GeometryCore",
+				"MeshConversion",
+				"MeshDescription",
 				"MeshResizingCore",
+				"SkeletalMeshDescription",
+				"SkeletalMeshUtilitiesCommon",
 			}
 		);
 	}
diff --git a/Engine/Plugins/Experimental/ChaosOutfitAsset/Source/ChaosOutfitAssetDataflowNodes/Private/ChaosOutfitAsset/OutfitAssetTerminalNode.cpp b/Engine/Plugins/Experimental/ChaosOutfitAsset/Source/ChaosOutfitAssetDataflowNodes/Private/ChaosOutfitAsset/OutfitAssetTerminalNode.cpp
--- a/Engine/Plugins/Experimental/ChaosOutfitAsset/Source/ChaosOutfitAssetDataflowNodes/Private/ChaosOutfitAsset/OutfitAssetTerminalNode.cpp
+++ b/Engine/Plugins/Experimental/ChaosOutfitAsset/Source/ChaosOutfitAssetDataflowNodes/Private/ChaosOutfitAsset/OutfitAssetTerminalNode.cpp
@@ -1,22 +1,107 @@
 // Copyright Epic Games, Inc. All Rights Reserved.
 
 #include "ChaosOutfitAsset/OutfitAssetTerminalNode.h"
+#if WITH_EDITOR
+#include "AssetRegistry/AssetRegistryModule.h"
+#endif
 #include "ChaosOutfitAsset/OutfitAsset.h"
 #include "Dataflow/DataflowInputOutput.h"
+#if WITH_EDITOR
+#include "Dataflow/DataflowObjectInterface.h"
+#include "Engine/SkeletalMesh.h"
+#include "Misc/FileHelper.h"
+#include "Misc/ScopedSlowTask.h"
+#include "ContentBrowserModule.h"
+#include "IContentBrowserSingleton.h"
+#endif
 
 #include UE_INLINE_GENERATED_CPP_BY_NAME(OutfitAssetTerminalNode)
 
 #define LOCTEXT_NAMESPACE "ChaosOutfitAssetTerminalNode"
 
+namespace UE::Chaos::OutfitAsset::Private
+{
+#if WITH_EDITOR
+	static void ExportToSkeletalMesh(const TObjectPtr<const UChaosOutfitAsset>& OutfitAsset)
+	{
+		using namespace UE::Geometry;
+		if (OutfitAsset)
+		{
+			// Make a name from the outfit asset name
+			const FString SkeletalMeshPath = FPackageName::GetLongPackagePath(OutfitAsset->GetOutermost()->GetName());
+			const FString OutfitAssetName = OutfitAsset->GetName();
+			FString SkeletalMeshName = FString(TEXT("SK_")) + (OutfitAssetName.StartsWith(TEXT("OA_")) ? OutfitAssetName.RightChop(3) : OutfitAssetName);
+			FString SkeletalMeshPackageName = FPaths::Combine(SkeletalMeshPath, SkeletalMeshName);
+
+			// Open a Save As dialog
+			FSaveAssetDialogConfig SaveAssetDialogConfig;
+			SaveAssetDialogConfig.DefaultPath = SkeletalMeshPath;
+			SaveAssetDialogConfig.DefaultAssetName = SkeletalMeshName;
+			SaveAssetDialogConfig.AssetClassNames.Add(USkeletalMesh::StaticClass()->GetClassPathName());
+			SaveAssetDialogConfig.ExistingAssetPolicy = ESaveAssetDialogExistingAssetPolicy::AllowButWarn;
+			SaveAssetDialogConfig.DialogTitleOverride = LOCTEXT("SaveAssetDialogConfig", "Convert Outfit Asset To Skeletal Mesh As");
+
+			const FContentBrowserModule& ContentBrowserModule = FModuleManager::LoadModuleChecked<FContentBrowserModule>("ContentBrowser");
+			const FString AssetSavePath = ContentBrowserModule.Get().CreateModalSaveAssetDialog(SaveAssetDialogConfig);
+			const FString AssetSavePackageName = FPackageName::ObjectPathToPackageName(AssetSavePath);
+			FText OutError;
+			if (AssetSavePath.IsEmpty() || !FFileHelper::IsFilenameValidForSaving(AssetSavePackageName, OutError))
+			{
+				return;
+			}
+			SkeletalMeshPackageName = AssetSavePackageName;
+			SkeletalMeshName = FPackageName::GetLongPackageAssetName(SkeletalMeshPackageName);
+
+			constexpr float NumSteps = 0.f;
+			FScopedSlowTask SlowTask(NumSteps, LOCTEXT("ExportToSkeletalMesh", "Exporting Cloth Outfit to a Skeletal Mesh asset..."));
+			SlowTask.MakeDialog();
+
+			UPackage* const SkeletalMeshPackage = CreatePackage(*SkeletalMeshPackageName);
+			USkeletalMesh* const SkeletalMesh =
+				NewObject<USkeletalMesh>(
+					SkeletalMeshPackage,
+					USkeletalMesh::StaticClass(),
+					FName(SkeletalMeshName),
+					RF_Public | RF_Standalone | RF_Transactional);
+
+			// Export to skeletal mesh
+			if (SkeletalMesh)
+			{
+				OutfitAsset->ExportToSkeletalMesh(*SkeletalMesh);
+			}
+
+			SkeletalMesh->MarkPackageDirty();
+
+			// Notify the asset registry
+			FAssetRegistryModule::AssetCreated(SkeletalMesh);
+		}
+	}
+#endif  // #if WITH_EDITOR
+}
+
 FChaosOutfitAssetTerminalNode::FChaosOutfitAssetTerminalNode(const UE::Dataflow::FNodeParameters& InParam, FGuid InGuid)
 	: FDataflowTerminalNode(InParam, InGuid)
+#if WITH_EDITOR
+	, ConvertToSkeletalMesh(
+		FDataflowFunctionProperty::FDelegate::CreateLambda([this](UE::Dataflow::FContext& Context)
+			{
+				if (const UE::Dataflow::FEngineContext* EngineContext = Context.AsType<UE::Dataflow::FEngineContext>())
+				{
+					if (const UChaosOutfitAsset* const OwnerOutfitAsset = Cast<UChaosOutfitAsset>(EngineContext->Owner))
+					{
+						UE::Chaos::OutfitAsset::Private::ExportToSkeletalMesh(OwnerOutfitAsset);
+					}
+				}
+			}))
+#endif
 {
 	RegisterInputConnection(&Outfit);
 }
 
 void FChaosOutfitAssetTerminalNode::SetAssetValue(TObjectPtr<UObject> Asset, UE::Dataflow::FContext& Context) const
 {
-	if (UChaosOutfitAsset* const OutfitAsset = Cast<UChaosOutfitAsset>(Asset.Get()))
+	TObjectPtr<UChaosOutfitAsset> OutfitAsset = Cast<UChaosOutfitAsset>(Asset);
+	if (OutfitAsset)
 	{
 		const TObjectPtr<const UChaosOutfit>& InOutfit = GetValue(Context, &Outfit);
 
diff --git a/Engine/Plugins/Experimental/ChaosOutfitAsset/Source/ChaosOutfitAssetDataflowNodes/Private/ChaosOutfitAsset/OutfitAssetTerminalNode.h b/Engine/Plugins/Experimental/ChaosOutfitAsset/Source/ChaosOutfitAssetDataflowNodes/Private/ChaosOutfitAsset/OutfitAssetTerminalNode.h
--- a/Engine/Plugins/Experimental/ChaosOutfitAsset/Source/ChaosOutfitAssetDataflowNodes/Private/ChaosOutfitAsset/OutfitAssetTerminalNode.h
+++ b/Engine/Plugins/Experimental/ChaosOutfitAsset/Source/ChaosOutfitAssetDataflowNodes/Private/ChaosOutfitAsset/OutfitAssetTerminalNode.h
@@ -2,10 +2,12 @@
 
 #pragma once 
 
+#include "Dataflow/DataflowFunctionProperty.h"
 #include "Dataflow/DataflowTerminalNode.h"
 #include "OutfitAssetTerminalNode.generated.h"
 
 class UChaosOutfit;
+class UChaosOutfitAsset;
 
 /**
  * Outfit terminal node to generate an outfit asset from several cloth assets.
@@ -28,5 +30,9 @@ private:
 	/** Input outfit collection. */
 	UPROPERTY(Meta = (DataflowInput))
 	TObjectPtr<const UChaosOutfit> Outfit;
+
+	/** Export this outfit as a skeletal mesh asset. */
+	UPROPERTY(EditAnywhere, Category = "Outfit Asset Terminal", Meta = (ButtonImage = "Icons.Convert"))
+	FDataflowFunctionProperty ConvertToSkeletalMesh;
 };
 
