[Backout] - CL41822120
#rnx
Original CL Desc
-----------------------------------------------------------------
Chaos Outfit Asset - Added SkeletalMesh export.

#rnx

diff --git a/Engine/Plugins/ChaosClothAsset/Source/ChaosClothAssetEngine/ChaosClothAssetEngine.Build.cs b/Engine/Plugins/ChaosClothAsset/Source/ChaosClothAssetEngine/ChaosClothAssetEngine.Build.cs
--- a/Engine/Plugins/ChaosClothAsset/Source/ChaosClothAssetEngine/ChaosClothAssetEngine.Build.cs
+++ b/Engine/Plugins/ChaosClothAsset/Source/ChaosClothAssetEngine/ChaosClothAssetEngine.Build.cs
@@ -21,21 +21,16 @@ public class ChaosClothAssetEngine : ModuleRules
 			new string[]
 			{
 				"Core",
+				"Engine",
 				"CoreUObject",
+				"RenderCore",
+				"RHI",
 				"Chaos",
-				"ChaosCaching",
-				"ChaosCloth",
 				"ChaosClothAsset",
+				"ChaosCloth",
+				"ChaosCaching",
 				"ClothingSystemRuntimeCommon",
-				"ClothingSystemRuntimeInterface",
-				"Engine",
-				"GeometryCore",
-				"MeshConversion",
-				"MeshDescription",
-				"RenderCore",
-				"RHI",
-				"SkeletalMeshDescription",
-				"SkeletalMeshUtilitiesCommon",
+				"ClothingSystemRuntimeInterface"
 			}
 		);
 
diff --git a/Engine/Plugins/ChaosClothAsset/Source/ChaosClothAssetEngine/Private/ChaosClothAsset/ClothAssetBase.cpp b/Engine/Plugins/ChaosClothAsset/Source/ChaosClothAssetEngine/Private/ChaosClothAsset/ClothAssetBase.cpp
--- a/Engine/Plugins/ChaosClothAsset/Source/ChaosClothAssetEngine/Private/ChaosClothAsset/ClothAssetBase.cpp
+++ b/Engine/Plugins/ChaosClothAsset/Source/ChaosClothAssetEngine/Private/ChaosClothAsset/ClothAssetBase.cpp
@@ -6,27 +6,15 @@
 #endif
 #include "ChaosClothAsset/ClothAssetPrivate.h"
 #include "ChaosClothAsset/ClothComponent.h"
-#if WITH_EDITOR
-#include "ChaosClothAsset/ClothEngineTools.h"
-#include "ChaosClothAsset/CollectionClothFacade.h"
-#include "DynamicMesh/DynamicMesh3.h"
-#endif
 #include "Engine/Engine.h"
 #include "Engine/SkeletalMesh.h"
 #include "Engine/SkinnedAssetAsyncCompileUtils.h"
+#include "Rendering/SkeletalMeshRenderData.h"
 #if WITH_EDITOR
 #include "Interfaces/ITargetPlatform.h"
-#include "Materials/Material.h"
-#include "Rendering/SkeletalMeshModel.h"
 #endif
-#include "Rendering/SkeletalMeshRenderData.h"
 #include "UObject/FortniteMainBranchObjectVersion.h"
-#if WITH_EDITOR
-#include "DynamicMeshToMeshDescription.h"
-#include "MaterialDomain.h"
-#include "SkeletalMeshAttributes.h"
-#include "StaticToSkeletalMeshConverter.h"
-#endif
+
 
 #include UE_INLINE_GENERATED_CPP_BY_NAME(ClothAssetBase)
 namespace UE::Chaos::ClothAsset::Private
@@ -496,106 +484,6 @@ PRAGMA_DISABLE_DEPRECATION_WARNINGS
 PRAGMA_ENABLE_DEPRECATION_WARNINGS
 }
 
-#if WITH_EDITOR
-bool UChaosClothAssetBase::ExportToSkeletalMesh(USkeletalMesh& SkeletalMesh) const
-{
-	using namespace UE::Chaos::ClothAsset;
-	using namespace UE::Geometry;
-
-	// Create one dynamic mesh from the multiple collections
-	TArray<UMaterialInterface*> RenderMaterials;
-	TArray<FDynamicMesh3> DynamicMeshes;
-	DynamicMeshes.AddDefaulted();  // TODO: Do the other LODs, one per Dynamic mesh
-
-	int32 MaterialOffset = 0;
-
-	const int32 NumModels = GetNumClothSimulationModels();
-
-	for (int32 ModelIndex = 0; ModelIndex < NumModels; ++ModelIndex)
-	{
-		const TArray<TSharedRef<const FManagedArrayCollection>>& ClothCollections = GetCollections(ModelIndex);
-
-		//for (const TSharedRef<const FManagedArrayCollection>& ClothCollection : ClothCollections)
-		if (ClothCollections.Num())
-		{
-			const TSharedRef<const FManagedArrayCollection>& ClothCollection = ClothCollections[0];  // TODO: Do other LODs
-			FCollectionClothConstFacade ClothFacade(ClothCollection);
-
-			if (ClothFacade.IsValid())
-			{
-				FDynamicMesh3 DynamicMeshPiece;
-				FClothEngineTools::BuildDynamicMeshFromRenderMesh(DynamicMeshPiece, ClothCollection, MaterialOffset);
-
-				if (ModelIndex == 0)
-				{
-					DynamicMeshes[0] = MoveTemp(DynamicMeshPiece);
-				}
-				else
-				{
-					DynamicMeshes[0].AppendWithOffsets(DynamicMeshPiece);
-				}
-
-				const TArrayView<const FString> MaterialPaths = ClothFacade.GetRenderMaterialPathName();
-				RenderMaterials.Reserve(RenderMaterials.Num() + MaterialPaths.Num());
-				for (int32 MaterialIndex = 0; MaterialIndex < MaterialPaths.Num(); ++MaterialIndex)
-				{
-					RenderMaterials.Emplace(LoadObject<UMaterialInterface>(nullptr, *MaterialPaths[MaterialIndex]));
-				}
-				MaterialOffset += MaterialPaths.Num();
-			}
-		}
-	}
-
-	// Create the skeletal mesh from the dynamic mesh
-	TArray<const FMeshDescription*> MeshDescriptions;
-
-	SkeletalMesh.PreEditChange(nullptr);
-	SkeletalMesh.GetImportedModel()->LODModels.Empty();
-	SkeletalMesh.ResetLODInfo();
-
-	TArray<FMeshDescription> LocalDescriptions;
-	for (const FDynamicMesh3& DynamicMesh : DynamicMeshes)
-	{
-		// Create a mesh description
-		FMeshDescription& MeshDescription = LocalDescriptions.AddDefaulted_GetRef();
-
-		// Add skeletal mesh attributes to the mesh description
-		FSkeletalMeshAttributes Attributes(MeshDescription);
-		Attributes.Register();
-
-		// Convert dynamic mesh to the mesh description
-		FConversionToMeshDescriptionOptions ConverterOptions;
-		FDynamicMeshToMeshDescription Converter(ConverterOptions);
-		constexpr bool bCopyTangents = true;
-		Converter.Convert(&DynamicMesh, MeshDescription, bCopyTangents);
-
-		// Add the created description to the list
-		MeshDescriptions.Add(&MeshDescription);
-	}
-
-	TArray<FSkeletalMaterial> SkeletalMaterials;
-	SkeletalMaterials.Reserve(Materials.Num());
-	for (UMaterialInterface* RenderMaterial : RenderMaterials)
-	{
-		SkeletalMaterials.Emplace(RenderMaterial, RenderMaterial->GetFName());
-	}
-	if (!SkeletalMaterials.Num())
-	{
-		SkeletalMaterials.Emplace(UMaterial::GetDefaultMaterial(MD_Surface));
-	}
-
-	constexpr bool bRecomputeNormals = false;
-	constexpr bool bRecomputeTangents = false;
-	return FStaticToSkeletalMeshConverter::InitializeSkeletalMeshFromMeshDescriptions(
-		&SkeletalMesh,
-		MeshDescriptions,
-		SkeletalMaterials,
-		GetRefSkeleton(),
-		bRecomputeNormals,
-		bRecomputeTangents);
-}
-#endif  // #if WITH_EDITOR
-
 void UChaosClothAssetBase::ReregisterComponents()
 {
 	// Recreate the simulation proxies with the updated physics asset
diff --git a/Engine/Plugins/ChaosClothAsset/Source/ChaosClothAssetEngine/Private/ChaosClothAsset/ClothCollectionToDynamicMesh.cpp b/Engine/Plugins/ChaosClothAsset/Source/ChaosClothAssetEngine/Private/ChaosClothAsset/ClothCollectionToDynamicMesh.cpp
deleted file mode 100644
diff --git a/Engine/Plugins/ChaosClothAsset/Source/ChaosClothAssetEngine/Private/ChaosClothAsset/ClothCollectionToDynamicMesh.h b/Engine/Plugins/ChaosClothAsset/Source/ChaosClothAssetEngine/Private/ChaosClothAsset/ClothCollectionToDynamicMesh.h
deleted file mode 100644
diff --git a/Engine/Plugins/ChaosClothAsset/Source/ChaosClothAssetEngine/Private/ChaosClothAsset/ClothCollectionToDynamicMeshMappingSupport.cpp b/Engine/Plugins/ChaosClothAsset/Source/ChaosClothAssetEngine/Private/ChaosClothAsset/ClothCollectionToDynamicMeshMappingSupport.cpp
deleted file mode 100644
diff --git a/Engine/Plugins/ChaosClothAsset/Source/ChaosClothAssetEngine/Private/ChaosClothAsset/ClothEngineTools.cpp b/Engine/Plugins/ChaosClothAsset/Source/ChaosClothAssetEngine/Private/ChaosClothAsset/ClothEngineTools.cpp
--- a/Engine/Plugins/ChaosClothAsset/Source/ChaosClothAssetEngine/Private/ChaosClothAsset/ClothEngineTools.cpp
+++ b/Engine/Plugins/ChaosClothAsset/Source/ChaosClothAssetEngine/Private/ChaosClothAsset/ClothEngineTools.cpp
@@ -1,149 +1,125 @@
 // Copyright Epic Games, Inc. All Rights Reserved.
 
 #include "ChaosClothAsset/ClothEngineTools.h"
-#if WITH_EDITOR
-#include "ChaosClothAsset/ClothCollectionToDynamicMesh.h"
-#endif
 #include "ChaosClothAsset/CollectionClothFacade.h"
 #include "ChaosClothAsset/CollectionClothSelectionFacade.h"
 #include "ChaosClothAsset/ClothCollectionGroup.h"
 #include "ChaosClothAsset/ClothGeometryTools.h"
 #include "ClothTetherData.h"
-#if WITH_EDITOR
-#include "DynamicMesh/DynamicMesh3.h"
-#include "DynamicMesh/DynamicMeshAttributeSet.h"
-#endif
 
 namespace UE::Chaos::ClothAsset
 {
-	namespace Private
+
+namespace Private
+{
+static void AppendTetherData(FCollectionClothFacade& ClothFacade, const FClothTetherData& TetherData)
+{
+	// Append new tethers
+	TArrayView<TArray<int32>> TetherKinematicIndex = ClothFacade.GetTetherKinematicIndex();
+	TArrayView<TArray<float>> TetherReferenceLength = ClothFacade.GetTetherReferenceLength();
+	for (const TArray<TTuple<int32, int32, float>>& TetherBatch : TetherData.Tethers)
 	{
-		static void AppendTetherData(FCollectionClothFacade& ClothFacade, const FClothTetherData& TetherData)
+		for (const TTuple<int32, int32, float>& Tether : TetherBatch)
 		{
-			// Append new tethers
-			TArrayView<TArray<int32>> TetherKinematicIndex = ClothFacade.GetTetherKinematicIndex();
-			TArrayView<TArray<float>> TetherReferenceLength = ClothFacade.GetTetherReferenceLength();
-			for (const TArray<TTuple<int32, int32, float>>& TetherBatch : TetherData.Tethers)
-			{
-				for (const TTuple<int32, int32, float>& Tether : TetherBatch)
-				{
-					// Tuple is Kinematic, Dynamic, RefLength
-					const int32 DynamicIndex = Tether.Get<1>();
-					TArray<int32>& KinematicIndex = TetherKinematicIndex[DynamicIndex];
-					TArray<float>& ReferenceLength = TetherReferenceLength[DynamicIndex];
-					check(KinematicIndex.Num() == ReferenceLength.Num());
-					checkSlow(KinematicIndex.Find(Tether.Get<0>()) == INDEX_NONE);
-					KinematicIndex.Add(Tether.Get<0>());
-					ReferenceLength.Add(Tether.Get<2>());
-				}
-			}
+			// Tuple is Kinematic, Dynamic, RefLength
+			const int32 DynamicIndex = Tether.Get<1>();
+			TArray<int32>& KinematicIndex = TetherKinematicIndex[DynamicIndex];
+			TArray<float>& ReferenceLength = TetherReferenceLength[DynamicIndex];
+			check(KinematicIndex.Num() == ReferenceLength.Num());
+			checkSlow(KinematicIndex.Find(Tether.Get<0>()) == INDEX_NONE);
+			KinematicIndex.Add(Tether.Get<0>());
+			ReferenceLength.Add(Tether.Get<2>());
 		}
 	}
+}
+}
 
-	void FClothEngineTools::GenerateTethers(const TSharedRef<FManagedArrayCollection>& ClothCollection, const FName& WeightMapName, const bool bGenerateGeodesicTethers, const FVector2f& MaxDistanceValue)
+void FClothEngineTools::GenerateTethers(const TSharedRef<FManagedArrayCollection>& ClothCollection, const FName& WeightMapName, const bool bGenerateGeodesicTethers, const FVector2f& MaxDistanceValue)
+{
+	FCollectionClothFacade ClothFacade(ClothCollection);
+	FClothGeometryTools::DeleteTethers(ClothCollection);
+	if (ClothFacade.HasWeightMap(WeightMapName))
 	{
-		FCollectionClothFacade ClothFacade(ClothCollection);
-		FClothGeometryTools::DeleteTethers(ClothCollection);
-		if (ClothFacade.HasWeightMap(WeightMapName))
+		FClothTetherData TetherData;
+		TArray<uint32> SimIndices;
+		SimIndices.Reserve(ClothFacade.GetNumSimFaces() * 3);
+		for (const FIntVector3& Face : ClothFacade.GetSimIndices3D())
 		{
-			FClothTetherData TetherData;
-			TArray<uint32> SimIndices;
-			SimIndices.Reserve(ClothFacade.GetNumSimFaces() * 3);
-			for (const FIntVector3& Face : ClothFacade.GetSimIndices3D())
-			{
-				SimIndices.Add(Face[0]);
-				SimIndices.Add(Face[1]);
-				SimIndices.Add(Face[2]);
-			}
-
-			if (MaxDistanceValue.Equals(FVector2f(0.f, 1.f)))
-			{
-				TetherData.GenerateTethers(ClothFacade.GetSimPosition3D(), TConstArrayView<uint32>(SimIndices), ClothFacade.GetWeightMap(WeightMapName), bGenerateGeodesicTethers);
-			}
-			else
-			{
+			SimIndices.Add(Face[0]);
+			SimIndices.Add(Face[1]);
+			SimIndices.Add(Face[2]);
+		}
 
-				const TSet<int32> KinematicVertices = FClothGeometryTools::GenerateKinematicVertices3D(ClothCollection, WeightMapName, MaxDistanceValue, NAME_None);
-				TetherData.GenerateTethers(ClothFacade.GetSimPosition3D(), TConstArrayView<uint32>(SimIndices), KinematicVertices, bGenerateGeodesicTethers);
-			}
+		if (MaxDistanceValue.Equals(FVector2f(0.f, 1.f)))
+		{
+			TetherData.GenerateTethers(ClothFacade.GetSimPosition3D(), TConstArrayView<uint32>(SimIndices), ClothFacade.GetWeightMap(WeightMapName), bGenerateGeodesicTethers);
+		}
+		else
+		{
 
-			Private::AppendTetherData(ClothFacade, TetherData);
+			const TSet<int32> KinematicVertices = FClothGeometryTools::GenerateKinematicVertices3D(ClothCollection, WeightMapName, MaxDistanceValue, NAME_None);
+			TetherData.GenerateTethers(ClothFacade.GetSimPosition3D(), TConstArrayView<uint32>(SimIndices), KinematicVertices, bGenerateGeodesicTethers);
 		}
+
+		Private::AppendTetherData(ClothFacade, TetherData);
 	}
+}
 
-	void FClothEngineTools::GenerateTethersFromSelectionSet(const TSharedRef<FManagedArrayCollection>& ClothCollection, const FName& FixedEndSet, const bool bGeodesicTethers)
-	{
-		FClothGeometryTools::DeleteTethers(ClothCollection);
-		FCollectionClothSelectionConstFacade SelectionFacade(ClothCollection);
+void FClothEngineTools::GenerateTethersFromSelectionSet(const TSharedRef<FManagedArrayCollection>& ClothCollection, const FName& FixedEndSet, const bool bGeodesicTethers)
+{
+	FClothGeometryTools::DeleteTethers(ClothCollection);
+	FCollectionClothSelectionConstFacade SelectionFacade(ClothCollection);
 
-		if (SelectionFacade.HasSelection(FixedEndSet) && SelectionFacade.GetSelectionGroup(FixedEndSet) == UE::Chaos::ClothAsset::ClothCollectionGroup::SimVertices3D)
+	if (SelectionFacade.HasSelection(FixedEndSet) && SelectionFacade.GetSelectionGroup(FixedEndSet) == UE::Chaos::ClothAsset::ClothCollectionGroup::SimVertices3D)
+	{
+		FCollectionClothFacade ClothFacade(ClothCollection);
+		FClothTetherData TetherData;
+		TArray<uint32> SimIndices;
+		SimIndices.Reserve(ClothFacade.GetNumSimFaces() * 3);
+		for (const FIntVector3& Face : ClothFacade.GetSimIndices3D())
 		{
-			FCollectionClothFacade ClothFacade(ClothCollection);
-			FClothTetherData TetherData;
-			TArray<uint32> SimIndices;
-			SimIndices.Reserve(ClothFacade.GetNumSimFaces() * 3);
-			for (const FIntVector3& Face : ClothFacade.GetSimIndices3D())
-			{
-				SimIndices.Add(Face[0]);
-				SimIndices.Add(Face[1]);
-				SimIndices.Add(Face[2]);
-			}
-
-			TetherData.GenerateTethers(ClothFacade.GetSimPosition3D(), TConstArrayView<uint32>(SimIndices), SelectionFacade.GetSelectionSet(FixedEndSet), bGeodesicTethers);
-			Private::AppendTetherData(ClothFacade, TetherData);
+			SimIndices.Add(Face[0]);
+			SimIndices.Add(Face[1]);
+			SimIndices.Add(Face[2]);
 		}
+
+		TetherData.GenerateTethers(ClothFacade.GetSimPosition3D(), TConstArrayView<uint32>(SimIndices), SelectionFacade.GetSelectionSet(FixedEndSet), bGeodesicTethers);
+		Private::AppendTetherData(ClothFacade, TetherData);
 	}
+}
 
-	void FClothEngineTools::GenerateTethersFromCustomSelectionSets(const TSharedRef<FManagedArrayCollection>& ClothCollection, const FName& InFixedEndSet, const TArray<TPair<FName, FName>>& CustomTetherEndSets, const bool bGeodesicTethers)
+void FClothEngineTools::GenerateTethersFromCustomSelectionSets(const TSharedRef<FManagedArrayCollection>& ClothCollection, const FName& InFixedEndSet, const TArray<TPair<FName, FName>>& CustomTetherEndSets, const bool bGeodesicTethers)
+{
+	FClothGeometryTools::DeleteTethers(ClothCollection);
+	FCollectionClothSelectionConstFacade SelectionFacade(ClothCollection);
+	if (SelectionFacade.HasSelection(InFixedEndSet) && SelectionFacade.GetSelectionGroup(InFixedEndSet) == UE::Chaos::ClothAsset::ClothCollectionGroup::SimVertices3D)
 	{
-		FClothGeometryTools::DeleteTethers(ClothCollection);
-		FCollectionClothSelectionConstFacade SelectionFacade(ClothCollection);
-		if (SelectionFacade.HasSelection(InFixedEndSet) && SelectionFacade.GetSelectionGroup(InFixedEndSet) == UE::Chaos::ClothAsset::ClothCollectionGroup::SimVertices3D)
-		{
 
-			FCollectionClothFacade ClothFacade(ClothCollection);
-			TArray<uint32> SimIndices;
-			SimIndices.Reserve(ClothFacade.GetNumSimFaces() * 3);
-			for (const FIntVector3& Face : ClothFacade.GetSimIndices3D())
-			{
-				SimIndices.Add(Face[0]);
-				SimIndices.Add(Face[1]);
-				SimIndices.Add(Face[2]);
-			}
+		FCollectionClothFacade ClothFacade(ClothCollection);
+		TArray<uint32> SimIndices;
+		SimIndices.Reserve(ClothFacade.GetNumSimFaces() * 3);
+		for (const FIntVector3& Face : ClothFacade.GetSimIndices3D())
+		{
+			SimIndices.Add(Face[0]);
+			SimIndices.Add(Face[1]);
+			SimIndices.Add(Face[2]);
+		}
 
-			const TSet<int32>& FixedEndSet = SelectionFacade.GetSelectionSet(InFixedEndSet);
+		const TSet<int32>& FixedEndSet = SelectionFacade.GetSelectionSet(InFixedEndSet);
 
-			for (const TPair<FName, FName>& TetherEnds : CustomTetherEndSets)
+		for (const TPair<FName, FName>& TetherEnds : CustomTetherEndSets)
+		{
+			const FName& CustomDynamicEndSet = TetherEnds.Get<0>();
+			const FName& CustomFixedEndSet = TetherEnds.Get<1>();
+			if (SelectionFacade.HasSelection(CustomFixedEndSet) && SelectionFacade.GetSelectionGroup(CustomFixedEndSet) == UE::Chaos::ClothAsset::ClothCollectionGroup::SimVertices3D &&
+				SelectionFacade.HasSelection(CustomDynamicEndSet) && SelectionFacade.GetSelectionGroup(CustomDynamicEndSet) == UE::Chaos::ClothAsset::ClothCollectionGroup::SimVertices3D)
 			{
-				const FName& CustomDynamicEndSet = TetherEnds.Get<0>();
-				const FName& CustomFixedEndSet = TetherEnds.Get<1>();
-				if (SelectionFacade.HasSelection(CustomFixedEndSet) && SelectionFacade.GetSelectionGroup(CustomFixedEndSet) == UE::Chaos::ClothAsset::ClothCollectionGroup::SimVertices3D &&
-					SelectionFacade.HasSelection(CustomDynamicEndSet) && SelectionFacade.GetSelectionGroup(CustomDynamicEndSet) == UE::Chaos::ClothAsset::ClothCollectionGroup::SimVertices3D)
-				{
-					FClothTetherData TetherData;
-					TetherData.GenerateTethers(ClothFacade.GetSimPosition3D(), TConstArrayView<uint32>(SimIndices), FixedEndSet, SelectionFacade.GetSelectionSet(CustomDynamicEndSet), SelectionFacade.GetSelectionSet(CustomFixedEndSet), bGeodesicTethers);
-					Private::AppendTetherData(ClothFacade, TetherData);
-				}
+				FClothTetherData TetherData;
+				TetherData.GenerateTethers(ClothFacade.GetSimPosition3D(), TConstArrayView<uint32>(SimIndices), FixedEndSet, SelectionFacade.GetSelectionSet(CustomDynamicEndSet), SelectionFacade.GetSelectionSet(CustomFixedEndSet), bGeodesicTethers);
+				Private::AppendTetherData(ClothFacade, TetherData);
 			}
 		}
 	}
+}
 
-#if WITH_EDITOR
-	void FClothEngineTools::BuildDynamicMeshFromSimMesh(UE::Geometry::FDynamicMesh3& DynamicMesh, const TSharedRef<const FManagedArrayCollection>& ClothCollection)
-	{
-		constexpr bool bDisableAttributes = false;
-		const int32 PatternIndex = INDEX_NONE;
-		DynamicMesh.EnableAttributes();
-		FClothCollectionToDynamicMesh::Convert(ClothCollection, PatternIndex, FClothCollectionToDynamicMesh::ESourceMeshType::Sim3D, DynamicMesh, bDisableAttributes);
-	}
-
-	void FClothEngineTools::BuildDynamicMeshFromRenderMesh(UE::Geometry::FDynamicMesh3& DynamicMesh, const TSharedRef<const FManagedArrayCollection>& ClothCollection, int32 MaterialOffset)
-	{
-		constexpr bool bDisableAttributes = false;
-		const int32 PatternIndex = INDEX_NONE;
-		DynamicMesh.EnableAttributes();
-		DynamicMesh.Attributes()->EnableMaterialID();
-		FClothCollectionToDynamicMesh::Convert(ClothCollection, PatternIndex, FClothCollectionToDynamicMesh::ESourceMeshType::Render, DynamicMesh, bDisableAttributes, MaterialOffset);
-	}
-#endif  // #if WITH_EDITOR
 }  // End namespace UE::Chaos::ClothAsset
diff --git a/Engine/Plugins/ChaosClothAsset/Source/ChaosClothAssetEngine/Public/ChaosClothAsset/ClothAssetBase.h b/Engine/Plugins/ChaosClothAsset/Source/ChaosClothAssetEngine/Public/ChaosClothAsset/ClothAssetBase.h
--- a/Engine/Plugins/ChaosClothAsset/Source/ChaosClothAssetEngine/Public/ChaosClothAsset/ClothAssetBase.h
+++ b/Engine/Plugins/ChaosClothAsset/Source/ChaosClothAssetEngine/Public/ChaosClothAsset/ClothAssetBase.h
@@ -260,16 +260,6 @@ public:
 
 	void SetHasVertexColors(bool InbHasVertexColors);
 
-#if WITH_EDITOR
-	/**
-	 * Export the graphical representation of this asset to a SkeletalMesh asset.
-	 * Includes skinning, but excludes all clothing simulation data since it isn't compatible with the Cloth Asset.
-	 * @param SkeletalMesh The empty skeletal mesh asset to be written to.
-	 * @return Whether the initialization of the skeletal mesh has succeeded or not.
-	 */
-	bool ExportToSkeletalMesh(USkeletalMesh& SkeletalMesh) const;
-#endif
-
 protected:
 	/** Reregister all components using this asset to reset the simulation in case anything has changed. */
 	void ReregisterComponents();
diff --git a/Engine/Plugins/ChaosClothAsset/Source/ChaosClothAssetEngine/Public/ChaosClothAsset/ClothEngineTools.h b/Engine/Plugins/ChaosClothAsset/Source/ChaosClothAssetEngine/Public/ChaosClothAsset/ClothEngineTools.h
--- a/Engine/Plugins/ChaosClothAsset/Source/ChaosClothAssetEngine/Public/ChaosClothAsset/ClothEngineTools.h
+++ b/Engine/Plugins/ChaosClothAsset/Source/ChaosClothAssetEngine/Public/ChaosClothAsset/ClothEngineTools.h
@@ -8,9 +8,6 @@
 struct FManagedArrayCollection;
 class FName;
 class FSkeletalMeshLODModel;
-#if WITH_EDITOR
-namespace UE::Geometry { class FDynamicMesh3; }
-#endif
 
 namespace UE::Chaos::ClothAsset
 {
@@ -24,16 +21,5 @@ namespace UE::Chaos::ClothAsset
 		static void GenerateTethersFromSelectionSet(const TSharedRef<FManagedArrayCollection>& ClothCollection, const FName& FixedEndSet, const bool bGeodesicTethers);
 		/** @param CustomTetherEndSets: First element of each pair is DynamicSet, second is FixedSet */
 		static void GenerateTethersFromCustomSelectionSets(const TSharedRef<FManagedArrayCollection>& ClothCollection, const FName& FixedEndSet, const TArray<TPair<FName, FName>>& CustomTetherEndSets, const bool bGeodesicTethers);
-
-#if WITH_EDITOR
-		/** Build a DynamicMesh from a sim mesh. */
-		static void BuildDynamicMeshFromSimMesh(UE::Geometry::FDynamicMesh3& DynamicMesh, const TSharedRef<const FManagedArrayCollection>& ClothCollection);
-
-		/**
-		 * Build a DynamicMesh from a render mesh.
-		 * Use the material offset to add to the final material group id if the material list doesn't start at 0.
-		 */
-		static void BuildDynamicMeshFromRenderMesh(UE::Geometry::FDynamicMesh3& DynamicMesh, const TSharedRef<const FManagedArrayCollection>& ClothCollection, int32 MaterialOffset = 0);
-#endif  // #if WITH_EDITOR
 	};
 }  // End namespace UE::Chaos::ClothAsset
\ No newline at end of file
diff --git a/Engine/Plugins/Experimental/ChaosOutfitAsset/Source/ChaosOutfitAssetDataflowNodes/ChaosOutfitAssetDataflowNodes.Build.cs b/Engine/Plugins/Experimental/ChaosOutfitAsset/Source/ChaosOutfitAssetDataflowNodes/ChaosOutfitAssetDataflowNodes.Build.cs
--- a/Engine/Plugins/Experimental/ChaosOutfitAsset/Source/ChaosOutfitAssetDataflowNodes/ChaosOutfitAssetDataflowNodes.Build.cs
+++ b/Engine/Plugins/Experimental/ChaosOutfitAsset/Source/ChaosOutfitAssetDataflowNodes/ChaosOutfitAssetDataflowNodes.Build.cs
@@ -14,18 +14,12 @@ public class ChaosOutfitAssetDataflowNodes : ModuleRules
 				"Core",
 				"CoreUObject",
 				"Chaos",
-				"ChaosClothAsset",
 				"ChaosClothAssetEngine",
 				"ChaosOutfitAssetEngine",
 				"DataflowCore",
 				"DataflowNodes",
 				"Engine",
-				"GeometryCore",
-				"MeshConversion",
-				"MeshDescription",
 				"MeshResizingCore",
-				"SkeletalMeshDescription",
-				"SkeletalMeshUtilitiesCommon",
 			}
 		);
 	}
diff --git a/Engine/Plugins/Experimental/ChaosOutfitAsset/Source/ChaosOutfitAssetDataflowNodes/Private/ChaosOutfitAsset/OutfitAssetTerminalNode.cpp b/Engine/Plugins/Experimental/ChaosOutfitAsset/Source/ChaosOutfitAssetDataflowNodes/Private/ChaosOutfitAsset/OutfitAssetTerminalNode.cpp
--- a/Engine/Plugins/Experimental/ChaosOutfitAsset/Source/ChaosOutfitAssetDataflowNodes/Private/ChaosOutfitAsset/OutfitAssetTerminalNode.cpp
+++ b/Engine/Plugins/Experimental/ChaosOutfitAsset/Source/ChaosOutfitAssetDataflowNodes/Private/ChaosOutfitAsset/OutfitAssetTerminalNode.cpp
@@ -1,107 +1,22 @@
 // Copyright Epic Games, Inc. All Rights Reserved.
 
 #include "ChaosOutfitAsset/OutfitAssetTerminalNode.h"
-#if WITH_EDITOR
-#include "AssetRegistry/AssetRegistryModule.h"
-#endif
 #include "ChaosOutfitAsset/OutfitAsset.h"
 #include "Dataflow/DataflowInputOutput.h"
-#if WITH_EDITOR
-#include "Dataflow/DataflowObjectInterface.h"
-#include "Engine/SkeletalMesh.h"
-#include "Misc/FileHelper.h"
-#include "Misc/ScopedSlowTask.h"
-#include "ContentBrowserModule.h"
-#include "IContentBrowserSingleton.h"
-#endif
 
 #include UE_INLINE_GENERATED_CPP_BY_NAME(OutfitAssetTerminalNode)
 
 #define LOCTEXT_NAMESPACE "ChaosOutfitAssetTerminalNode"
 
-namespace UE::Chaos::OutfitAsset::Private
-{
-#if WITH_EDITOR
-	static void ExportToSkeletalMesh(const TObjectPtr<const UChaosOutfitAsset>& OutfitAsset)
-	{
-		using namespace UE::Geometry;
-		if (OutfitAsset)
-		{
-			// Make a name from the outfit asset name
-			const FString SkeletalMeshPath = FPackageName::GetLongPackagePath(OutfitAsset->GetOutermost()->GetName());
-			const FString OutfitAssetName = OutfitAsset->GetName();
-			FString SkeletalMeshName = FString(TEXT("SK_")) + (OutfitAssetName.StartsWith(TEXT("OA_")) ? OutfitAssetName.RightChop(3) : OutfitAssetName);
-			FString SkeletalMeshPackageName = FPaths::Combine(SkeletalMeshPath, SkeletalMeshName);
-
-			// Open a Save As dialog
-			FSaveAssetDialogConfig SaveAssetDialogConfig;
-			SaveAssetDialogConfig.DefaultPath = SkeletalMeshPath;
-			SaveAssetDialogConfig.DefaultAssetName = SkeletalMeshName;
-			SaveAssetDialogConfig.AssetClassNames.Add(USkeletalMesh::StaticClass()->GetClassPathName());
-			SaveAssetDialogConfig.ExistingAssetPolicy = ESaveAssetDialogExistingAssetPolicy::AllowButWarn;
-			SaveAssetDialogConfig.DialogTitleOverride = LOCTEXT("SaveAssetDialogConfig", "Convert Outfit Asset To Skeletal Mesh As");
-
-			const FContentBrowserModule& ContentBrowserModule = FModuleManager::LoadModuleChecked<FContentBrowserModule>("ContentBrowser");
-			const FString AssetSavePath = ContentBrowserModule.Get().CreateModalSaveAssetDialog(SaveAssetDialogConfig);
-			const FString AssetSavePackageName = FPackageName::ObjectPathToPackageName(AssetSavePath);
-			FText OutError;
-			if (AssetSavePath.IsEmpty() || !FFileHelper::IsFilenameValidForSaving(AssetSavePackageName, OutError))
-			{
-				return;
-			}
-			SkeletalMeshPackageName = AssetSavePackageName;
-			SkeletalMeshName = FPackageName::GetLongPackageAssetName(SkeletalMeshPackageName);
-
-			constexpr float NumSteps = 0.f;
-			FScopedSlowTask SlowTask(NumSteps, LOCTEXT("ExportToSkeletalMesh", "Exporting Cloth Outfit to a Skeletal Mesh asset..."));
-			SlowTask.MakeDialog();
-
-			UPackage* const SkeletalMeshPackage = CreatePackage(*SkeletalMeshPackageName);
-			USkeletalMesh* const SkeletalMesh =
-				NewObject<USkeletalMesh>(
-					SkeletalMeshPackage,
-					USkeletalMesh::StaticClass(),
-					FName(SkeletalMeshName),
-					RF_Public | RF_Standalone | RF_Transactional);
-
-			// Export to skeletal mesh
-			if (SkeletalMesh)
-			{
-				OutfitAsset->ExportToSkeletalMesh(*SkeletalMesh);
-			}
-
-			SkeletalMesh->MarkPackageDirty();
-
-			// Notify the asset registry
-			FAssetRegistryModule::AssetCreated(SkeletalMesh);
-		}
-	}
-#endif  // #if WITH_EDITOR
-}
-
 FChaosOutfitAssetTerminalNode::FChaosOutfitAssetTerminalNode(const UE::Dataflow::FNodeParameters& InParam, FGuid InGuid)
 	: FDataflowTerminalNode(InParam, InGuid)
-#if WITH_EDITOR
-	, ConvertToSkeletalMesh(
-		FDataflowFunctionProperty::FDelegate::CreateLambda([this](UE::Dataflow::FContext& Context)
-			{
-				if (const UE::Dataflow::FEngineContext* EngineContext = Context.AsType<UE::Dataflow::FEngineContext>())
-				{
-					if (const UChaosOutfitAsset* const OwnerOutfitAsset = Cast<UChaosOutfitAsset>(EngineContext->Owner))
-					{
-						UE::Chaos::OutfitAsset::Private::ExportToSkeletalMesh(OwnerOutfitAsset);
-					}
-				}
-			}))
-#endif
 {
 	RegisterInputConnection(&Outfit);
 }
 
 void FChaosOutfitAssetTerminalNode::SetAssetValue(TObjectPtr<UObject> Asset, UE::Dataflow::FContext& Context) const
 {
-	TObjectPtr<UChaosOutfitAsset> OutfitAsset = Cast<UChaosOutfitAsset>(Asset);
-	if (OutfitAsset)
+	if (UChaosOutfitAsset* const OutfitAsset = Cast<UChaosOutfitAsset>(Asset.Get()))
 	{
 		const TObjectPtr<const UChaosOutfit>& InOutfit = GetValue(Context, &Outfit);
 
diff --git a/Engine/Plugins/Experimental/ChaosOutfitAsset/Source/ChaosOutfitAssetDataflowNodes/Private/ChaosOutfitAsset/OutfitAssetTerminalNode.h b/Engine/Plugins/Experimental/ChaosOutfitAsset/Source/ChaosOutfitAssetDataflowNodes/Private/ChaosOutfitAsset/OutfitAssetTerminalNode.h
--- a/Engine/Plugins/Experimental/ChaosOutfitAsset/Source/ChaosOutfitAssetDataflowNodes/Private/ChaosOutfitAsset/OutfitAssetTerminalNode.h
+++ b/Engine/Plugins/Experimental/ChaosOutfitAsset/Source/ChaosOutfitAssetDataflowNodes/Private/ChaosOutfitAsset/OutfitAssetTerminalNode.h
@@ -2,12 +2,10 @@
 
 #pragma once 
 
-#include "Dataflow/DataflowFunctionProperty.h"
 #include "Dataflow/DataflowTerminalNode.h"
 #include "OutfitAssetTerminalNode.generated.h"
 
 class UChaosOutfit;
-class UChaosOutfitAsset;
 
 /**
  * Outfit terminal node to generate an outfit asset from several cloth assets.
@@ -30,9 +28,5 @@ private:
 	/** Input outfit collection. */
 	UPROPERTY(Meta = (DataflowInput))
 	TObjectPtr<const UChaosOutfit> Outfit;
-
-	/** Export this outfit as a skeletal mesh asset. */
-	UPROPERTY(EditAnywhere, Category = "Outfit Asset Terminal", Meta = (ButtonImage = "Icons.Convert"))
-	FDataflowFunctionProperty ConvertToSkeletalMesh;
 };
 
