Fix crash when cooking cloth assets with new morph targets.

diff --git a/Engine/Plugins/ChaosClothAsset/Source/ChaosClothAsset/Private/ChaosClothAsset/ClothCollection.cpp b/Engine/Plugins/ChaosClothAsset/Source/ChaosClothAsset/Private/ChaosClothAsset/ClothCollection.cpp
--- a/Engine/Plugins/ChaosClothAsset/Source/ChaosClothAsset/Private/ChaosClothAsset/ClothCollection.cpp
+++ b/Engine/Plugins/ChaosClothAsset/Source/ChaosClothAsset/Private/ChaosClothAsset/ClothCollection.cpp
@@ -220,7 +220,7 @@ namespace UE::Chaos::ClothAsset
 		FabricPressure = ManagedArrayCollection->FindAttribute<float>(ClothCollectionAttribute::FabricPressure, ClothCollectionGroup::Fabrics);
 		FabricLayer = ManagedArrayCollection->FindAttribute<int32>(ClothCollectionAttribute::FabricLayer, ClothCollectionGroup::Fabrics);
 		FabricCollisionThickness = ManagedArrayCollection->FindAttribute<float>(ClothCollectionAttribute::FabricCollisionThickness, ClothCollectionGroup::Fabrics);
-		
+
 		// Sim Faces Group
 		SimIndices2D = ManagedArrayCollection->FindAttribute<FIntVector3>(ClothCollectionAttribute::SimIndices2D, ClothCollectionGroup::SimFaces);
 		SimIndices3D = ManagedArrayCollection->FindAttribute<FIntVector3>(ClothCollectionAttribute::SimIndices3D, ClothCollectionGroup::SimFaces);
@@ -271,7 +271,7 @@ namespace UE::Chaos::ClothAsset
 
 	bool FClothCollection::IsValid(EClothCollectionOptionalSchemas OptionalSchemas) const
 	{
-		if (OptionalSchemas == EClothCollectionOptionalSchemas::CookedOnly)
+		if (EnumHasAnyFlags(OptionalSchemas, EClothCollectionOptionalSchemas::CookedOnly))
 		{
 			// Sim Morph Targets Group
 			return
@@ -293,9 +293,9 @@ namespace UE::Chaos::ClothAsset
 			// Solvers Group
 			(!EnumHasAnyFlags(OptionalSchemas, EClothCollectionOptionalSchemas::Solvers) ||
 				(SolverGravity &&
-				 SolverAirDamping &&
-				 SolverSubSteps &&
-				 SolverTimeStep )) &&
+					SolverAirDamping &&
+					SolverSubSteps &&
+					SolverTimeStep)) &&
 
 			// Seam Group
 			SeamStitchStart &&
@@ -311,7 +311,7 @@ namespace UE::Chaos::ClothAsset
 			SimFacesStart &&
 			SimFacesEnd &&
 			SimPatternFabric &&
-		
+
 
 			// Render Patterns Group
 			RenderVerticesStart &&
@@ -329,15 +329,15 @@ namespace UE::Chaos::ClothAsset
 			// Fabrics Group
 			(!EnumHasAnyFlags(OptionalSchemas, EClothCollectionOptionalSchemas::Fabrics) ||
 				(FabricBendingStiffness &&
-				FabricBucklingStiffness &&
-				FabricStretchStiffness &&
-				FabricBucklingRatio &&
-				FabricDensity &&
-				FabricFriction &&
-				FabricDamping &&
-				FabricPressure &&
-				FabricLayer &&
-				FabricCollisionThickness )) &&
+					FabricBucklingStiffness &&
+					FabricStretchStiffness &&
+					FabricBucklingRatio &&
+					FabricDensity &&
+					FabricFriction &&
+					FabricDamping &&
+					FabricPressure &&
+					FabricLayer &&
+					FabricCollisionThickness)) &&
 
 			// Sim Vertices 2D Group
 			SimPosition2D &&
@@ -356,14 +356,14 @@ namespace UE::Chaos::ClothAsset
 			// Sim Morph Targets Group
 			(!EnumHasAnyFlags(OptionalSchemas, EClothCollectionOptionalSchemas::SimMorphTargets) ||
 				(SimMorphTargetName &&
-				SimMorphTargetVerticesStart &&
-				SimMorphTargetVerticesEnd)) &&
+					SimMorphTargetVerticesStart &&
+					SimMorphTargetVerticesEnd)) &&
 
 			// Sim Morph Target Vertices Group
 			(!EnumHasAnyFlags(OptionalSchemas, EClothCollectionOptionalSchemas::SimMorphTargets) ||
 				(SimMorphTargetPositionDelta &&
-				SimMorphTargetTangentZDelta &&
-				SimMorphTargetSimVertex3DIndex)) &&
+					SimMorphTargetTangentZDelta &&
+					SimMorphTargetSimVertex3DIndex)) &&
 
 			// Render Faces Group
 			RenderIndices &&
@@ -379,11 +379,11 @@ namespace UE::Chaos::ClothAsset
 			RenderBoneWeights &&
 			(!EnumHasAnyFlags(OptionalSchemas, EClothCollectionOptionalSchemas::RenderDeformer) ||
 				(RenderDeformerPositionBaryCoordsAndDist &&
-				RenderDeformerNormalBaryCoordsAndDist &&
-				RenderDeformerTangentBaryCoordsAndDist &&
-				RenderDeformerSimIndices3D &&
-				RenderDeformerWeight &&
-				RenderDeformerSkinningBlend));
+					RenderDeformerNormalBaryCoordsAndDist &&
+					RenderDeformerTangentBaryCoordsAndDist &&
+					RenderDeformerSimIndices3D &&
+					RenderDeformerWeight &&
+					RenderDeformerSkinningBlend));
 	}
 
 	void FClothCollection::DefineSchema(EClothCollectionOptionalSchemas OptionalSchemas)
@@ -402,7 +402,7 @@ namespace UE::Chaos::ClothAsset
 		FManagedArrayCollection::FConstructionParameters SimVertices3DDependency(ClothCollectionGroup::SimVertices3D, bSaved, bAllowCircularDependency);  // Any attribute with this dependency must handle welding and splitting in FCollectionClothSeamFacade
 		FManagedArrayCollection::FConstructionParameters SimMorphTargetVerticesDependency(ClothCollectionGroup::SimMorphTargetVertices, bSaved, bAllowCircularDependency);
 
-		if (OptionalSchemas == EClothCollectionOptionalSchemas::CookedOnly)
+		if (EnumHasAnyFlags(OptionalSchemas, EClothCollectionOptionalSchemas::CookedOnly))
 		{
 			// Special case: only define cooked schema.
 
@@ -431,7 +431,7 @@ namespace UE::Chaos::ClothAsset
 			SolverSubSteps = &ManagedArrayCollection->AddAttribute<int32>(ClothCollectionAttribute::SolverSubSteps, ClothCollectionGroup::Solvers);
 			SolverTimeStep = &ManagedArrayCollection->AddAttribute<float>(ClothCollectionAttribute::SolverTimeStep, ClothCollectionGroup::Solvers);
 		}
-		
+
 		// Seams Group
 		SeamStitchStart = &ManagedArrayCollection->AddAttribute<int32>(ClothCollectionAttribute::SeamStitchStart, ClothCollectionGroup::Seams, SeamStitchesDependency);
 		SeamStitchEnd = &ManagedArrayCollection->AddAttribute<int32>(ClothCollectionAttribute::SeamStitchEnd, ClothCollectionGroup::Seams, SeamStitchesDependency);
@@ -468,7 +468,7 @@ namespace UE::Chaos::ClothAsset
 			FabricLayer = &ManagedArrayCollection->AddAttribute<int32>(ClothCollectionAttribute::FabricLayer, ClothCollectionGroup::Fabrics);
 			FabricCollisionThickness = &ManagedArrayCollection->AddAttribute<float>(ClothCollectionAttribute::FabricCollisionThickness, ClothCollectionGroup::Fabrics);
 		}
-		
+
 		// Sim Faces Group
 		SimIndices2D = &ManagedArrayCollection->AddAttribute<FIntVector3>(ClothCollectionAttribute::SimIndices2D, ClothCollectionGroup::SimFaces, SimVertices2DDependency);
 		SimIndices3D = &ManagedArrayCollection->AddAttribute<FIntVector3>(ClothCollectionAttribute::SimIndices3D, ClothCollectionGroup::SimFaces, SimVertices3DDependency);
diff --git a/Engine/Plugins/ChaosClothAsset/Source/ChaosClothAsset/Private/ChaosClothAsset/CollectionClothFacade.cpp b/Engine/Plugins/ChaosClothAsset/Source/ChaosClothAsset/Private/ChaosClothAsset/CollectionClothFacade.cpp
--- a/Engine/Plugins/ChaosClothAsset/Source/ChaosClothAsset/Private/ChaosClothAsset/CollectionClothFacade.cpp
+++ b/Engine/Plugins/ChaosClothAsset/Source/ChaosClothAsset/Private/ChaosClothAsset/CollectionClothFacade.cpp
@@ -21,7 +21,7 @@ namespace UE::Chaos::ClothAsset
 
 	bool FCollectionClothConstFacade::IsValid(EClothCollectionOptionalSchemas OptionalSchemas) const
 	{
-		return ClothCollection->IsValid(OptionalSchemas) && ClothCollection->GetNumElements(ClothCollectionGroup::Lods) == 1;
+		return ClothCollection->IsValid(OptionalSchemas) && (EnumHasAnyFlags(OptionalSchemas, EClothCollectionOptionalSchemas::CookedOnly) || ClothCollection->GetNumElements(ClothCollectionGroup::Lods) == 1);
 	}
 
 	bool FCollectionClothConstFacade::HasValidSimulationData() const
@@ -788,15 +788,15 @@ namespace UE::Chaos::ClothAsset
 
 	void FCollectionClothFacade::InitializeCookedOnly(const FCollectionClothConstFacade& Other)
 	{
-		Reset();
+		constexpr bool bCookedOnly = true;
+		SetNumSimMorphTargets(0, bCookedOnly);
 
 		// Sim Morph Targets Group
-		const int32 StartNumSimMorphTargets = GetNumSimMorphTargets();
 		const int32 OtherNumSimMorphTargets = Other.GetNumSimMorphTargets();
-		SetNumSimMorphTargets(StartNumSimMorphTargets + OtherNumSimMorphTargets);
+		SetNumSimMorphTargets(OtherNumSimMorphTargets, bCookedOnly);
 		for (int32 TargetIndex = 0; TargetIndex < OtherNumSimMorphTargets; ++TargetIndex)
 		{
-			GetSimMorphTarget(TargetIndex + StartNumSimMorphTargets).Initialize(Other.GetSimMorphTarget(TargetIndex), 0);
+			GetSimMorphTarget(TargetIndex).Initialize(Other.GetSimMorphTarget(TargetIndex), 0);
 		}
 	}
 
@@ -997,11 +997,12 @@ namespace UE::Chaos::ClothAsset
 		GetClothCollection()->RemoveElements(ClothCollectionGroup::SimPatterns, SortedDeletionList);
 	}
 
-	void FCollectionClothFacade::SetNumSimMorphTargets(int32 InNumMorphTargets)
+	void FCollectionClothFacade::SetNumSimMorphTargets(int32 InNumMorphTargets, bool bCookedOnly)
 	{
-		if (!IsValid(EClothCollectionOptionalSchemas::SimMorphTargets))
+		const EClothCollectionOptionalSchemas OptionalSchemas = bCookedOnly ? EClothCollectionOptionalSchemas::CookedOnly : EClothCollectionOptionalSchemas::SimMorphTargets;
+		if (!IsValid(OptionalSchemas))
 		{
-			DefineSchema(EClothCollectionOptionalSchemas::SimMorphTargets);
+			DefineSchema(OptionalSchemas);
 		}
 
 		const int32 NumMorphTargets = GetNumSimMorphTargets();
diff --git a/Engine/Plugins/ChaosClothAsset/Source/ChaosClothAsset/Public/ChaosClothAsset/CollectionClothFacade.h b/Engine/Plugins/ChaosClothAsset/Source/ChaosClothAsset/Public/ChaosClothAsset/CollectionClothFacade.h
--- a/Engine/Plugins/ChaosClothAsset/Source/ChaosClothAsset/Public/ChaosClothAsset/CollectionClothFacade.h
+++ b/Engine/Plugins/ChaosClothAsset/Source/ChaosClothAsset/Public/ChaosClothAsset/CollectionClothFacade.h
@@ -307,7 +307,7 @@ namespace UE::Chaos::ClothAsset
 
 		//~ Sim Morph Targets Group
 		/** Set the new number of sim morph targets to this cloth. */
-		UE_API void SetNumSimMorphTargets(int32 NumMorphTargets);
+		UE_API void SetNumSimMorphTargets(int32 NumMorphTargets, bool bCookedOnly = false);
 		/** Add a sim morph target to this cloth and return its index in the sim morph target pattern list */
 		UE_API int32 AddSimMorphTarget();
 		/** Return a sim morph target facade for the specified morph target index. */
