Fix CIS error in previous cloth morph targets CL.

diff --git a/Engine/Plugins/ChaosClothAsset/Source/ChaosClothAssetEngine/Private/ChaosClothAsset/ClothSimulationMesh.cpp b/Engine/Plugins/ChaosClothAsset/Source/ChaosClothAssetEngine/Private/ChaosClothAsset/ClothSimulationMesh.cpp
--- a/Engine/Plugins/ChaosClothAsset/Source/ChaosClothAssetEngine/Private/ChaosClothAsset/ClothSimulationMesh.cpp
+++ b/Engine/Plugins/ChaosClothAsset/Source/ChaosClothAssetEngine/Private/ChaosClothAsset/ClothSimulationMesh.cpp
@@ -31,7 +31,7 @@ namespace UE::Chaos::ClothAsset
 		ClothFacades.Reserve(ManagedArrayCollections.Num());
 		for (const TSharedRef<const FManagedArrayCollection>& Collection : ManagedArrayCollections)
 		{
-			ClothFacades.Emplace(Collection);
+			ClothFacades.Emplace(MakeShared<FCollectionClothConstFacade>(Collection));
 		}
 	}
 
@@ -248,14 +248,14 @@ namespace UE::Chaos::ClothAsset
 
 	int32 FClothSimulationMesh::FindMorphTargetByName(int32 LODIndex, const FString& Name) const
 	{
-		return ClothFacades.IsValidIndex(LODIndex) ? ClothFacades[LODIndex].FindSimMorphTargetIndexByName(Name) : INDEX_NONE;
+		return ClothFacades.IsValidIndex(LODIndex) ? ClothFacades[LODIndex]->FindSimMorphTargetIndexByName(Name) : INDEX_NONE;
 	}
 
 	TConstArrayView<FVector3f> FClothSimulationMesh::GetMorphTargetPositionDeltas(int32 LODIndex, int32 MorphTargetIndex) const
 	{
-		if (ClothFacades.IsValidIndex(LODIndex) && (MorphTargetIndex >= 0 && MorphTargetIndex < ClothFacades[LODIndex].GetNumSimMorphTargets()))
+		if (ClothFacades.IsValidIndex(LODIndex) && (MorphTargetIndex >= 0 && MorphTargetIndex < ClothFacades[LODIndex]->GetNumSimMorphTargets()))
 		{
-			FCollectionClothSimMorphTargetConstFacade MorphTargetFacade = ClothFacades[LODIndex].GetSimMorphTarget(MorphTargetIndex);
+			FCollectionClothSimMorphTargetConstFacade MorphTargetFacade = ClothFacades[LODIndex]->GetSimMorphTarget(MorphTargetIndex);
 			return MorphTargetFacade.GetSimMorphTargetPositionDelta();
 		}
 		return TConstArrayView<FVector3f>();
@@ -263,9 +263,9 @@ namespace UE::Chaos::ClothAsset
 
 	TConstArrayView<FVector3f> FClothSimulationMesh::GetMorphTargetTangentZDeltas(int32 LODIndex, int32 MorphTargetIndex) const
 	{
-		if (ClothFacades.IsValidIndex(LODIndex) && (MorphTargetIndex >= 0 && MorphTargetIndex < ClothFacades[LODIndex].GetNumSimMorphTargets()))
+		if (ClothFacades.IsValidIndex(LODIndex) && (MorphTargetIndex >= 0 && MorphTargetIndex < ClothFacades[LODIndex]->GetNumSimMorphTargets()))
 		{
-			FCollectionClothSimMorphTargetConstFacade MorphTargetFacade = ClothFacades[LODIndex].GetSimMorphTarget(MorphTargetIndex);
+			FCollectionClothSimMorphTargetConstFacade MorphTargetFacade = ClothFacades[LODIndex]->GetSimMorphTarget(MorphTargetIndex);
 			return MorphTargetFacade.GetSimMorphTargetTangentZDelta();
 		}
 		return TConstArrayView<FVector3f>();
@@ -273,9 +273,9 @@ namespace UE::Chaos::ClothAsset
 
 	TConstArrayView<int32> FClothSimulationMesh::GetMorphTargetIndices(int32 LODIndex, int32 MorphTargetIndex) const
 	{
-		if (ClothFacades.IsValidIndex(LODIndex) && (MorphTargetIndex >= 0 && MorphTargetIndex < ClothFacades[LODIndex].GetNumSimMorphTargets()))
+		if (ClothFacades.IsValidIndex(LODIndex) && (MorphTargetIndex >= 0 && MorphTargetIndex < ClothFacades[LODIndex]->GetNumSimMorphTargets()))
 		{
-			FCollectionClothSimMorphTargetConstFacade MorphTargetFacade = ClothFacades[LODIndex].GetSimMorphTarget(MorphTargetIndex);
+			FCollectionClothSimMorphTargetConstFacade MorphTargetFacade = ClothFacades[LODIndex]->GetSimMorphTarget(MorphTargetIndex);
 			return MorphTargetFacade.GetSimMorphTargetSimVertex3DIndex();
 		}
 		return TConstArrayView<int32>();
diff --git a/Engine/Plugins/ChaosClothAsset/Source/ChaosClothAssetEngine/Private/ChaosClothAsset/ClothSimulationMesh.h b/Engine/Plugins/ChaosClothAsset/Source/ChaosClothAssetEngine/Private/ChaosClothAsset/ClothSimulationMesh.h
--- a/Engine/Plugins/ChaosClothAsset/Source/ChaosClothAssetEngine/Private/ChaosClothAsset/ClothSimulationMesh.h
+++ b/Engine/Plugins/ChaosClothAsset/Source/ChaosClothAssetEngine/Private/ChaosClothAsset/ClothSimulationMesh.h
@@ -68,6 +68,6 @@ namespace UE::Chaos::ClothAsset
 		const FChaosClothSimulationModel& ClothSimulationModel;
 		const FClothSimulationContext& ClothSimulationContext;
 		const TArray<TSharedRef<const FManagedArrayCollection>>& ManagedArrayCollections;
-		TArray<FCollectionClothConstFacade> ClothFacades;
+		TArray<TSharedRef<FCollectionClothConstFacade>> ClothFacades;
 	};
 }
