Include validating against RefSkeleton when checking if Sim and Render skin weights are valid when building ClothAssets.

diff --git a/Engine/Plugins/ChaosClothAsset/Source/ChaosClothAssetEngine/Private/ChaosClothAsset/ClothAsset.cpp b/Engine/Plugins/ChaosClothAsset/Source/ChaosClothAssetEngine/Private/ChaosClothAsset/ClothAsset.cpp
--- a/Engine/Plugins/ChaosClothAsset/Source/ChaosClothAssetEngine/Private/ChaosClothAsset/ClothAsset.cpp
+++ b/Engine/Plugins/ChaosClothAsset/Source/ChaosClothAssetEngine/Private/ChaosClothAsset/ClothAsset.cpp
@@ -46,8 +46,13 @@ namespace UE::Chaos::ClothAsset::Private
 		bClothCollectionOnlyCookRequiredFacades,
 		TEXT("Default setting for culling managed arrays on the cloth collection during the cook. Default[true]"));
 
-	static bool HasValidSkinweights(const TConstArrayView<TArray<int32>> BoneIndices, TConstArrayView<TArray<float>> BoneWeights)
+	static bool HasValidSkinweights(const TConstArrayView<TArray<int32>> BoneIndices, TConstArrayView<TArray<float>> BoneWeights, const FReferenceSkeleton* RefSkeleton)
 	{
+		if (!RefSkeleton)
+		{
+			return false;
+		}
+
 		check(BoneIndices.Num() == BoneWeights.Num());
 		for (int32 Index = 0; Index < BoneIndices.Num(); ++Index)
 		{
@@ -55,6 +60,14 @@ namespace UE::Chaos::ClothAsset::Private
 			{
 				return false;
 			}
+			for (const int32 BoneIndex : BoneIndices[Index])
+			{
+				if (!RefSkeleton->IsValidIndex(BoneIndex))
+				{
+					return false;
+				}
+			}
+
 		}
 		return true;
 	}
@@ -238,8 +251,8 @@ void UChaosClothAsset::Serialize(FArchive& Ar)
 			const FCollectionClothConstFacade ClothConstFacade(ClothCollection);
 			if (ClothConstFacade.IsValid())
 			{
-				const bool bHasValidSimSkinweights = Private::HasValidSkinweights(ClothConstFacade.GetSimBoneIndices(), ClothConstFacade.GetSimBoneWeights());
-				const bool bHasValidRenderSkinweights = Private::HasValidSkinweights(ClothConstFacade.GetRenderBoneIndices(), ClothConstFacade.GetRenderBoneWeights());
+				const bool bHasValidSimSkinweights = Private::HasValidSkinweights(ClothConstFacade.GetSimBoneIndices(), ClothConstFacade.GetSimBoneWeights(), &GetRefSkeleton());
+				const bool bHasValidRenderSkinweights = Private::HasValidSkinweights(ClothConstFacade.GetRenderBoneIndices(), ClothConstFacade.GetRenderBoneWeights(), &GetRefSkeleton());
 				if (!bHasValidSimSkinweights || !bHasValidRenderSkinweights)
 				{
 					TSharedRef<FManagedArrayCollection> NewClothCollection = MakeShared<FManagedArrayCollection>(*ClothCollection);
@@ -494,8 +507,6 @@ void UChaosClothAsset::Build(
 
 	// Iterate through the LODs
 	FString PhysicsAssetPathName;
-	FString SkeletalMeshPathName;
-
 	for (int32 LodIndex = 0; LodIndex < InClothCollections.Num(); ++LodIndex)
 	{
 		// New LOD
@@ -585,12 +596,18 @@ void UChaosClothAsset::Build(
 		{
 			using namespace ::Chaos::Softs;
 			PhysicsAssetPathName = InClothFacade.GetPhysicsAssetPathName();
-			SkeletalMeshPathName = InClothFacade.GetSkeletalMeshPathName();
+			const FString SkeletalMeshPathName = InClothFacade.GetSkeletalMeshPathName(); 
+			USkeletalMesh* const SkeletalMesh = SkeletalMeshPathName.IsEmpty() ? nullptr :
+				LoadObject<USkeletalMesh>(nullptr, *SkeletalMeshPathName, nullptr, LOAD_None, nullptr);
+
+			// Set reference skeleton
+			SetSkeleton(SkeletalMesh ? SkeletalMesh->GetSkeleton() : nullptr); // For completion only, this is not being used and might mismatch the skeletal mesh's reference skeleton
+			Super::SetReferenceSkeleton(SkeletalMesh ? &SkeletalMesh->GetRefSkeleton() : nullptr);
 		}
 
 		// Fix the skeleton mesh binding if needed, which can cause crashes in the render code, or make the sim mesh disappear
-		const bool bHasValidSimSkinweights = Private::HasValidSkinweights(InClothFacade.GetSimBoneIndices(), InClothFacade.GetSimBoneWeights());
-		const bool bHasValidRenderSkinweights = Private::HasValidSkinweights(InClothFacade.GetRenderBoneIndices(), InClothFacade.GetRenderBoneWeights());
+		const bool bHasValidSimSkinweights = Private::HasValidSkinweights(InClothFacade.GetSimBoneIndices(), InClothFacade.GetSimBoneWeights(), &GetRefSkeleton());
+		const bool bHasValidRenderSkinweights = Private::HasValidSkinweights(InClothFacade.GetRenderBoneIndices(), InClothFacade.GetRenderBoneWeights(), &GetRefSkeleton());
 		if (!ensureAlwaysMsgf(bHasValidSimSkinweights && bHasValidRenderSkinweights, TEXT("A Dataflow node, likely an import node, has generated missing or invalid skin weights in this collection LOD. This must be fixed ASAP!")))
 		{
 			FClothGeometryTools::BindMeshToRootBone(ClothCollection, !bHasValidSimSkinweights, !bHasValidRenderSkinweights);
@@ -608,13 +625,6 @@ void UChaosClothAsset::Build(
 		OutClothCollections.Emplace(MoveTemp(ClothCollection));
 	}
 
-	// Set reference skeleton
-	USkeletalMesh* const SkeletalMesh = SkeletalMeshPathName.IsEmpty() ? nullptr :
-		LoadObject<USkeletalMesh>(nullptr, *SkeletalMeshPathName, nullptr, LOAD_None, nullptr);
-
-	SetSkeleton(SkeletalMesh ? SkeletalMesh->GetSkeleton() : nullptr); // For completion only, this is not being used and might mismatch the skeletal mesh's reference skeleton
-	Super::SetReferenceSkeleton(SkeletalMesh ? &SkeletalMesh->GetRefSkeleton() : nullptr);
-
 	// Set physics asset (note: the cloth asset's physics asset is only replaced if a collection path name is found valid)
 	PhysicsAsset = !PhysicsAssetPathName.IsEmpty() ?
 		LoadObject<UPhysicsAsset>(this, *PhysicsAssetPathName, nullptr, LOAD_None, nullptr) :
