ChaosCloth: Fix ForceNextUpdateTeleport and ForceNextUpdateTeleportAndReset to not stick.

diff --git a/Engine/Plugins/ChaosClothAsset/Source/ChaosClothAssetEngine/Private/ChaosClothAsset/ClothComponent.cpp b/Engine/Plugins/ChaosClothAsset/Source/ChaosClothAssetEngine/Private/ChaosClothAsset/ClothComponent.cpp
--- a/Engine/Plugins/ChaosClothAsset/Source/ChaosClothAssetEngine/Private/ChaosClothAsset/ClothComponent.cpp
+++ b/Engine/Plugins/ChaosClothAsset/Source/ChaosClothAssetEngine/Private/ChaosClothAsset/ClothComponent.cpp
@@ -661,9 +661,11 @@ void UChaosClothComponent::StartNewParallelSimulation(float DeltaTime)
 void UChaosClothComponent::UpdateClothTeleport()
 {
 	const FMatrix CurRootBoneMat = GetBoneMatrix(0);
-	ClothTeleportMode = bTeleport ? (bReset ? EClothingTeleportMode::TeleportAndReset : EClothingTeleportMode::Teleport) : EClothingTeleportMode::None;
-	ClothTeleportMode = UE::ClothingSimulation::TeleportHelpers::CalculateClothingTeleport(ClothTeleportMode, CurRootBoneMat, PrevRootBoneMatrix, bReset, ClothTeleportDistThresholdSquared, ClothTeleportCosineThresholdInRad);
+	ClothTeleportMode = (bTeleport || bTeleportOnce) ? ((bReset || bResetOnce) ? EClothingTeleportMode::TeleportAndReset : EClothingTeleportMode::Teleport) : EClothingTeleportMode::None;
+	ClothTeleportMode = UE::ClothingSimulation::TeleportHelpers::CalculateClothingTeleport(ClothTeleportMode, CurRootBoneMat, PrevRootBoneMatrix, (bReset||bResetOnce), ClothTeleportDistThresholdSquared, ClothTeleportCosineThresholdInRad);
 	PrevRootBoneMatrix = CurRootBoneMat;
+	bTeleportOnce = false;
+	bResetOnce = false;
 }
 
 void UChaosClothComponent::HandleExistingParallelSimulation()
diff --git a/Engine/Plugins/ChaosClothAsset/Source/ChaosClothAssetEngine/Public/ChaosClothAsset/ClothComponent.h b/Engine/Plugins/ChaosClothAsset/Source/ChaosClothAssetEngine/Public/ChaosClothAsset/ClothComponent.h
--- a/Engine/Plugins/ChaosClothAsset/Source/ChaosClothAssetEngine/Public/ChaosClothAsset/ClothComponent.h
+++ b/Engine/Plugins/ChaosClothAsset/Source/ChaosClothAssetEngine/Public/ChaosClothAsset/ClothComponent.h
@@ -90,17 +90,17 @@ public:
 
 	/** Teleport the cloth particles to the new reference bone location keeping pose and velocities prior to advancing the simulation. */
 	UFUNCTION(BlueprintCallable, Category = "ClothComponent", Meta = (Keywords = "Chaos Cloth Teleport"))
-	void ForceNextUpdateTeleport() { bTeleport = true; bReset = false; }
+	void ForceNextUpdateTeleport() { bTeleportOnce = true; bResetOnce = false; }
 
 	/** Teleport the cloth particles to the new reference bone location while reseting the pose and velocities prior to advancing the simulation. */
 	UFUNCTION(BlueprintCallable, Category = "ClothComponent", Meta = (Keywords = "Chaos Cloth Teleport Reset"))
-	void ForceNextUpdateTeleportAndReset() { bTeleport = bReset = true; }
+	void ForceNextUpdateTeleportAndReset() { bTeleportOnce = bResetOnce = true; }
 
 	/** Return whether teleport is currently requested. Use GetClothTeleportMode to get teleport mode including any auto-teleport based on teleport thresholds. */
-	bool NeedsTeleport() const { return bTeleport; }
+	bool NeedsTeleport() const { return bTeleport || bTeleportOnce; }
 
 	/** Return whether reseting the pose is currently requested. Use GetClothTeleportMode to get teleport mode including any auto-teleport based on teleport thresholds*/
-	bool NeedsReset() const { return bReset; }
+	bool NeedsReset() const { return bReset || bResetOnce; }
 
 	/** Get currently calculated teleport mode. */
 	EClothingTeleportMode GetClothTeleportMode() const { return ClothTeleportMode; }
@@ -403,6 +403,12 @@ private:
 	/** Currently calculated ClothTeleport based on bTeleport, bReset as well as any teleport calculated based on TeleportDistanceThreshold and TeleportRotationThreshold */
 	EClothingTeleportMode ClothTeleportMode;
 
+	/** Like bTeleport, but cleared every frame */
+	bool bTeleportOnce = false;
+
+	/** Like bReset, but cleared every frame */
+	bool bResetOnce = false;
+
 	/**
 	 * Simulation properties per cloth model. 
 	 * Transient because it doesn't need to be serialized but contains pointers to UObject that needs not be garbage collected.
