Chaos Cloth Asset - Fixed disabling the properties set by Dataflow in the details panel editor.

#rnx

diff --git a/Engine/Plugins/ChaosClothAsset/Source/ChaosClothAssetEngine/Public/ChaosClothAsset/ClothAssetBase.h b/Engine/Plugins/ChaosClothAsset/Source/ChaosClothAssetEngine/Public/ChaosClothAsset/ClothAssetBase.h
--- a/Engine/Plugins/ChaosClothAsset/Source/ChaosClothAssetEngine/Public/ChaosClothAsset/ClothAssetBase.h
+++ b/Engine/Plugins/ChaosClothAsset/Source/ChaosClothAssetEngine/Public/ChaosClothAsset/ClothAssetBase.h
@@ -359,7 +359,7 @@ protected:
 
 #if WITH_EDITORONLY_DATA
 	/** Property used in various edit conditions to enable property changes when this asset has no Dataflow asset. */
-	UPROPERTY(Transient)
+	UPROPERTY(VisibleAnywhere, Category = Dataflow)
 	bool bHasDataflowAsset = false;
 #endif
 
