Add ability to define custom regions of cloth that can resize together.

Currently just supporting trilinear interpolation of a region.

Updated the ResizeOutfitTemplate graph to be able to do the custom resizing of the data exists.

#virtualized

diff --git a/Engine/Plugins/ChaosClothAsset/ChaosClothAsset.uplugin b/Engine/Plugins/ChaosClothAsset/ChaosClothAsset.uplugin
--- a/Engine/Plugins/ChaosClothAsset/ChaosClothAsset.uplugin
+++ b/Engine/Plugins/ChaosClothAsset/ChaosClothAsset.uplugin
@@ -38,6 +38,10 @@
 		{
 			"Name": "Dataflow",
 			"Enabled": true
+		},
+		{
+			"Name": "MeshResizing",
+			"Enabled": true
 		}
 	]
 }
diff --git a/Engine/Plugins/ChaosClothAsset/Source/ChaosClothAsset/ChaosClothAsset.Build.cs b/Engine/Plugins/ChaosClothAsset/Source/ChaosClothAsset/ChaosClothAsset.Build.cs
--- a/Engine/Plugins/ChaosClothAsset/Source/ChaosClothAsset/ChaosClothAsset.Build.cs
+++ b/Engine/Plugins/ChaosClothAsset/Source/ChaosClothAsset/ChaosClothAsset.Build.cs
@@ -14,6 +14,7 @@ public class ChaosClothAsset : ModuleRules
 				"Core",
 				"GeometryCore",
 				"MeshConversion",
+				"MeshResizingCore",
 				"ClothingSystemRuntimeCommon",
 				"RenderCore"
 			}
diff --git a/Engine/Plugins/ChaosClothAsset/Source/ChaosClothAsset/Private/ChaosClothAsset/ClothCollection.cpp b/Engine/Plugins/ChaosClothAsset/Source/ChaosClothAsset/Private/ChaosClothAsset/ClothCollection.cpp
--- a/Engine/Plugins/ChaosClothAsset/Source/ChaosClothAsset/Private/ChaosClothAsset/ClothCollection.cpp
+++ b/Engine/Plugins/ChaosClothAsset/Source/ChaosClothAsset/Private/ChaosClothAsset/ClothCollection.cpp
@@ -99,7 +99,8 @@ namespace UE::Chaos::ClothAsset::Private
 		ClothCollectionAttribute::SimBoneIndices,
 		ClothCollectionAttribute::SimBoneWeights,
 		ClothCollectionAttribute::SimVertex2DLookup,
-		ClothCollectionAttribute::SeamStitchLookup
+		ClothCollectionAttribute::SeamStitchLookup,
+		ClothCollectionAttribute::SimCustomResizingBlend,
 	};
 
 	//~ Sim Morph Targets Group
@@ -140,7 +141,15 @@ namespace UE::Chaos::ClothAsset::Private
 		ClothCollectionAttribute::RenderDeformerTangentBaryCoordsAndDist,
 		ClothCollectionAttribute::RenderDeformerSimIndices3D,
 		ClothCollectionAttribute::RenderDeformerWeight,
-		ClothCollectionAttribute::RenderDeformerSkinningBlend
+		ClothCollectionAttribute::RenderDeformerSkinningBlend,
+		ClothCollectionAttribute::RenderCustomResizingBlend,
+	};
+
+	//~ Group Resizing Group
+	static const TArray<FName> GroupResizingGroupAttributes =
+	{
+		ClothCollectionAttribute::CustomResizingRegionSet,
+		ClothCollectionAttribute::CustomResizingRegionType
 	};
 
 	// All fixed attributes for this collection
@@ -160,12 +169,15 @@ namespace UE::Chaos::ClothAsset::Private
 		{ ClothCollectionGroup::RenderVertices, RenderVerticesGroupAttributes },
 		{ ClothCollectionGroup::Fabrics, SimFabricGroupAttributes },
 		{ ClothCollectionGroup::Solvers, SolverGroupAttributes },
+		{ ClothCollectionGroup::CustomResizingRegions, GroupResizingGroupAttributes },
 	};
 
 	// All paintable attributes for this collection
 	static const TMap<FName, TArray<FName>> UserAccessibleAttributeNamesMap =
 	{
+		{ ClothCollectionGroup::SimVertices3D, { ClothCollectionAttribute::SimCustomResizingBlend } },  // Can be painted if required
 		{ ClothCollectionGroup::RenderVertices, { ClothCollectionAttribute::RenderDeformerSkinningBlend } },  // Can be painted if required
+		{ ClothCollectionGroup::RenderVertices, { ClothCollectionAttribute::RenderCustomResizingBlend } },  // Can be painted if required
 	};
 }  // End namespace UE::Chaos::ClothAsset::Private
 
@@ -238,6 +250,7 @@ namespace UE::Chaos::ClothAsset
 		TetherReferenceLength = ManagedArrayCollection->FindAttribute<TArray<float>>(ClothCollectionAttribute::TetherReferenceLength, ClothCollectionGroup::SimVertices3D);
 		SimVertex2DLookup = ManagedArrayCollection->FindAttribute<TArray<int32>>(ClothCollectionAttribute::SimVertex2DLookup, ClothCollectionGroup::SimVertices3D);
 		SeamStitchLookup = ManagedArrayCollection->FindAttribute<TArray<int32>>(ClothCollectionAttribute::SeamStitchLookup, ClothCollectionGroup::SimVertices3D);
+		SimCustomResizingBlend = ManagedArrayCollection->FindAttribute<float>(ClothCollectionAttribute::SimCustomResizingBlend, ClothCollectionGroup::SimVertices3D);
 
 		// Sim Morph Targets Group
 		SimMorphTargetName = ManagedArrayCollection->FindAttribute<FString>(ClothCollectionAttribute::SimMorphTargetName, ClothCollectionGroup::SimMorphTargets);
@@ -267,6 +280,11 @@ namespace UE::Chaos::ClothAsset
 		RenderDeformerSimIndices3D = ManagedArrayCollection->FindAttribute<TArray<FIntVector3>>(ClothCollectionAttribute::RenderDeformerSimIndices3D, ClothCollectionGroup::RenderVertices);
 		RenderDeformerWeight = ManagedArrayCollection->FindAttribute<TArray<float>>(ClothCollectionAttribute::RenderDeformerWeight, ClothCollectionGroup::RenderVertices);
 		RenderDeformerSkinningBlend = ManagedArrayCollection->FindAttribute<float>(ClothCollectionAttribute::RenderDeformerSkinningBlend, ClothCollectionGroup::RenderVertices);
+		RenderCustomResizingBlend = ManagedArrayCollection->FindAttribute<float>(ClothCollectionAttribute::RenderCustomResizingBlend, ClothCollectionGroup::RenderVertices);
+
+		// Group Resizing Group
+		CustomResizingRegionSet = ManagedArrayCollection->FindAttribute<FString>(ClothCollectionAttribute::CustomResizingRegionSet, ClothCollectionGroup::CustomResizingRegions);
+		CustomResizingRegionType = ManagedArrayCollection->FindAttribute<int32>(ClothCollectionAttribute::CustomResizingRegionType, ClothCollectionGroup::CustomResizingRegions);
 	}
 
 	bool FClothCollection::IsValid(EClothCollectionOptionalSchemas OptionalSchemas) const
@@ -352,6 +370,8 @@ namespace UE::Chaos::ClothAsset
 			TetherReferenceLength &&
 			SimVertex2DLookup &&
 			SeamStitchLookup &&
+			(!EnumHasAnyFlags(OptionalSchemas, EClothCollectionOptionalSchemas::Resizing) ||
+				SimCustomResizingBlend) &&
 
 			// Sim Morph Targets Group
 			(!EnumHasAnyFlags(OptionalSchemas, EClothCollectionOptionalSchemas::SimMorphTargets) ||
@@ -379,11 +399,18 @@ namespace UE::Chaos::ClothAsset
 			RenderBoneWeights &&
 			(!EnumHasAnyFlags(OptionalSchemas, EClothCollectionOptionalSchemas::RenderDeformer) ||
 				(RenderDeformerPositionBaryCoordsAndDist &&
-					RenderDeformerNormalBaryCoordsAndDist &&
-					RenderDeformerTangentBaryCoordsAndDist &&
-					RenderDeformerSimIndices3D &&
-					RenderDeformerWeight &&
-					RenderDeformerSkinningBlend));
+				RenderDeformerNormalBaryCoordsAndDist &&
+				RenderDeformerTangentBaryCoordsAndDist &&
+				RenderDeformerSimIndices3D &&
+				RenderDeformerWeight &&
+				RenderDeformerSkinningBlend)) &&
+			(!EnumHasAnyFlags(OptionalSchemas, EClothCollectionOptionalSchemas::Resizing) ||
+				RenderCustomResizingBlend) &&
+				
+			// Group Resizing Group
+			(!EnumHasAnyFlags(OptionalSchemas, EClothCollectionOptionalSchemas::Resizing) ||
+				(CustomResizingRegionSet &&
+					CustomResizingRegionType));
 	}
 
 	void FClothCollection::DefineSchema(EClothCollectionOptionalSchemas OptionalSchemas)
@@ -526,6 +553,19 @@ namespace UE::Chaos::ClothAsset
 			RenderDeformerWeight = &ManagedArrayCollection->AddAttribute<TArray<float>>(ClothCollectionAttribute::RenderDeformerWeight, ClothCollectionGroup::RenderVertices);
 			RenderDeformerSkinningBlend = &ManagedArrayCollection->AddAttribute<float>(ClothCollectionAttribute::RenderDeformerSkinningBlend, ClothCollectionGroup::RenderVertices);
 		}
+		
+		if (EnumHasAnyFlags(OptionalSchemas, EClothCollectionOptionalSchemas::Resizing))
+		{
+			// Sim Vertices 3D Group
+			SimCustomResizingBlend = &ManagedArrayCollection->AddAttribute<float>(ClothCollectionAttribute::SimCustomResizingBlend, ClothCollectionGroup::SimVertices3D);
+
+			// Render Vertices Group
+			RenderCustomResizingBlend = &ManagedArrayCollection->AddAttribute<float>(ClothCollectionAttribute::RenderCustomResizingBlend, ClothCollectionGroup::RenderVertices);
+
+			// Group Resizing Group
+			CustomResizingRegionSet = &ManagedArrayCollection->AddAttribute<FString>(ClothCollectionAttribute::CustomResizingRegionSet, ClothCollectionGroup::CustomResizingRegions);
+			CustomResizingRegionType = &ManagedArrayCollection->AddAttribute<int32>(ClothCollectionAttribute::CustomResizingRegionType, ClothCollectionGroup::CustomResizingRegions);
+		}
 	}
 
 	int32 FClothCollection::GetNumElements(const FName& GroupName) const
diff --git a/Engine/Plugins/ChaosClothAsset/Source/ChaosClothAsset/Private/ChaosClothAsset/ClothCollection.h b/Engine/Plugins/ChaosClothAsset/Source/ChaosClothAsset/Private/ChaosClothAsset/ClothCollection.h
--- a/Engine/Plugins/ChaosClothAsset/Source/ChaosClothAsset/Private/ChaosClothAsset/ClothCollection.h
+++ b/Engine/Plugins/ChaosClothAsset/Source/ChaosClothAsset/Private/ChaosClothAsset/ClothCollection.h
@@ -217,6 +217,7 @@ namespace UE::Chaos::ClothAsset
 		const TManagedArray<TArray<float>>* GetTetherReferenceLength() const { return TetherReferenceLength; }
 		const TManagedArray<TArray<int32>>* GetSimVertex2DLookup() const { return SimVertex2DLookup; }
 		const TManagedArray<TArray<int32>>* GetSeamStitchLookup() const { return SeamStitchLookup; }
+		const TManagedArray<float>* GetSimCustomResizingBlend() const { return SimCustomResizingBlend; }
 
 		//~ Sim Morph Targets Group
 		const TManagedArray<FString>* GetSimMorphTargetName() const { return SimMorphTargetName; }
@@ -246,6 +247,11 @@ namespace UE::Chaos::ClothAsset
 		const TManagedArray<TArray<FIntVector3>>* GetRenderDeformerSimIndices3D() const { return RenderDeformerSimIndices3D; }
 		const TManagedArray<TArray<float>>* GetRenderDeformerWeight() const { return RenderDeformerWeight; }
 		const TManagedArray<float>* GetRenderDeformerSkinningBlend() const { return RenderDeformerSkinningBlend; }
+		const TManagedArray<float>* GetRenderCustomResizingBlend() const { return RenderCustomResizingBlend; }
+
+		//~ Custom Resizing Regions Group
+		const TManagedArray<FString>* GetCustomResizingRegionSet() const { return CustomResizingRegionSet; }
+		const TManagedArray<int32>* GetCustomResizingRegionType() const { return CustomResizingRegionType; }
 
 		//~ LODs Group (There should be only one LOD per ClothCollection)
 		TManagedArray<FString>* GetPhysicsAssetPathName(){ return PhysicsAssetPathName; }
@@ -309,6 +315,7 @@ namespace UE::Chaos::ClothAsset
 		TManagedArray<TArray<float>>* GetTetherReferenceLength() { return TetherReferenceLength; }
 		TManagedArray<TArray<int32>>* GetSimVertex2DLookup() { return SimVertex2DLookup; }
 		TManagedArray<TArray<int32>>* GetSeamStitchLookup() { return SeamStitchLookup; }
+		TManagedArray<float>* GetSimCustomResizingBlend() { return SimCustomResizingBlend; }
 
 		//~ Sim Morph Targets Group
 		TManagedArray<FString>* GetSimMorphTargetName() { return SimMorphTargetName; }
@@ -338,6 +345,11 @@ namespace UE::Chaos::ClothAsset
 		TManagedArray<TArray<FIntVector3>>* GetRenderDeformerSimIndices3D() { return RenderDeformerSimIndices3D; }
 		TManagedArray<TArray<float>>* GetRenderDeformerWeight() { return RenderDeformerWeight; }
 		TManagedArray<float>* GetRenderDeformerSkinningBlend() { return RenderDeformerSkinningBlend; }
+		TManagedArray<float>* GetRenderCustomResizingBlend() { return RenderCustomResizingBlend; }
+
+		//~ Group Resizing Group
+		TManagedArray<FString>* GetCustomResizingRegionSet() { return CustomResizingRegionSet; }
+		TManagedArray<int32>* GetCustomResizingRegionType() { return CustomResizingRegionType; }
 
 	private:
 		//~ Cloth collection
@@ -405,6 +417,7 @@ namespace UE::Chaos::ClothAsset
 		TManagedArray<TArray<float>>* TetherReferenceLength;
 		TManagedArray<TArray<int32>>* SimVertex2DLookup; // Lookup into corresponding 2D vertices
 		TManagedArray<TArray<int32>>* SeamStitchLookup; // Lookup into any seam stitches which weld this vertex
+		TManagedArray<float>* SimCustomResizingBlend; // Resizing weight map. How much to use group resizing result (1) vs rbf interpolation result (0) when doing sim mesh resizing
 
 		//~ Sim Morph Targets Group
 		TManagedArray<FString>* SimMorphTargetName;
@@ -434,6 +447,11 @@ namespace UE::Chaos::ClothAsset
 		TManagedArray<TArray<FIntVector3>>* RenderDeformerSimIndices3D;  // The source mesh triangle
 		TManagedArray<TArray<float>>* RenderDeformerWeight;  // For weighted averaging of multiple triangle influences
 		TManagedArray<float>* RenderDeformerSkinningBlend;  // Render weight map. How much the vertex actually contributes, value between 0 (fully deformed) and 1 (fully skinned)
+		TManagedArray<float>* RenderCustomResizingBlend; // Resizing weight map. How much to use group resizing result (1) vs rbf interpolation result (0) when doing render mesh resizing
+
+		//~ Group Resizing Group
+		TManagedArray<FString>* CustomResizingRegionSet; // Name of the set associated with this binding type
+		TManagedArray<int32>* CustomResizingRegionType; // Actually an Enum for binding type
 	};
 
 	template<typename T>
diff --git a/Engine/Plugins/ChaosClothAsset/Source/ChaosClothAsset/Private/ChaosClothAsset/ClothGeometryTools.cpp b/Engine/Plugins/ChaosClothAsset/Source/ChaosClothAsset/Private/ChaosClothAsset/ClothGeometryTools.cpp
--- a/Engine/Plugins/ChaosClothAsset/Source/ChaosClothAsset/Private/ChaosClothAsset/ClothGeometryTools.cpp
+++ b/Engine/Plugins/ChaosClothAsset/Source/ChaosClothAsset/Private/ChaosClothAsset/ClothGeometryTools.cpp
@@ -10,10 +10,14 @@
 #include "DynamicMesh/DynamicMeshAttributeSet.h"
 #include "DynamicMesh/DynamicVertexSkinWeightsAttribute.h"
 #include "DynamicMesh/NonManifoldMappingSupport.h"
+#include "Math/OrientedBox.h"
 #include "Math/Vector.h"
+#include "MeshResizing/CustomRegionResizing.h"
+#include "MeshResizing/RBFInterpolation.h"
 #include "Util/IndexUtil.h"
 #include "Utils/ClothingMeshUtils.h"
 #include "Algo/RemoveIf.h"
+#include "MeshDescriptionToDynamicMesh.h"
 #include "RenderMath.h"
 
 namespace UE::Chaos::ClothAsset
@@ -466,6 +470,69 @@ namespace UE::Chaos::ClothAsset
 		}
 	}// namespace Private::SimMeshBuilder
 
+	namespace Private::Resizing
+	{
+		static void ApplyGroupResizing(FCollectionClothFacade& ClothFacade, const FMeshDescription& TargetMeshDescription, const FMeshResizingRBFInterpolationData& InterpolationData, const TArray<FMeshResizingCustomRegion>& ResizingGroupData, TArrayView<FVector3f> Positions, TConstArrayView<float> ResizingBlend)
+		{
+			if (ClothFacade.IsValid(EClothCollectionOptionalSchemas::Resizing))
+			{
+				TConstArrayView<int32> ClothSetTypes = ClothFacade.GetCustomResizingRegionType();
+				check(ClothSetTypes.Num() == ResizingGroupData.Num());
+
+				// Gather trilinear interpolation data.
+				TArray<int32> TrilinearInterpolationGroups;
+				TArray<FVector3d> BoundCorners;
+				int32 BoundCornersIndex = 0;
+
+				for (int32 GroupIndex = 0; GroupIndex < ClothSetTypes.Num(); ++GroupIndex)
+				{
+					if ((EMeshResizingCustomRegionType)ClothSetTypes[GroupIndex] == EMeshResizingCustomRegionType::TrilinearInterpolation)
+					{
+						const FMeshResizingCustomRegion& ResizingData = ResizingGroupData[GroupIndex];
+
+						if (ResizingData.IsValid())
+						{
+							TrilinearInterpolationGroups.Add(GroupIndex);
+
+							const FMatrix TriangleMatrix(FVector3d(ResizingData.SourceAxis0), FVector3d(ResizingData.SourceAxis1), FVector3d(ResizingData.SourceAxis2), ResizingData.SourceOrigin);
+							FOrientedBox BoundBox;
+							BoundBox.Center = TriangleMatrix.TransformPosition(FVector3d(ResizingData.RegionBoundsCentroid));
+							BoundBox.AxisX = FVector3d(ResizingData.SourceAxis0);
+							BoundBox.AxisY = FVector3d(ResizingData.SourceAxis1);
+							BoundBox.AxisZ = FVector3d(ResizingData.SourceAxis2);
+							BoundBox.ExtentX = ResizingData.RegionBoundsExtents.X;
+							BoundBox.ExtentY = ResizingData.RegionBoundsExtents.Y;
+							BoundBox.ExtentZ = ResizingData.RegionBoundsExtents.Z;
+							BoundCorners.SetNum((BoundCornersIndex + 1) * 8);
+							BoundBox.CalcVertices(&BoundCorners[BoundCornersIndex]);
+							BoundCornersIndex += 8;
+						}
+					}
+				}
+				if (!TrilinearInterpolationGroups.IsEmpty())
+				{
+					UE::MeshResizing::FRBFInterpolation::DeformPoints(TargetMeshDescription, InterpolationData, BoundCorners);
+					BoundCornersIndex = 0;
+					check(BoundCorners.Num() == TrilinearInterpolationGroups.Num() * 8);
+
+					const TArray<FVector3f> OrigPositions(Positions.GetData(), Positions.Num());
+					for (const int32 GroupIndex : TrilinearInterpolationGroups)
+					{
+						const FMeshResizingCustomRegion& ResizingData = ResizingGroupData[GroupIndex];
+
+						UE::MeshResizing::FCustomRegionResizing::InterpolateCustomRegionPoints(ResizingData, TConstArrayView<FVector3d>(BoundCorners.GetData() + 8 * BoundCornersIndex, 8), Positions);
+						BoundCornersIndex += 8;
+					}
+
+					for (int32 VertexIndex = 0; VertexIndex < ResizingBlend.Num(); ++VertexIndex)
+					{
+						Positions[VertexIndex] = FMath::Lerp(OrigPositions[VertexIndex], Positions[VertexIndex], ResizingBlend[VertexIndex]);
+					}
+				}
+			}
+		}
+	} // namespace Private::Resizing
+
 	// TODO: Move these functions to the cloth facade?
 	bool FClothGeometryTools::HasSimMesh(const TSharedRef<const FManagedArrayCollection>& ClothCollection)
 	{
@@ -1812,4 +1879,129 @@ PRAGMA_ENABLE_DEPRECATION_WARNINGS
 			}
 		}
 	}
+
+	void FClothGeometryTools::SetGroupResizingData(const TSharedRef<FManagedArrayCollection>& ClothCollection, const TConstArrayView<FName>& SetNames, const TConstArrayView<int32>& SetTypes)
+	{
+		check(SetNames.Num() == SetTypes.Num());
+		FCollectionClothFacade ClothFacade(ClothCollection);
+
+		ClothFacade.DefineSchema(EClothCollectionOptionalSchemas::Resizing);
+
+		// Clear the sim and render blends
+		TArrayView<float> SimBlend = ClothFacade.GetSimCustomResizingBlend();
+		TArrayView<float> RenderBlend = ClothFacade.GetRenderCustomResizingBlend();
+		for (int32 Index = 0; Index < SimBlend.Num(); ++Index)
+		{
+			SimBlend[Index] = 0.f;
+		}
+		for (int32 Index = 0; Index < RenderBlend.Num(); ++Index)
+		{
+			RenderBlend[Index] = 0.f;
+		}
+
+		// Filter out invalid sets
+		TArray<FName> FilteredSetNames;
+		TArray<int32> FilteredSetTypes;
+		FilteredSetNames.Reserve(SetNames.Num());
+		FilteredSetTypes.Reserve(SetNames.Num());
+		TSet<int32> SetValues;
+		for (int32 Index = 0; Index < SetNames.Num(); ++Index)
+		{
+			// Try to get the set as a SimVertex3D set
+			if (ConvertSelectionToNewGroupType(ClothCollection, SetNames[Index], ClothCollectionGroup::SimVertices3D, SetValues))
+			{
+				FilteredSetNames.Add(SetNames[Index]);
+				FilteredSetTypes.Add(SetTypes[Index]);
+
+				for (const int32 Vertex : SetValues)
+				{
+					if (SimBlend.IsValidIndex(Vertex))
+					{
+						SimBlend[Vertex] = 1.f;
+					}
+				}
+			}
+			else if (ConvertSelectionToNewGroupType(ClothCollection, SetNames[Index], ClothCollectionGroup::RenderVertices, SetValues))
+			{
+				FilteredSetNames.Add(SetNames[Index]);
+				FilteredSetTypes.Add(SetTypes[Index]);
+
+				for (const int32 Vertex : SetValues)
+				{
+					if (RenderBlend.IsValidIndex(Vertex))
+					{
+						RenderBlend[Vertex] = 1.f;
+					}
+				}
+			}
+		}
+
+		check(FilteredSetNames.Num() == FilteredSetTypes.Num());
+		ClothFacade.SetNumCustomResizingRegions(FilteredSetNames.Num());
+		TArrayView<FString> ClothSetNames = ClothFacade.GetCustomResizingRegionSet();
+		TArrayView<int32> ClothSetTypes = ClothFacade.GetCustomResizingRegionType();
+		for (int32 Index = 0; Index < FilteredSetNames.Num(); ++Index)
+		{
+			ClothSetNames[Index] = FilteredSetNames[Index].ToString();
+			ClothSetTypes[Index] = FilteredSetTypes[Index];
+		}
+	}
+
+	void FClothGeometryTools::GenerateSimMeshResizingGroupData(const TSharedRef<FManagedArrayCollection>& ClothCollection, const FMeshDescription& SourceMeshDescription, TArray<FMeshResizingCustomRegion>& OutResizingGroupData)
+	{
+		FCollectionClothFacade ClothFacade(ClothCollection);
+		if (ClothFacade.IsValid(EClothCollectionOptionalSchemas::Resizing))
+		{
+			UE::Geometry::FDynamicMesh3 SourceMesh;
+			FMeshDescriptionToDynamicMesh Converter;
+			Converter.Convert(&SourceMeshDescription, SourceMesh, true);
+
+			TConstArrayView<FString> ClothSetNames = ClothFacade.GetCustomResizingRegionSet();
+			TConstArrayView<int32> ClothSetTypes = ClothFacade.GetCustomResizingRegionType();
+			OutResizingGroupData.SetNum(ClothSetNames.Num());
+			for (int32 GroupIndex = 0; GroupIndex < ClothSetNames.Num(); ++GroupIndex)
+			{
+				TSet<int32> SetValues;
+				if (ConvertSelectionToNewGroupType(ClothCollection, FName(ClothSetNames[GroupIndex]), ClothCollectionGroup::SimVertices3D, SetValues))
+				{
+					UE::MeshResizing::FCustomRegionResizing::GenerateCustomRegion(ClothFacade.GetSimPosition3D(), SourceMesh, SetValues, OutResizingGroupData[GroupIndex]);
+				}
+			}
+		}
+	}
+
+	void FClothGeometryTools::GenerateRenderMeshResizingGroupData(const TSharedRef<FManagedArrayCollection>& ClothCollection, const FMeshDescription& SourceMeshDescription, TArray<FMeshResizingCustomRegion>& OutResizingGroupData)
+	{
+		FCollectionClothFacade ClothFacade(ClothCollection);
+		if (ClothFacade.IsValid(EClothCollectionOptionalSchemas::Resizing))
+		{
+			UE::Geometry::FDynamicMesh3 SourceMesh;
+			FMeshDescriptionToDynamicMesh Converter;
+			Converter.Convert(&SourceMeshDescription, SourceMesh, true);
+
+			TConstArrayView<FString> ClothSetNames = ClothFacade.GetCustomResizingRegionSet();
+			TConstArrayView<int32> ClothSetTypes = ClothFacade.GetCustomResizingRegionType();
+			OutResizingGroupData.SetNum(ClothSetNames.Num());
+			for (int32 GroupIndex = 0; GroupIndex < ClothSetNames.Num(); ++GroupIndex)
+			{
+				TSet<int32> SetValues;
+				if (ConvertSelectionToNewGroupType(ClothCollection, FName(ClothSetNames[GroupIndex]), ClothCollectionGroup::RenderVertices, SetValues))
+				{
+					UE::MeshResizing::FCustomRegionResizing::GenerateCustomRegion(ClothFacade.GetRenderPosition(), SourceMesh, SetValues, OutResizingGroupData[GroupIndex]);
+				}
+			}
+		}
+	}
+
+	void FClothGeometryTools::ApplySimGroupResizing(const TSharedRef<FManagedArrayCollection>& ClothCollection, const FMeshDescription& TargetMeshDescription, const FMeshResizingRBFInterpolationData& InterpolationData, const TArray<FMeshResizingCustomRegion>& ResizingGroupData)
+	{
+		FCollectionClothFacade ClothFacade(ClothCollection);
+		Private::Resizing::ApplyGroupResizing(ClothFacade, TargetMeshDescription, InterpolationData, ResizingGroupData, ClothFacade.GetSimPosition3D(), ClothFacade.GetSimCustomResizingBlend());
+	}
+
+	void FClothGeometryTools::ApplyRenderGroupResizing(const TSharedRef<FManagedArrayCollection>& ClothCollection, const FMeshDescription& TargetMeshDescription, const FMeshResizingRBFInterpolationData& InterpolationData, const TArray<FMeshResizingCustomRegion>& ResizingGroupData)
+	{
+		FCollectionClothFacade ClothFacade(ClothCollection);
+		Private::Resizing::ApplyGroupResizing(ClothFacade, TargetMeshDescription, InterpolationData, ResizingGroupData, ClothFacade.GetRenderPosition(), ClothFacade.GetRenderCustomResizingBlend());
+	}
 }  // End namespace UE::Chaos::ClothAsset
diff --git a/Engine/Plugins/ChaosClothAsset/Source/ChaosClothAsset/Private/ChaosClothAsset/CollectionClothFacade.cpp b/Engine/Plugins/ChaosClothAsset/Source/ChaosClothAsset/Private/ChaosClothAsset/CollectionClothFacade.cpp
--- a/Engine/Plugins/ChaosClothAsset/Source/ChaosClothAsset/Private/ChaosClothAsset/CollectionClothFacade.cpp
+++ b/Engine/Plugins/ChaosClothAsset/Source/ChaosClothAsset/Private/ChaosClothAsset/CollectionClothFacade.cpp
@@ -108,6 +108,7 @@ namespace UE::Chaos::ClothAsset
 		ResultHash = HashCombineFast(ResultHash, ClothCollection->GetElementsTypeHash(ClothCollection->GetTetherReferenceLength()));
 		ResultHash = HashCombineFast(ResultHash, ClothCollection->GetElementsTypeHash(ClothCollection->GetSimVertex2DLookup()));
 		ResultHash = HashCombineFast(ResultHash, ClothCollection->GetElementsTypeHash(ClothCollection->GetSeamStitchLookup()));
+		ResultHash = HashCombineFast(ResultHash, ClothCollection->GetElementsTypeHash(ClothCollection->GetSimCustomResizingBlend()));
 
 		//~ Sim Morph Targets Group
 		ResultHash = HashCombineFast(ResultHash, ClothCollection->GetElementsTypeHash(ClothCollection->GetSimMorphTargetName()));
@@ -149,6 +150,11 @@ namespace UE::Chaos::ClothAsset
 		ResultHash = HashCombineFast(ResultHash, ClothCollection->GetElementsTypeHash(ClothCollection->GetRenderDeformerSimIndices3D()));
 		ResultHash = HashCombineFast(ResultHash, ClothCollection->GetElementsTypeHash(ClothCollection->GetRenderDeformerWeight()));
 		ResultHash = HashCombineFast(ResultHash, ClothCollection->GetElementsTypeHash(ClothCollection->GetRenderDeformerSkinningBlend()));
+		ResultHash = HashCombineFast(ResultHash, ClothCollection->GetElementsTypeHash(ClothCollection->GetRenderCustomResizingBlend()));
+
+		//~ Resizing Binding Group
+		ResultHash = HashCombineFast(ResultHash, ClothCollection->GetElementsTypeHash(ClothCollection->GetCustomResizingRegionSet()));
+		ResultHash = HashCombineFast(ResultHash, ClothCollection->GetElementsTypeHash(ClothCollection->GetCustomResizingRegionType()));
 
 		if (bIncludeWeightMaps)
 		{
@@ -283,6 +289,11 @@ namespace UE::Chaos::ClothAsset
 		return ClothCollection->GetElements(ClothCollection->GetSeamStitchLookup());
 	}
 
+	TConstArrayView<float> FCollectionClothConstFacade::GetSimCustomResizingBlend() const
+	{
+		return ClothCollection->GetElements(ClothCollection->GetSimCustomResizingBlend());
+	}
+
 	int32 FCollectionClothConstFacade::GetNumSimFaces() const
 	{
 		return ClothCollection->GetNumElements(ClothCollectionGroup::SimFaces);
@@ -491,6 +502,11 @@ namespace UE::Chaos::ClothAsset
 		return ClothCollection->GetElements(ClothCollection->GetRenderDeformerSkinningBlend());
 	}
 
+	TConstArrayView<float> FCollectionClothConstFacade::GetRenderCustomResizingBlend() const
+	{
+		return ClothCollection->GetElements(ClothCollection->GetRenderCustomResizingBlend());
+	}
+
 	int32 FCollectionClothConstFacade::GetNumRenderFaces() const
 	{
 		return ClothCollection->GetNumElements(ClothCollectionGroup::RenderFaces);
@@ -501,6 +517,21 @@ namespace UE::Chaos::ClothAsset
 		return ClothCollection->GetElements(ClothCollection->GetRenderIndices());
 	}
 
+	int32 FCollectionClothConstFacade::GetNumCustomResizingRegions() const
+	{
+		return ClothCollection->GetNumElements(ClothCollectionGroup::CustomResizingRegions);
+	}
+
+	TConstArrayView<FString> FCollectionClothConstFacade::GetCustomResizingRegionSet() const
+	{
+		return ClothCollection->GetElements(ClothCollection->GetCustomResizingRegionSet());
+	}
+
+	TConstArrayView<int32> FCollectionClothConstFacade::GetCustomResizingRegionType() const
+	{
+		return ClothCollection->GetElements(ClothCollection->GetCustomResizingRegionType());
+	}
+
 	bool FCollectionClothConstFacade::HasWeightMap(const FName& Name) const
 	{
 		return ClothCollection->HasUserDefinedAttribute<float>(Name, ClothCollectionGroup::SimVertices3D);
@@ -628,6 +659,7 @@ namespace UE::Chaos::ClothAsset
 			GetClothCollection()->SetNumElements(1, ClothCollectionGroup::Solvers);  //
 		}
 		SetNumSimMorphTargets(0);
+		SetNumCustomResizingRegions(0);
 	}
 
 	void FCollectionClothFacade::Initialize(const FCollectionClothConstFacade& Other)
@@ -648,6 +680,7 @@ namespace UE::Chaos::ClothAsset
 
 	void FCollectionClothFacade::Append(const FCollectionClothConstFacade& Other)
 	{
+
 #if DO_ENSURE
 		for (int32 SeamIndex = 0; SeamIndex < GetNumSeams(); ++SeamIndex)
 		{
@@ -666,6 +699,12 @@ namespace UE::Chaos::ClothAsset
 			SetSkeletalMeshPathName(Other.GetSkeletalMeshPathName());
 		}
 
+		const bool bOtherHasResizingData = Other.IsValid(EClothCollectionOptionalSchemas::Resizing);
+		if (bOtherHasResizingData)
+		{
+			DefineSchema(EClothCollectionOptionalSchemas::Resizing);
+		}
+
 		// Very important order of operations to ensure indices don't get messed up:
 		// 1) Append 3D Vertices, but don't set 2D Lookups or SeamStitch Lookups because those indices don't exist yet.
 		// 2) Append Sim Patterns (includes 2D Vertices, have 3D dependency)
@@ -682,6 +721,10 @@ namespace UE::Chaos::ClothAsset
 		FClothCollection::CopyArrayViewData(GetSimBoneWeights().Right(OtherNumSimVertices3D), Other.GetSimBoneWeights());
 		FClothCollection::CopyArrayViewDataAndApplyOffset(GetTetherKinematicIndex().Right(OtherNumSimVertices3D), Other.GetTetherKinematicIndex(), StartNumSimVertices3D);
 		FClothCollection::CopyArrayViewData(GetTetherReferenceLength().Right(OtherNumSimVertices3D), Other.GetTetherReferenceLength());
+		if (bOtherHasResizingData)
+		{
+			FClothCollection::CopyArrayViewData(GetSimCustomResizingBlend().Right(OtherNumSimVertices3D), Other.GetSimCustomResizingBlend());
+		}
 
 		// Fabrics Group
 		const int32 StartNumFabrics = GetNumFabrics();
@@ -746,6 +789,16 @@ namespace UE::Chaos::ClothAsset
 			GetRenderPattern(PatternIndex + StartNumRenderPatterns).Initialize(Other.GetRenderPattern(PatternIndex), StartNumSimVertices3D);
 		}
 
+		// Resizing Binding Group
+		if (bOtherHasResizingData)
+		{
+			const int32 StartNumResizingBindingGroups = GetNumCustomResizingRegions();
+			const int32 OtherNumResizingBindingGroups = Other.GetNumCustomResizingRegions();
+			SetNumCustomResizingRegions(StartNumResizingBindingGroups + OtherNumResizingBindingGroups);
+			FClothCollection::CopyArrayViewData(GetCustomResizingRegionSet().Right(OtherNumResizingBindingGroups), Other.GetCustomResizingRegionSet());
+			FClothCollection::CopyArrayViewData(GetCustomResizingRegionType().Right(OtherNumResizingBindingGroups), Other.GetCustomResizingRegionType());
+		}
+
 		// Weight maps
 		const TArray<FName> WeightMapNames = Other.GetWeightMapNames();
 		for (const FName& WeightMapName : WeightMapNames)
@@ -891,6 +944,11 @@ namespace UE::Chaos::ClothAsset
 	{
 		return GetClothCollection()->GetElements(GetClothCollection()->GetTetherReferenceLength());
 	}
+	
+	TArrayView<float> FCollectionClothFacade::GetSimCustomResizingBlend()
+	{
+		return GetClothCollection()->GetElements(GetClothCollection()->GetSimCustomResizingBlend());
+	}
 
 	TArrayView<TArray<int32>> FCollectionClothFacade::GetSeamStitchLookupPrivate()
 	{
@@ -1272,11 +1330,43 @@ namespace UE::Chaos::ClothAsset
 		return GetClothCollection()->GetElements(GetClothCollection()->GetRenderDeformerSkinningBlend());
 	}
 
+	TArrayView<float> FCollectionClothFacade::GetRenderCustomResizingBlend()
+	{
+		return GetClothCollection()->GetElements(GetClothCollection()->GetRenderCustomResizingBlend());
+	}
+
 	TArrayView<FIntVector3> FCollectionClothFacade::GetRenderIndices()
 	{
 		return GetClothCollection()->GetElements(GetClothCollection()->GetRenderIndices());
 	}
 
+	void FCollectionClothFacade::SetNumCustomResizingRegions(int32 InNumGroups)
+	{
+		if (!IsValid(EClothCollectionOptionalSchemas::Resizing))
+		{
+			DefineSchema(EClothCollectionOptionalSchemas::Resizing);
+		}
+
+		const int32 NumGroups = GetNumCustomResizingRegions();
+
+		GetClothCollection()->SetNumElements(InNumGroups, ClothCollectionGroup::CustomResizingRegions);
+
+		for (int32 Index = NumGroups; Index < InNumGroups; ++Index)
+		{
+			GetCustomResizingRegionSet()[Index] = FString();
+		}
+	}
+
+	TArrayView<FString> FCollectionClothFacade::GetCustomResizingRegionSet()
+	{
+		return GetClothCollection()->GetElements(GetClothCollection()->GetCustomResizingRegionSet());
+	}
+
+	TArrayView<int32> FCollectionClothFacade::GetCustomResizingRegionType()
+	{
+		return GetClothCollection()->GetElements(GetClothCollection()->GetCustomResizingRegionType());
+	}
+
 	void FCollectionClothFacade::AddWeightMap(const FName& Name)
 	{
 		check(IsValid());
diff --git a/Engine/Plugins/ChaosClothAsset/Source/ChaosClothAsset/Private/ChaosClothAsset/CollectionClothRenderPatternFacade.cpp b/Engine/Plugins/ChaosClothAsset/Source/ChaosClothAsset/Private/ChaosClothAsset/CollectionClothRenderPatternFacade.cpp
--- a/Engine/Plugins/ChaosClothAsset/Source/ChaosClothAsset/Private/ChaosClothAsset/CollectionClothRenderPatternFacade.cpp
+++ b/Engine/Plugins/ChaosClothAsset/Source/ChaosClothAsset/Private/ChaosClothAsset/CollectionClothRenderPatternFacade.cpp
@@ -162,6 +162,15 @@ namespace UE::Chaos::ClothAsset
 			GetElementIndex());
 	}
 
+	TConstArrayView<float> FCollectionClothRenderPatternConstFacade::GetRenderCustomResizingBlend() const
+	{
+		return ClothCollection->GetElements(
+			ClothCollection->GetRenderCustomResizingBlend(),
+			ClothCollection->GetRenderVerticesStart(),
+			ClothCollection->GetRenderVerticesEnd(),
+			GetElementIndex());
+	}
+
 	int32 FCollectionClothRenderPatternConstFacade::GetNumRenderFaces() const
 	{
 		return ClothCollection->GetNumElements(
@@ -236,6 +245,11 @@ namespace UE::Chaos::ClothAsset
 			FClothCollection::CopyArrayViewData(GetRenderDeformerSkinningBlend(), Other.GetRenderDeformerSkinningBlend());
 		}
 
+		if (GetClothCollection()->IsValid(EClothCollectionOptionalSchemas::Resizing) && Other.ClothCollection->IsValid(EClothCollectionOptionalSchemas::Resizing))
+		{
+			FClothCollection::CopyArrayViewData(GetRenderCustomResizingBlend(), Other.GetRenderCustomResizingBlend());
+		}
+
 		//~ Render Faces Group
 		const int32 RenderVertexOffset = GetRenderVerticesOffset() - Other.GetRenderVerticesOffset();
 		SetNumRenderFaces(Other.GetNumRenderFaces());
@@ -415,6 +429,15 @@ namespace UE::Chaos::ClothAsset
 			GetElementIndex());
 	}
 
+	TArrayView<float> FCollectionClothRenderPatternFacade::GetRenderCustomResizingBlend()
+	{
+		return GetClothCollection()->GetElements(
+			GetClothCollection()->GetRenderCustomResizingBlend(),
+			GetClothCollection()->GetRenderVerticesStart(),
+			GetClothCollection()->GetRenderVerticesEnd(),
+			GetElementIndex());
+	}
+
 	void FCollectionClothRenderPatternFacade::SetNumRenderFaces(int32 NumRenderFaces)
 	{
 		GetClothCollection()->SetNumElements(
diff --git a/Engine/Plugins/ChaosClothAsset/Source/ChaosClothAsset/Public/ChaosClothAsset/ClothCollectionAttribute.h b/Engine/Plugins/ChaosClothAsset/Source/ChaosClothAsset/Public/ChaosClothAsset/ClothCollectionAttribute.h
--- a/Engine/Plugins/ChaosClothAsset/Source/ChaosClothAsset/Public/ChaosClothAsset/ClothCollectionAttribute.h
+++ b/Engine/Plugins/ChaosClothAsset/Source/ChaosClothAsset/Public/ChaosClothAsset/ClothCollectionAttribute.h
@@ -73,6 +73,7 @@ namespace UE::Chaos::ClothAsset
 		inline const FName TetherReferenceLength(TEXT("TetherReferenceLength"));
 		inline const FName SimVertex2DLookup(TEXT("SimVertex2DLookup"));
 		inline const FName SeamStitchLookup(TEXT("SeamStitchLookup"));
+		inline const FName SimCustomResizingBlend(TEXT("SimCustomResizingBlend"));
 
 		// Sim Morph Targets Group
 		inline const FName SimMorphTargetName(TEXT("SimMorphTargetName"));
@@ -102,5 +103,10 @@ namespace UE::Chaos::ClothAsset
 		inline const FName RenderDeformerSimIndices3D(TEXT("RenderDeformerSimIndices3D"));
 		inline const FName RenderDeformerWeight(TEXT("RenderDeformerWeight"));
 		inline const FName RenderDeformerSkinningBlend(TEXT("RenderDeformerSkinningBlend"));
+		inline const FName RenderCustomResizingBlend(TEXT("RenderCustomResizingBlend"));
+
+		// Custom Resizing Regions Group
+		inline const FName CustomResizingRegionSet(TEXT("CustomResizingRegionSet"));
+		inline const FName CustomResizingRegionType(TEXT("CustomResizingRegionType"));
 	}
 }  // End namespace UE::Chaos::ClothAsset
diff --git a/Engine/Plugins/ChaosClothAsset/Source/ChaosClothAsset/Public/ChaosClothAsset/ClothCollectionGroup.h b/Engine/Plugins/ChaosClothAsset/Source/ChaosClothAsset/Public/ChaosClothAsset/ClothCollectionGroup.h
--- a/Engine/Plugins/ChaosClothAsset/Source/ChaosClothAsset/Public/ChaosClothAsset/ClothCollectionGroup.h
+++ b/Engine/Plugins/ChaosClothAsset/Source/ChaosClothAsset/Public/ChaosClothAsset/ClothCollectionGroup.h
@@ -39,5 +39,7 @@ namespace UE::Chaos::ClothAsset
 		inline const FName SimMorphTargets = FName(TEXT("SimMorphTargets"));
 		/** Contains morph target vertices */
 		inline const FName SimMorphTargetVertices = FName(TEXT("SimMorphTargetVertices"));
+		/** Contains custom region resizing data */
+		inline const FName CustomResizingRegions = FName(TEXT("CustomResizingRegions"));
 	}
 }  // End namespace UE::Chaos::ClothAsset
diff --git a/Engine/Plugins/ChaosClothAsset/Source/ChaosClothAsset/Public/ChaosClothAsset/ClothCollectionOptionalSchemas.h b/Engine/Plugins/ChaosClothAsset/Source/ChaosClothAsset/Public/ChaosClothAsset/ClothCollectionOptionalSchemas.h
--- a/Engine/Plugins/ChaosClothAsset/Source/ChaosClothAsset/Public/ChaosClothAsset/ClothCollectionOptionalSchemas.h
+++ b/Engine/Plugins/ChaosClothAsset/Source/ChaosClothAsset/Public/ChaosClothAsset/ClothCollectionOptionalSchemas.h
@@ -16,6 +16,7 @@ namespace UE::Chaos::ClothAsset
 		Fabrics = 1 << 3,
 		SimMorphTargets = 1 << 4,
 		CookedOnly = 1 << 5, // If this flag is found, only cooked schemas will be considered.
+		Resizing = 1 << 6,
 	};
 	ENUM_CLASS_FLAGS(EClothCollectionOptionalSchemas)
 }  // End namespace UE::Chaos::ClothAsset
diff --git a/Engine/Plugins/ChaosClothAsset/Source/ChaosClothAsset/Public/ChaosClothAsset/ClothGeometryTools.h b/Engine/Plugins/ChaosClothAsset/Source/ChaosClothAsset/Public/ChaosClothAsset/ClothGeometryTools.h
--- a/Engine/Plugins/ChaosClothAsset/Source/ChaosClothAsset/Public/ChaosClothAsset/ClothGeometryTools.h
+++ b/Engine/Plugins/ChaosClothAsset/Source/ChaosClothAsset/Public/ChaosClothAsset/ClothGeometryTools.h
@@ -8,6 +8,9 @@
 
 struct FManagedArrayCollection;
 namespace UE::Geometry { class FDynamicMesh3; }
+struct FMeshDescription;
+struct FMeshResizingRBFInterpolationData;
+struct FMeshResizingCustomRegion;
 
 namespace UE::Chaos::ClothAsset
 {
@@ -176,6 +179,31 @@ namespace UE::Chaos::ClothAsset
 		 */
 		static UE_API void ApplyProxyDeformer(const TSharedRef<FManagedArrayCollection>& ClothCollection, bool bIgnoreSkinningBlend);
 		static UE_API void ApplyProxyDeformer(FCollectionClothFacade& ClothFacade, bool bIgnoreSkinningBlend, const TConstArrayView<FVector3f>& SimPositions, const TConstArrayView<FVector3f>& SimNormals);
+
+		/**
+		 * Set Group Resizing data from supplied sets and types.
+		 */
+		static UE_API void SetGroupResizingData(const TSharedRef<FManagedArrayCollection>& ClothCollection, const TConstArrayView<FName>& SetNames, const TConstArrayView<int32>& SetTypes);
+
+		/**
+		 * Generate FMeshResizingCustomRegion data from Sim Mesh Group Resizing data
+		 */
+		static UE_API void GenerateSimMeshResizingGroupData(const TSharedRef<FManagedArrayCollection>& ClothCollection, const FMeshDescription& SourceMesh, TArray<FMeshResizingCustomRegion>& OutResizingGroupData);
+
+		/**
+		 * Generate FMeshResizingCustomRegion data from Render Mesh Group Resizing data
+		 */
+		static UE_API void GenerateRenderMeshResizingGroupData(const TSharedRef<FManagedArrayCollection>& ClothCollection, const FMeshDescription& SourceMesh, TArray<FMeshResizingCustomRegion>& OutResizingGroupData);
+
+		/**
+		 * Apply Group Resizing to the sim mesh.
+		 */
+		static UE_API void ApplySimGroupResizing(const TSharedRef<FManagedArrayCollection>& ClothCollection, const FMeshDescription& TargetMesh, const FMeshResizingRBFInterpolationData& InterpolationData, const TArray<FMeshResizingCustomRegion>& ResizingGroupData);
+
+		/**
+		 * Apply Group Resizing to the render mesh.
+		 */
+		static UE_API void ApplyRenderGroupResizing(const TSharedRef<FManagedArrayCollection>& ClothCollection, const FMeshDescription& TargetMesh, const FMeshResizingRBFInterpolationData& InterpolationData, const TArray<FMeshResizingCustomRegion>& ResizingGroupData);
 	};
 }  // End namespace UE::Chaos::ClothAsset
 
diff --git a/Engine/Plugins/ChaosClothAsset/Source/ChaosClothAsset/Public/ChaosClothAsset/CollectionClothFacade.h b/Engine/Plugins/ChaosClothAsset/Source/ChaosClothAsset/Public/ChaosClothAsset/CollectionClothFacade.h
--- a/Engine/Plugins/ChaosClothAsset/Source/ChaosClothAsset/Public/ChaosClothAsset/CollectionClothFacade.h
+++ b/Engine/Plugins/ChaosClothAsset/Source/ChaosClothAsset/Public/ChaosClothAsset/CollectionClothFacade.h
@@ -100,6 +100,7 @@ namespace UE::Chaos::ClothAsset
 		UE_API TConstArrayView<TArray<float>> GetTetherReferenceLength() const;
 		UE_API TConstArrayView<TArray<int32>> GetSimVertex2DLookup() const;
 		UE_API TConstArrayView<TArray<int32>> GetSeamStitchLookup() const;
+		UE_API TConstArrayView<float> GetSimCustomResizingBlend() const;
 
 		//~ Sim Faces Group
 		/** Return the total number of simulation faces for this collection across all patterns. */
@@ -180,11 +181,17 @@ namespace UE::Chaos::ClothAsset
 		UE_API TConstArrayView<TArray<FIntVector3>> GetRenderDeformerSimIndices3D() const;
 		UE_API TConstArrayView<TArray<float>> GetRenderDeformerWeight() const;
 		UE_API TConstArrayView<float> GetRenderDeformerSkinningBlend() const;
+		UE_API TConstArrayView<float> GetRenderCustomResizingBlend() const;
 
 		//~ Render Faces Group
 		UE_API int32 GetNumRenderFaces() const;
 		UE_API TConstArrayView<FIntVector3> GetRenderIndices() const;
 
+		//~ Custom Resizing Regions Group
+		UE_API int32 GetNumCustomResizingRegions() const;
+		UE_API TConstArrayView<FString> GetCustomResizingRegionSet() const;
+		UE_API TConstArrayView<int32> GetCustomResizingRegionType() const;
+
 		//~ Weight Maps
 		/** Return whether this cloth collection has the specified weight map. */
 		UE_API bool HasWeightMap(const FName& Name) const;
@@ -278,6 +285,7 @@ namespace UE::Chaos::ClothAsset
 		UE_API TArrayView<TArray<float>> GetSimBoneWeights();
 		UE_API TArrayView<TArray<int32>> GetTetherKinematicIndex();
 		UE_API TArrayView<TArray<float>> GetTetherReferenceLength();
+		UE_API TArrayView<float> GetSimCustomResizingBlend();
 
 		/** This will remove the 3D vertices, but the associated seams and 2D vertices will still exist, and point to INDEX_NONE */
 		UE_API void RemoveSimVertices3D(int32 NumSimVertices);
@@ -385,11 +393,17 @@ namespace UE::Chaos::ClothAsset
 		UE_API TArrayView<TArray<FIntVector3>> GetRenderDeformerSimIndices3D();
 		UE_API TArrayView<TArray<float>> GetRenderDeformerWeight();
 		UE_API TArrayView<float> GetRenderDeformerSkinningBlend();
+		UE_API TArrayView<float> GetRenderCustomResizingBlend();
 
 		//~ Render Faces Group
 		/** SetNumRenderFaces per pattern within pattern facade. */
 		UE_API TArrayView<FIntVector3> GetRenderIndices();
 
+		//~ Custom Resizing Regions Group
+		UE_API void SetNumCustomResizingRegions(int32 NumGroups);
+		UE_API TArrayView<FString> GetCustomResizingRegionSet();
+		UE_API TArrayView<int32> GetCustomResizingRegionType();
+
 		//~ Weight Maps
 		/** Add a new weight map to this cloth. Access is then done per pattern. */
 		UE_API void AddWeightMap(const FName& Name);
diff --git a/Engine/Plugins/ChaosClothAsset/Source/ChaosClothAsset/Public/ChaosClothAsset/CollectionClothRenderPatternFacade.h b/Engine/Plugins/ChaosClothAsset/Source/ChaosClothAsset/Public/ChaosClothAsset/CollectionClothRenderPatternFacade.h
--- a/Engine/Plugins/ChaosClothAsset/Source/ChaosClothAsset/Public/ChaosClothAsset/CollectionClothRenderPatternFacade.h
+++ b/Engine/Plugins/ChaosClothAsset/Source/ChaosClothAsset/Public/ChaosClothAsset/CollectionClothRenderPatternFacade.h
@@ -53,6 +53,7 @@ namespace UE::Chaos::ClothAsset
 		UE_API TConstArrayView<TArray<FIntVector3>> GetRenderDeformerSimIndices3D() const;
 		UE_API TConstArrayView<TArray<float>> GetRenderDeformerWeight() const;
 		UE_API TConstArrayView<float> GetRenderDeformerSkinningBlend() const;
+		UE_API TConstArrayView<float> GetRenderCustomResizingBlend() const;
 
 		//~ Render Faces Group
 		// Note: RenderIndices points to the collection arrays, not the pattern arrays
@@ -130,6 +131,7 @@ namespace UE::Chaos::ClothAsset
 		UE_API TArrayView<TArray<FIntVector3>> GetRenderDeformerSimIndices3D();
 		UE_API TArrayView<TArray<float>> GetRenderDeformerWeight();
 		UE_API TArrayView<float> GetRenderDeformerSkinningBlend();
+		UE_API TArrayView<float> GetRenderCustomResizingBlend();
 
 		//~ Render Faces Group
 		/** Grow or shrink the space reserved for render faces for this pattern within the cloth collection and return its start index. */
diff --git a/Engine/Plugins/ChaosClothAssetEditor/Source/ChaosClothAssetDataflowNodes/Private/ChaosClothAsset/ApplyResizingNode.cpp b/Engine/Plugins/ChaosClothAssetEditor/Source/ChaosClothAssetDataflowNodes/Private/ChaosClothAsset/ApplyResizingNode.cpp
--- a/Engine/Plugins/ChaosClothAssetEditor/Source/ChaosClothAssetDataflowNodes/Private/ChaosClothAsset/ApplyResizingNode.cpp
+++ b/Engine/Plugins/ChaosClothAssetEditor/Source/ChaosClothAssetDataflowNodes/Private/ChaosClothAsset/ApplyResizingNode.cpp
@@ -1,9 +1,11 @@
 // Copyright Epic Games, Inc. All Rights Reserved.
 
 #include "ChaosClothAsset/ApplyResizingNode.h"
+#include "ChaosClothAsset/ClothGeometryTools.h"
 #include "ChaosClothAsset/CollectionClothFacade.h"
 #include "Dataflow/DataflowInputOutput.h"
 #include "Engine/SkeletalMesh.h"
+#include "MeshResizing/CustomRegionResizing.h"
 
 #include UE_INLINE_GENERATED_CPP_BY_NAME(ApplyResizingNode)
 
@@ -22,6 +24,12 @@ FChaosClothAssetApplyResizingNode::FChaosClothAssetApplyResizingNode(const UE::D
 	RegisterInputConnection(&bForceApplyToRenderMesh)
 		.SetCanHidePin(true)
 		.SetPinIsHidden(true);
+	RegisterInputConnection(&SourceSkeletalMesh)
+		.SetCanHidePin(true)
+		.SetPinIsHidden(true);
+	RegisterInputConnection(&bSkipCustomRegionResizing)
+		.SetCanHidePin(true)
+		.SetPinIsHidden(true);
 }
 
 void FChaosClothAssetApplyResizingNode::Evaluate(UE::Dataflow::FContext& Context, const FDataflowOutput* Out) const
@@ -40,18 +48,52 @@ void FChaosClothAssetApplyResizingNode::Evaluate(UE::Dataflow::FContext& Context
 				const int32 InLODIndex = GetValue(Context, &SkeletalMeshLODIndex);
 				if (InTargetSKM->IsValidLODIndex(InLODIndex) && InTargetSKM->HasMeshDescription(InLODIndex))
 				{
-					if (const FMeshDescription* const MeshDescription = InTargetSKM->GetMeshDescription(InLODIndex))
+					if (const FMeshDescription* const TargetMeshDescription = InTargetSKM->GetMeshDescription(InLODIndex))
 					{
 						const FMeshResizingRBFInterpolationData& InInterpolationData = GetValue(Context, &InterpolationData);
 						if (InInterpolationData.SampleIndices.Num())
 						{
 							if (!GetValue(Context, &bForceApplyToRenderMesh) && ClothFacade.HasValidSimulationData())
 							{
-								UE::MeshResizing::FRBFInterpolation::DeformPoints(*MeshDescription, InInterpolationData, ClothFacade.GetSimPosition3D(), ClothFacade.GetSimNormal());
+								TArray<FMeshResizingCustomRegion> ResizingGroupData;
+								if (!GetValue(Context, &bSkipCustomRegionResizing))
+								{
+									if (TObjectPtr<const USkeletalMesh> InSourceSKM = GetValue(Context, &SourceSkeletalMesh))
+									{
+										if (const FMeshDescription* const SourceMeshDescription = InTargetSKM->GetMeshDescription(InLODIndex))
+										{
+											FClothGeometryTools::GenerateSimMeshResizingGroupData(ClothCollection, *SourceMeshDescription, ResizingGroupData);
+										}
+									}
+								}
+
+								UE::MeshResizing::FRBFInterpolation::DeformPoints(*TargetMeshDescription, InInterpolationData, ClothFacade.GetSimPosition3D(), ClothFacade.GetSimNormal());
+
+								if (!ResizingGroupData.IsEmpty())
+								{
+									FClothGeometryTools::ApplySimGroupResizing(ClothCollection, *TargetMeshDescription, InInterpolationData, ResizingGroupData);
+								}
 							}
 							else if (ClothFacade.HasValidRenderData())
 							{
-								UE::MeshResizing::FRBFInterpolation::DeformPoints(*MeshDescription, InInterpolationData, ClothFacade.GetRenderPosition(), ClothFacade.GetRenderNormal(), ClothFacade.GetRenderTangentU(), ClothFacade.GetRenderTangentV());
+								TArray<FMeshResizingCustomRegion> ResizingGroupData;
+								if (!GetValue(Context, &bSkipCustomRegionResizing))
+								{
+									if (TObjectPtr<const USkeletalMesh> InSourceSKM = GetValue(Context, &SourceSkeletalMesh))
+									{
+										if (const FMeshDescription* const SourceMeshDescription = InTargetSKM->GetMeshDescription(InLODIndex))
+										{
+											FClothGeometryTools::GenerateRenderMeshResizingGroupData(ClothCollection, *SourceMeshDescription, ResizingGroupData);
+										}
+									}
+								}
+
+								UE::MeshResizing::FRBFInterpolation::DeformPoints(*TargetMeshDescription, InInterpolationData, ClothFacade.GetRenderPosition(), ClothFacade.GetRenderNormal(), ClothFacade.GetRenderTangentU(), ClothFacade.GetRenderTangentV());
+
+								if (!ResizingGroupData.IsEmpty())
+								{
+									FClothGeometryTools::ApplyRenderGroupResizing(ClothCollection, *TargetMeshDescription, InInterpolationData, ResizingGroupData);
+								}
 							}
 						}
 					}
diff --git a/Engine/Plugins/ChaosClothAssetEditor/Source/ChaosClothAssetDataflowNodes/Private/ChaosClothAsset/ChaosClothAssetDataflowNodesModule.cpp b/Engine/Plugins/ChaosClothAssetEditor/Source/ChaosClothAssetDataflowNodes/Private/ChaosClothAsset/ChaosClothAssetDataflowNodesModule.cpp
--- a/Engine/Plugins/ChaosClothAssetEditor/Source/ChaosClothAssetDataflowNodes/Private/ChaosClothAsset/ChaosClothAssetDataflowNodesModule.cpp
+++ b/Engine/Plugins/ChaosClothAssetEditor/Source/ChaosClothAssetDataflowNodes/Private/ChaosClothAsset/ChaosClothAssetDataflowNodesModule.cpp
@@ -15,6 +15,7 @@
 #include "ChaosClothAsset/CopySimulationToRenderMeshNode.h"
 #include "ChaosClothAsset/DatasmithImportNode.h"
 #include "ChaosClothAsset/DeleteElementNode.h"
+#include "ChaosClothAsset/CustomRegionResizingNode.h"
 #include "ChaosClothAsset/ImportFilePathCustomization.h"
 #include "ChaosClothAsset/ImportNode.h"
 #include "ChaosClothAsset/ImportSimulationCacheNode.h"
@@ -106,9 +107,11 @@ namespace UE::Chaos::ClothAsset
 			DATAFLOW_NODE_REGISTER_CREATION_FACTORY(FChaosClothAssetCollectionToDynamicMeshNode);
 			DATAFLOW_NODE_REGISTER_CREATION_FACTORY(FChaosClothAssetCollectionQueryNode);
 			DATAFLOW_NODE_REGISTER_CREATION_FACTORY(FChaosClothAssetCopySimulationToRenderMeshNode);
+			DATAFLOW_NODE_REGISTER_CREATION_FACTORY(FChaosClothAssetCustomRegionResizingNode);
 			DATAFLOW_NODE_REGISTER_CREATION_FACTORY(FChaosClothAssetDatasmithImportNode);
 			DATAFLOW_NODE_REGISTER_CREATION_FACTORY(FChaosClothAssetDeleteElementNode);
 			DATAFLOW_NODE_REGISTER_CREATION_FACTORY(FChaosClothAssetExtractWeightMapNode);
+			DATAFLOW_NODE_REGISTER_CREATION_FACTORY(FChaosClothAssetExtractSelectionSetNode);
 			DATAFLOW_NODE_REGISTER_CREATION_FACTORY(FChaosClothAssetImportNode);
 			DATAFLOW_NODE_REGISTER_CREATION_FACTORY(FChaosClothAssetImportSimulationCacheNode);
 			DATAFLOW_NODE_REGISTER_CREATION_FACTORY(FChaosClothAssetMakeClothAssetNode);
diff --git a/Engine/Plugins/ChaosClothAssetEditor/Source/ChaosClothAssetDataflowNodes/Private/ChaosClothAsset/ClothCollectionToDynamicMeshNode.cpp b/Engine/Plugins/ChaosClothAssetEditor/Source/ChaosClothAssetDataflowNodes/Private/ChaosClothAsset/ClothCollectionToDynamicMeshNode.cpp
--- a/Engine/Plugins/ChaosClothAssetEditor/Source/ChaosClothAssetDataflowNodes/Private/ChaosClothAsset/ClothCollectionToDynamicMeshNode.cpp
+++ b/Engine/Plugins/ChaosClothAssetEditor/Source/ChaosClothAssetDataflowNodes/Private/ChaosClothAsset/ClothCollectionToDynamicMeshNode.cpp
@@ -4,6 +4,7 @@
 #include "ChaosClothAsset/ClothPatternToDynamicMesh.h"
 #include "ChaosClothAsset/ClothPatternToDynamicMeshMappingSupport.h"
 #include "ChaosClothAsset/CollectionClothFacade.h"
+#include "ChaosClothAsset/CollectionClothSelectionFacade.h"
 #include "ChaosClothAsset/ClothCollectionGroup.h"
 #include "Dataflow/DataflowInputOutput.h"
 #include "DynamicMesh/DynamicMesh3.h"
@@ -349,4 +350,68 @@ void FChaosClothAssetExtractWeightMapNode::Evaluate(UE::Dataflow::FContext& Cont
 		SetValue(Context, MoveTemp(Result), &ExtractedWeightMap);
 	}
 }
+
+
+
+FChaosClothAssetExtractSelectionSetNode::FChaosClothAssetExtractSelectionSetNode(const UE::Dataflow::FNodeParameters& InParam, FGuid InGuid)
+	: FDataflowNode(InParam, InGuid)
+{
+	RegisterInputConnection(&Collection);
+	RegisterInputConnection(&Selection.StringValue, GET_MEMBER_NAME_CHECKED(FChaosClothAssetConnectableIStringValue, StringValue));
+	RegisterInputConnection(&DynamicMesh)
+		.SetCanHidePin(true)
+		.SetPinIsHidden(true);
+	RegisterOutputConnection(&ExtractedSelectionSet);
+	RegisterOutputConnection(&ExtractedSelectionArray);
+}
+
+void FChaosClothAssetExtractSelectionSetNode::Evaluate(UE::Dataflow::FContext& Context, const FDataflowOutput* Out) const
+{
+	if (Out->IsA(&ExtractedSelectionSet) || Out->IsA(&ExtractedSelectionArray))
+	{
+		using namespace UE::Chaos::ClothAsset;
+
+		TSet<int32> Result;
+
+		FManagedArrayCollection InCollection = GetValue<FManagedArrayCollection>(Context, &Collection);
+		const TSharedRef<FManagedArrayCollection> ClothCollection = MakeShared<FManagedArrayCollection>(MoveTemp(InCollection));
+		FCollectionClothConstFacade ClothFacade(ClothCollection);
+		FCollectionClothSelectionConstFacade SelectionFacade(ClothCollection);
+		if (ClothFacade.IsValid() && SelectionFacade.IsValid())
+		{
+			const FName SelectionName(*GetValue<FString>(Context, &Selection.StringValue));
+			if (SelectionFacade.HasSelection(SelectionName))
+			{
+				const FName SelectionGroup = SelectionFacade.GetSelectionGroup(SelectionName);
+				if (SelectionGroup == ClothCollectionGroup::SimVertices3D || SelectionGroup == ClothCollectionGroup::RenderVertices)
+				{
+					Result = SelectionFacade.GetSelectionSet(SelectionName);
+				}
+			}
+
+			if (bReorderForDynamicMesh && !Result.IsEmpty())
+			{
+				if (const TObjectPtr<const UDynamicMesh> InMesh = GetValue(Context, &DynamicMesh))
+				{
+					const UE::Geometry::FDynamicMesh3& DynMesh = InMesh->GetMeshRef();
+					FClothPatternToDynamicMeshMappingSupport ClothMapping(DynMesh);
+					TSet<int32> ReorderedResult;
+					ReorderedResult.Reserve(Result.Num());
+					for (int32 VertexID : DynMesh.VertexIndicesItr())
+					{
+						const int32 ClothMeshID = ClothMapping.GetOriginalVertexID(VertexID);
+						if (Result.Contains(ClothMeshID))
+						{
+							ReorderedResult.Add(VertexID);
+						}
+					}
+
+					Result = MoveTemp(ReorderedResult);
+				}
+			}
+		}
+		SetValue(Context, Result.Array(), &ExtractedSelectionArray);
+		SetValue(Context, MoveTemp(Result), &ExtractedSelectionSet);
+	}
+}
 #undef LOCTEXT_NAMESPACE
diff --git a/Engine/Plugins/ChaosClothAssetEditor/Source/ChaosClothAssetDataflowNodes/Private/ChaosClothAsset/CustomRegionResizingNode.cpp b/Engine/Plugins/ChaosClothAssetEditor/Source/ChaosClothAssetDataflowNodes/Private/ChaosClothAsset/CustomRegionResizingNode.cpp
new file mode 100644
--- /dev/null
+++ b/Engine/Plugins/ChaosClothAssetEditor/Source/ChaosClothAssetDataflowNodes/Private/ChaosClothAsset/CustomRegionResizingNode.cpp
@@ -0,0 +1,153 @@
+// Copyright Epic Games, Inc. All Rights Reserved.
+
+#include "ChaosClothAsset/CustomRegionResizingNode.h"
+#include "ChaosClothAsset/ClothCollectionAttribute.h"
+#include "ChaosClothAsset/ClothGeometryTools.h"
+#include "Dataflow/DataflowInputOutput.h"
+
+#include UE_INLINE_GENERATED_CPP_BY_NAME(CustomRegionResizingNode)
+
+#define LOCTEXT_NAMESPACE "ChaosClothAssetCustomRegionResizingNode"
+
+FChaosClothAssetCustomRegionResizingNode::FChaosClothAssetCustomRegionResizingNode(const UE::Dataflow::FNodeParameters& InParam, FGuid InGuid)
+	: FDataflowNode(InParam, InGuid)
+{
+	using namespace UE::Chaos::ClothAsset;
+	SimCustomResizingBlendName = ClothCollectionAttribute::SimCustomResizingBlend.ToString();
+	RenderCustomResizingBlend = ClothCollectionAttribute::RenderCustomResizingBlend.ToString();
+
+	RegisterInputConnection(&Collection);
+	RegisterOutputConnection(&Collection, &Collection);
+	RegisterOutputConnection(&SimCustomResizingBlendName);
+	RegisterOutputConnection(&RenderCustomResizingBlend);
+
+	check(GetNumInputs() == NumRequiredInputs);
+
+	// Add a set of pins to start
+	for (int32 Index = 0; Index < NumInitialOptionalInputs; ++Index)
+	{
+		AddPins();
+	}
+}
+
+void FChaosClothAssetCustomRegionResizingNode::Evaluate(UE::Dataflow::FContext& Context, const FDataflowOutput* Out) const
+{
+	if (Out->IsA<FManagedArrayCollection>(&Collection))
+	{
+		using namespace UE::Chaos::ClothAsset;
+		using namespace Chaos::Softs;
+
+		FManagedArrayCollection InCollection = GetValue<FManagedArrayCollection>(Context, &Collection);
+		const TSharedRef<FManagedArrayCollection> ClothCollection = MakeShared<FManagedArrayCollection>(MoveTemp(InCollection));
+
+		TArray<FName> InputSets;
+		TArray<int32> InputTypes;
+		InputSets.SetNum(InputGroupData.Num());
+		InputTypes.SetNum(InputGroupData.Num());
+		for (int32 Index = 0; Index < InputGroupData.Num(); ++Index)
+		{
+			InputSets[Index] = FName(*GetValue(Context, GetConnectionReference(Index)));
+			InputTypes[Index] = (int32)InputGroupData[Index].ResizingType;
+		}
+
+		FClothGeometryTools::SetGroupResizingData(ClothCollection, InputSets, InputTypes);
+		SetValue(Context, MoveTemp(*ClothCollection), &Collection);
+	}
+	else if(Out->IsA<FString>(&SimCustomResizingBlendName))
+	{
+		SetValue(Context, SimCustomResizingBlendName, &SimCustomResizingBlendName);
+	}
+	else if (Out->IsA<FString>(&RenderCustomResizingBlend))
+	{
+		SetValue(Context, RenderCustomResizingBlend, &RenderCustomResizingBlend);
+	}
+}
+
+TArray<UE::Dataflow::FPin> FChaosClothAssetCustomRegionResizingNode::AddPins()
+{
+	const int32 Index = InputGroupData.AddDefaulted();
+	const FDataflowInput& Input = RegisterInputArrayConnection(GetConnectionReference(Index), GET_MEMBER_NAME_CHECKED(FChaosClothAssetConnectableIStringValue, StringValue));
+	return { { UE::Dataflow::FPin::EDirection::INPUT, Input.GetType(), Input.GetName() } };
+}
+
+TArray<UE::Dataflow::FPin> FChaosClothAssetCustomRegionResizingNode::GetPinsToRemove() const
+{
+	const int32 Index = InputGroupData.Num() - 1;
+	check(InputGroupData.IsValidIndex(Index));
+	if (const FDataflowInput* const Input = FindInput(GetConnectionReference(Index)))
+	{
+		return { { UE::Dataflow::FPin::EDirection::INPUT, Input->GetType(), Input->GetName() } };
+	}
+	return Super::GetPinsToRemove();
+}
+
+void FChaosClothAssetCustomRegionResizingNode::OnPinRemoved(const UE::Dataflow::FPin& Pin)
+{
+	const int32 Index = InputGroupData.Num() - 1;
+	check(InputGroupData.IsValidIndex(Index));
+#if DO_CHECK
+	const FDataflowInput* const Input = FindInput(GetConnectionReference(Index));
+	check(Input);
+	check(Input->GetName() == Pin.Name);
+	check(Input->GetType() == Pin.Type);
+#endif
+	InputGroupData.SetNum(Index);
+
+	return Super::OnPinRemoved(Pin);
+}
+
+void FChaosClothAssetCustomRegionResizingNode::PostSerialize(const FArchive& Ar)
+{
+	if (Ar.IsLoading())
+	{
+		for (int32 Index = 0; Index < NumInitialOptionalInputs; ++Index)
+		{
+			check(FindInput(GetConnectionReference(Index)));
+		}
+
+		for (int32 Index = NumInitialOptionalInputs; Index < InputGroupData.Num(); ++Index)
+		{
+			FindOrRegisterInputArrayConnection(GetConnectionReference(Index));
+		}
+		if (Ar.IsTransacting())
+		{
+			const int32 OrigNumRegisteredInputs = GetNumInputs();
+			check(OrigNumRegisteredInputs >= NumRequiredInputs + NumInitialOptionalInputs);
+			const int32 OrigNumOptionalInputs = InputGroupData.Num();
+			const int32 OrigNumRegisteredOptionalInputs = OrigNumRegisteredInputs - NumRequiredInputs;
+			if (OrigNumRegisteredOptionalInputs > OrigNumOptionalInputs)
+			{
+				// Inputs have been removed.
+				// Temporarily expand Collections so we can get connection references.
+				InputGroupData.SetNum(OrigNumRegisteredOptionalInputs);
+				for (int32 Index = OrigNumOptionalInputs; Index < InputGroupData.Num(); ++Index)
+				{
+					UnregisterInputConnection(GetConnectionReference(Index));
+				}
+				InputGroupData.SetNum(OrigNumOptionalInputs);
+			}
+		}
+		else
+		{
+			ensureAlways(InputGroupData.Num() + NumRequiredInputs == GetNumInputs());
+		}
+	}
+}
+
+UE::Dataflow::TConnectionReference<FString> FChaosClothAssetCustomRegionResizingNode::GetConnectionReference(int32 Index) const
+{
+	return { &InputGroupData[Index].InputSet.StringValue, Index, &InputGroupData };
+}
+
+TArray<FChaosClothAssetCustomRegionResizingNode::FRegionData> FChaosClothAssetCustomRegionResizingNode::GetRegionData(UE::Dataflow::FContext& Context) const
+{
+	TArray<FChaosClothAssetCustomRegionResizingNode::FRegionData> BindingSetData;
+	BindingSetData.SetNumUninitialized(InputGroupData.Num());
+	for (int32 Index = 0; Index < InputGroupData.Num(); ++Index)
+	{
+		BindingSetData[Index].InputSet = *GetValue(Context, GetConnectionReference(Index));
+		BindingSetData[Index].ResizingType = InputGroupData[Index].ResizingType;
+	}
+	return BindingSetData;
+}
+#undef LOCTEXT_NAMESPACE
diff --git a/Engine/Plugins/ChaosClothAssetEditor/Source/ChaosClothAssetDataflowNodes/Public/ChaosClothAsset/ApplyResizingNode.h b/Engine/Plugins/ChaosClothAssetEditor/Source/ChaosClothAssetDataflowNodes/Public/ChaosClothAsset/ApplyResizingNode.h
--- a/Engine/Plugins/ChaosClothAssetEditor/Source/ChaosClothAssetDataflowNodes/Public/ChaosClothAsset/ApplyResizingNode.h
+++ b/Engine/Plugins/ChaosClothAssetEditor/Source/ChaosClothAssetDataflowNodes/Public/ChaosClothAsset/ApplyResizingNode.h
@@ -14,7 +14,7 @@ USTRUCT(Meta = (DataflowCloth, Experimental))
 struct FChaosClothAssetApplyResizingNode final : public FDataflowNode
 {
 	GENERATED_USTRUCT_BODY()
-	DATAFLOW_NODE_DEFINE_INTERNAL(FChaosClothAssetApplyResizingNode, "ApplyResizing", "Cloth", "Apply Cloth Asset Resizing")
+	DATAFLOW_NODE_DEFINE_INTERNAL(FChaosClothAssetApplyResizingNode, "ApplyResizing", "Cloth", "Apply Cloth Outfit Resizing")
 
 public:
 
@@ -32,7 +32,7 @@ private:
 	UPROPERTY(meta = (DataflowInput))
 	TObjectPtr<const USkeletalMesh> TargetSkeletalMesh;
 
-	/** Which Skeletal Mesh LOD to use for the TargetSkeletalMesh */
+	/** Source and Target mesh LOD.*/
 	UPROPERTY(EditAnywhere, Category = "Resizing", meta = (DataflowInput, ClampMin = 0))
 	int32 SkeletalMeshLODIndex = 0;
 
@@ -43,4 +43,12 @@ private:
 	/** Force apply to the render mesh. Otherwise, the sim mesh will be resized if it exists, only resizing the render mesh if no sim mesh data exists. */
 	UPROPERTY(EditAnywhere, Category = "Resizing", meta = (DataflowInput))
 	bool bForceApplyToRenderMesh = false;
+
+	/** The source mesh. Used for CustomResizingRegions. Must have matching vertices with TargetMesh */
+	UPROPERTY(meta = (DataflowInput))
+	TObjectPtr<const USkeletalMesh> SourceSkeletalMesh;
+
+	/** Skip applying Custom Region Resizing data. */
+	UPROPERTY(EditAnywhere, Category = "Resizing", meta = (DataflowInput))
+	bool bSkipCustomRegionResizing = false;
 };
\ No newline at end of file
diff --git a/Engine/Plugins/ChaosClothAssetEditor/Source/ChaosClothAssetDataflowNodes/Public/ChaosClothAsset/ClothCollectionToDynamicMeshNode.h b/Engine/Plugins/ChaosClothAssetEditor/Source/ChaosClothAssetDataflowNodes/Public/ChaosClothAsset/ClothCollectionToDynamicMeshNode.h
--- a/Engine/Plugins/ChaosClothAssetEditor/Source/ChaosClothAssetDataflowNodes/Public/ChaosClothAsset/ClothCollectionToDynamicMeshNode.h
+++ b/Engine/Plugins/ChaosClothAssetEditor/Source/ChaosClothAssetDataflowNodes/Public/ChaosClothAsset/ClothCollectionToDynamicMeshNode.h
@@ -147,4 +147,46 @@ private:
 	UPROPERTY(meta = (DataflowOutput))
 	TArray<float> ExtractedWeightMap;
 
+};
+
+
+/** Extract a selection set from a Cloth Collection.
+ */
+USTRUCT(Meta = (DataflowCloth, Experimental))
+struct FChaosClothAssetExtractSelectionSetNode : public FDataflowNode
+{
+	GENERATED_USTRUCT_BODY()
+	DATAFLOW_NODE_DEFINE_INTERNAL(FChaosClothAssetExtractSelectionSetNode, "ExtractClothSelectionSet", "Cloth", "Extract Cloth Selection Set")
+
+public:
+
+	FChaosClothAssetExtractSelectionSetNode(const UE::Dataflow::FNodeParameters& InParam, FGuid InGuid = FGuid::NewGuid());
+
+	virtual void Evaluate(UE::Dataflow::FContext& Context, const FDataflowOutput* Out) const override;
+
+private:
+
+	UPROPERTY(Meta = (DataflowInput))
+	FManagedArrayCollection Collection;
+
+	/** Name of the selection set to be extracted. Currently only SimVertices3D and RenderVertices sets are supported. */
+	UPROPERTY(EditAnywhere, Category = "Extract Selection Set", Meta = (DataflowInput))
+	FChaosClothAssetConnectableIStringValue Selection = { TEXT("Selection") };
+
+	/** Reorder extracted indices to match the order of a DynamicMesh that was created via ClothCollectionToDynamicMesh */
+	UPROPERTY(EditAnywhere, Category = "Extract Selection Set")
+	bool bReorderForDynamicMesh = true;
+
+	/** Dynamic mesh used to reorder indices. */
+	UPROPERTY(meta = (DataflowInput, EditCondition = "bReorderForDynamicMesh"))
+	TObjectPtr<UDynamicMesh> DynamicMesh;
+
+	/** Extracted Selection Set as a Set */
+	UPROPERTY(meta = (DataflowOutput))
+	TSet<int32> ExtractedSelectionSet;
+
+	/** Extracted Selection Set as an array */
+	UPROPERTY(meta = (DataflowOutput))
+	TArray<int32> ExtractedSelectionArray;
+
 };
\ No newline at end of file
diff --git a/Engine/Plugins/ChaosClothAssetEditor/Source/ChaosClothAssetDataflowNodes/Public/ChaosClothAsset/CustomRegionResizingNode.h b/Engine/Plugins/ChaosClothAssetEditor/Source/ChaosClothAssetDataflowNodes/Public/ChaosClothAsset/CustomRegionResizingNode.h
new file mode 100644
--- /dev/null
+++ b/Engine/Plugins/ChaosClothAssetEditor/Source/ChaosClothAssetDataflowNodes/Public/ChaosClothAsset/CustomRegionResizingNode.h
@@ -0,0 +1,91 @@
+// Copyright Epic Games, Inc. All Rights Reserved.
+
+#pragma once 
+
+#include "Dataflow/DataflowNode.h"
+#include "GeometryCollection/ManagedArrayCollection.h"
+#include "MeshResizing/RBFInterpolation.h"
+#include "MeshResizing/CustomRegionResizing.h"
+#include "ChaosClothAsset/ConnectableValue.h"
+#include "CustomRegionResizingNode.generated.h"
+
+
+/** Data to generate a Custom Resizing Region.*/
+USTRUCT()
+struct FChaosClothAssetCustomRegionResizingInput
+{
+	GENERATED_USTRUCT_BODY()
+
+	/** Source Set. Must be convertible to a vertex set. Render sets will only be used if this ClothAsset does not have a sim mesh.*/
+	UPROPERTY(EditAnywhere, Category = "Custom Region Resizing")
+	FChaosClothAssetConnectableIStringValue InputSet = {TEXT("InputSet")};
+
+	/** Custom Resizing type. */
+	UPROPERTY(EditAnywhere, Category = "Custom Region Resizing")
+	EMeshResizingCustomRegionType ResizingType = EMeshResizingCustomRegionType::TrilinearInterpolation;
+};
+
+/** Node for adding custom region resizing data used by the ChaosOutfitAsset's Resizeable Outfit.*/
+USTRUCT(Meta = (DataflowCloth, Experimental))
+struct FChaosClothAssetCustomRegionResizingNode final : public FDataflowNode
+{
+	GENERATED_USTRUCT_BODY()
+	DATAFLOW_NODE_DEFINE_INTERNAL(FChaosClothAssetCustomRegionResizingNode, "CustomRegionResizing", "Cloth", "Cloth Outfit Custom Region Resizing")
+
+public:
+
+	FChaosClothAssetCustomRegionResizingNode(const UE::Dataflow::FNodeParameters& InParam, FGuid InGuid = FGuid::NewGuid());
+
+private:
+	//~ Begin FDataflowNode interface
+	virtual void Evaluate(UE::Dataflow::FContext& Context, const FDataflowOutput* Out) const override;
+	virtual TArray<UE::Dataflow::FPin> AddPins() override;
+	virtual bool CanAddPin() const override 
+	{ 
+		return true; 
+	}
+	virtual bool CanRemovePin() const override 
+	{ 
+		return InputGroupData.Num() > NumInitialOptionalInputs;
+	}
+	virtual TArray<UE::Dataflow::FPin> GetPinsToRemove() const override;
+	virtual void OnPinRemoved(const UE::Dataflow::FPin& Pin) override;
+	virtual void PostSerialize(const FArchive& Ar) override;
+	//~ End FDataflowNode interface
+
+
+	UE::Dataflow::TConnectionReference<FString> GetConnectionReference(int32 Index) const;
+
+	static constexpr int32 NumRequiredInputs = 1; // non-construction set inputs
+	static constexpr int32 NumInitialOptionalInputs = 1;
+
+	struct FRegionData
+	{
+		FString InputSet;
+		EMeshResizingCustomRegionType ResizingType;
+	};
+	TArray<FRegionData> GetRegionData(UE::Dataflow::FContext& Context) const;
+
+	UPROPERTY(Meta = (Dataflowinput, DataflowOutput, DataflowPassthrough = "Collection"))
+	FManagedArrayCollection Collection;
+
+	/** Input custom region data. This data will be stored in the resulting Cloth Asset, and available for use when resizing.*/
+	UPROPERTY(EditAnywhere, EditFixedSize, Category = "Custom Region Resizing", Meta = (SkipInDisplayNameChain))
+	TArray<FChaosClothAssetCustomRegionResizingInput> InputGroupData;
+
+	/** The name of the sim mesh weight map generated by this node detailing the vertices should use Custom Regions as opposed to 
+	* traditional resizing interpolation.
+	* Value ranges between 0 (use traditional resizing interpolation) and 1 (use resizing custom group interpolation).
+	* These regions can be blended if desired.
+	* The name of this sim mesh weight map cannot be changed and is only provided for further tweaking.*/
+	UPROPERTY(VisibleAnywhere, Category = "Custom Region Resizing", Meta = (DataflowOutput))
+	FString SimCustomResizingBlendName;
+
+	/** The name of the render mesh weight map generated by this node detailing the vertices should use Custom Regions as opposed to
+	* traditional resizing interpolation.
+	* Value ranges between 0 (use traditional resizing interpolation) and 1 (use resizing custom group interpolation).
+	* These regions can be blended if desired.
+	* The name of this render mesh weight map cannot be changed and is only provided for further tweaking.*/
+	UPROPERTY(VisibleAnywhere, Category = "Custom Region Resizing", Meta = (DataflowOutput))
+	FString RenderCustomResizingBlend;
+};
\ No newline at end of file
diff --git a/Engine/Plugins/Experimental/MeshResizing/Source/MeshResizingCore/Private/MeshResizing/CustomRegionResizing.cpp b/Engine/Plugins/Experimental/MeshResizing/Source/MeshResizingCore/Private/MeshResizing/CustomRegionResizing.cpp
new file mode 100644
--- /dev/null
+++ b/Engine/Plugins/Experimental/MeshResizing/Source/MeshResizingCore/Private/MeshResizing/CustomRegionResizing.cpp
@@ -0,0 +1,197 @@
+// Copyright Epic Games, Inc. All Rights Reserved.
+
+#include "MeshResizing/CustomRegionResizing.h"
+#include "DynamicMesh/DynamicMesh3.h"
+#include "DynamicMesh/DynamicMeshAttributeSet.h"
+#include "DynamicMesh/DynamicMeshAABBTree3.h"
+#include "IndexTypes.h"
+
+#include UE_INLINE_GENERATED_CPP_BY_NAME(CustomRegionResizing)
+
+namespace UE::MeshResizing
+{
+	namespace Private
+	{
+		static bool CalculateTriangleFrameNoAttributes(const UE::Geometry::FTriangle3d& Triangle, FVector3f& OutTangentU, FVector3f& OutTangentV, FVector3f& OutNormal)
+		{
+			OutTangentU = FVector3f(Triangle.V[1] - Triangle.V[0]).GetSafeNormal();
+			OutNormal = OutTangentU.Cross(FVector3f(Triangle.V[2] - Triangle.V[0])).GetSafeNormal();
+			OutTangentV = OutNormal.Cross(OutTangentU);
+			return !OutTangentU.IsZero() && !OutTangentV.IsZero() && !OutNormal.IsZero();
+		}
+
+		static bool CalculateFrameFromAttributes(const UE::Geometry::FDynamicMesh3& SourceMesh, const int32 FaceIndex, const FVector3f& Barys, FVector3f& OutTangentU, FVector3f& OutTangentV, FVector3f& OutNormal)
+		{
+			check(SourceMesh.HasAttributes() && SourceMesh.Attributes()->NumNormalLayers() >= 2);
+
+			const UE::Geometry::FDynamicMeshNormalOverlay* const NormalOverlay = SourceMesh.Attributes()->PrimaryNormals();
+			check(NormalOverlay);
+			NormalOverlay->GetTriBaryInterpolate(FaceIndex, &Barys.X, &OutNormal.X);
+			OutNormal = OutNormal.GetSafeNormal();
+
+			const UE::Geometry::FDynamicMeshNormalOverlay* const TangentOverlay = SourceMesh.Attributes()->PrimaryTangents();
+			check(TangentOverlay);
+			TangentOverlay->GetTriBaryInterpolate(FaceIndex, &Barys.X, &OutTangentU.X);
+			OutTangentU = OutTangentU.GetSafeNormal();
+			OutTangentU = (OutTangentU - (OutTangentU | OutNormal) * OutNormal).GetSafeNormal();
+			OutTangentV = OutNormal.Cross(OutTangentU);
+			return !OutTangentU.IsZero() && !OutTangentV.IsZero() && !OutNormal.IsZero();
+		}
+	}
+
+	void FCustomRegionResizing::GenerateCustomRegion(const TConstArrayView<FVector3f>& BoundPositions, const UE::Geometry::FDynamicMesh3& SourceMesh, const TSet<int32>& RegionVertices, FMeshResizingCustomRegion& OutData)
+	{
+		
+		OutData.Reset();
+		OutData.RegionVertices.Reserve(RegionVertices.Num());
+
+		FVector3d Centroid(0.);
+		for (int32 RigidBoundVertex : RegionVertices)
+		{
+			if (BoundPositions.IsValidIndex(RigidBoundVertex))
+			{
+				OutData.RegionVertices.Add(RigidBoundVertex);
+				Centroid += FVector3d(BoundPositions[RigidBoundVertex]);
+			}
+		}
+		if (OutData.RegionVertices.Num() == 0)
+		{
+			return;
+		}
+
+		Centroid /= OutData.RegionVertices.Num();
+
+		const bool bHasValidNormalAndTangentLayers = SourceMesh.HasAttributes() && SourceMesh.Attributes()->NumNormalLayers() >= 2;
+
+		UE::Geometry::IMeshSpatial::FQueryOptions QueryOptions;
+
+		if (bHasValidNormalAndTangentLayers)
+		{
+			QueryOptions.TriangleFilterF = [&SourceMesh, &Centroid](int32 FaceIndex)
+				{
+					UE::Geometry::FDistPoint3Triangle3d DistQuery = UE::Geometry::TMeshQueries<UE::Geometry::FDynamicMesh3>::TriangleDistance(SourceMesh, FaceIndex, Centroid);
+
+					FVector3f TangentU, TangentV, Normal;
+					return Private::CalculateFrameFromAttributes(SourceMesh, FaceIndex, FVector3f(DistQuery.TriangleBaryCoords), TangentU, TangentV, Normal);
+				};
+		}
+		else
+		{
+			QueryOptions.TriangleFilterF = [&SourceMesh](int32 FaceIndex)
+				{
+					UE::Geometry::FTriangle3d Tri;
+					SourceMesh.GetTriVertices(FaceIndex, Tri.V[0], Tri.V[1], Tri.V[2]);
+
+					FVector3f TangentU, TangentV, Normal;
+					return Private::CalculateTriangleFrameNoAttributes(Tri, TangentU, TangentV, Normal);
+				};
+		}
+
+		UE::Geometry::FDynamicMeshAABBTree3 AABBTree(&SourceMesh);
+		double DistSq;
+		OutData.SourceFaceIndex = AABBTree.FindNearestTriangle(Centroid, DistSq, QueryOptions);
+		static_assert(IndexConstants::InvalidID == INDEX_NONE);
+		if (OutData.SourceFaceIndex == INDEX_NONE)
+		{
+			return;
+		}
+
+		UE::Geometry::FDistPoint3Triangle3d DistQuery = UE::Geometry::TMeshQueries<UE::Geometry::FDynamicMesh3>::TriangleDistance(SourceMesh, OutData.SourceFaceIndex, Centroid);
+		const float Distance = FMath::Sqrt(DistQuery.ComputeResult());
+		OutData.SourceBaryCoords = FVector3f(DistQuery.TriangleBaryCoords);
+		OutData.SourceOrigin = DistQuery.ClosestTrianglePoint;
+
+		// Calculate local frame for triangle.
+		if (bHasValidNormalAndTangentLayers)
+		{
+			if (!ensure(Private::CalculateFrameFromAttributes(SourceMesh, OutData.SourceFaceIndex, OutData.SourceBaryCoords, OutData.SourceAxis0, OutData.SourceAxis1, OutData.SourceAxis2)))
+			{
+				OutData.Reset();
+				return;
+			}
+		}
+		else
+		{
+			if (!ensure(Private::CalculateTriangleFrameNoAttributes(DistQuery.Triangle, OutData.SourceAxis0, OutData.SourceAxis1, OutData.SourceAxis2)))
+			{
+				OutData.Reset();
+				return;
+			}
+		}
+		const FMatrix TriangleMatrix(FVector3d(OutData.SourceAxis0), FVector3d(OutData.SourceAxis1), FVector3d(OutData.SourceAxis2), OutData.SourceOrigin);
+		OutData.RegionBoundsCentroid = FVector3f(TriangleMatrix.InverseTransformPosition(Centroid));
+
+		// Need to calculate transformed extents.
+		OutData.RegionBoundsExtents = FVector3f(0.f);
+		for (int32 RigidBoundVertex : OutData.RegionVertices)
+		{
+			const FVector3f TransformedPoint = FVector3f(TriangleMatrix.InverseTransformPosition(FVector3d(BoundPositions[RigidBoundVertex])));
+			const FVector3f TransformedDelta = (TransformedPoint - OutData.RegionBoundsCentroid).GetAbs();
+			OutData.RegionBoundsExtents = OutData.RegionBoundsExtents.ComponentMax(TransformedDelta);
+		}
+		const FVector3f MinBounds = OutData.RegionBoundsCentroid - OutData.RegionBoundsExtents;
+		FVector3f RecipBoundsSize;
+		RecipBoundsSize.X = OutData.RegionBoundsExtents.X > UE_SMALL_NUMBER ? .5f / OutData.RegionBoundsExtents.X : 0.f;
+		RecipBoundsSize.Y = OutData.RegionBoundsExtents.Y > UE_SMALL_NUMBER ? .5f / OutData.RegionBoundsExtents.Y : 0.f;
+		RecipBoundsSize.Z = OutData.RegionBoundsExtents.Z > UE_SMALL_NUMBER ? .5f / OutData.RegionBoundsExtents.Z : 0.f;
+
+		OutData.RegionVertexCoords.Reset(OutData.RegionVertices.Num());
+		for (int32 RigidBoundVertex : OutData.RegionVertices)
+		{
+			// Transform to bounds space
+			const FVector3f TransformedPoint = FVector3f(TriangleMatrix.InverseTransformPosition(FVector3d(BoundPositions[RigidBoundVertex])));
+			OutData.RegionVertexCoords.Emplace((TransformedPoint - MinBounds) * RecipBoundsSize);
+		}
+	}
+
+	bool FCustomRegionResizing::CalculateFrameForCustomRegion(const UE::Geometry::FDynamicMesh3& SourceMesh, const FMeshResizingCustomRegion& BindingGroup, FVector3d& OutOrigin, FVector3f& OutTangentU, FVector3f& OutTangentV, FVector3f& OutNormal)
+	{
+
+		UE::Geometry::FTriangle3d Triangle;
+		SourceMesh.GetTriVertices(BindingGroup.SourceFaceIndex, Triangle.V[0], Triangle.V[1], Triangle.V[2]);
+
+		OutOrigin = Triangle.BarycentricPoint(FVector(BindingGroup.SourceBaryCoords));
+
+		const bool bHasValidNormalAndTangentLayers = SourceMesh.HasAttributes() && SourceMesh.Attributes()->NumNormalLayers() >= 2;
+		if (bHasValidNormalAndTangentLayers)
+		{
+			return Private::CalculateFrameFromAttributes(SourceMesh, BindingGroup.SourceFaceIndex, BindingGroup.SourceBaryCoords, OutTangentU, OutTangentV, OutNormal);
+		}
+		else
+		{
+			return Private::CalculateTriangleFrameNoAttributes(Triangle, OutTangentU, OutTangentV, OutNormal);
+		}
+	}
+
+	void FCustomRegionResizing::InterpolateCustomRegionPoints(const FMeshResizingCustomRegion& BindingGroup, const TConstArrayView<FVector3d>& BoundsCorners, TArrayView<FVector3f> BoundPositions)
+	{
+		check(BoundsCorners.Num() == 8);
+		for (int32 Index = 0; Index < BindingGroup.RegionVertices.Num(); ++Index)
+		{
+			if (BoundPositions.IsValidIndex(BindingGroup.RegionVertices[Index]))
+			{
+				// BoundsCorners are in the order of FOrientedBox::CalcVertices
+				// Xijk = Center + (Signs[i] * ExtentX) * AxisX + (Signs[j] * ExtentY) * AxisY + (Signs[k] * ExtentZ) * AxisZ where Signs[] = { -1, 1 };
+				// BoundsCorners[0] = X000
+				// BoundsCorners[1] = X001
+				// BoundsCorners[2] = X010
+				// BoundsCorners[3] = X011
+				// BoundsCorners[4] = X100
+				// BoundsCorners[5] = X101
+				// BoundsCorners[6] = X110
+				// BoundsCorners[7] = X111
+				const FVector3d X00 = BoundsCorners[0] * (1. - BindingGroup.RegionVertexCoords[Index].X) + BoundsCorners[4] * BindingGroup.RegionVertexCoords[Index].X;
+				const FVector3d X01 = BoundsCorners[1] * (1. - BindingGroup.RegionVertexCoords[Index].X) + BoundsCorners[5] * BindingGroup.RegionVertexCoords[Index].X;
+				const FVector3d X10 = BoundsCorners[2] * (1. - BindingGroup.RegionVertexCoords[Index].X) + BoundsCorners[6] * BindingGroup.RegionVertexCoords[Index].X;
+				const FVector3d X11 = BoundsCorners[3] * (1. - BindingGroup.RegionVertexCoords[Index].X) + BoundsCorners[7] * BindingGroup.RegionVertexCoords[Index].X;
+
+				const FVector X0 = X00 * (1. - BindingGroup.RegionVertexCoords[Index].Y) + X10 * BindingGroup.RegionVertexCoords[Index].Y;
+				const FVector X1 = X01 * (1. - BindingGroup.RegionVertexCoords[Index].Y) + X11 * BindingGroup.RegionVertexCoords[Index].Y;
+
+				const FVector X = X0 * (1. - BindingGroup.RegionVertexCoords[Index].Z) + X1 * BindingGroup.RegionVertexCoords[Index].Z;
+
+				BoundPositions[BindingGroup.RegionVertices[Index]] = FVector3f(X);
+			}
+		}
+	}
+}
\ No newline at end of file
diff --git a/Engine/Plugins/Experimental/MeshResizing/Source/MeshResizingCore/Private/MeshResizing/RBFInterpolation.cpp b/Engine/Plugins/Experimental/MeshResizing/Source/MeshResizingCore/Private/MeshResizing/RBFInterpolation.cpp
--- a/Engine/Plugins/Experimental/MeshResizing/Source/MeshResizingCore/Private/MeshResizing/RBFInterpolation.cpp
+++ b/Engine/Plugins/Experimental/MeshResizing/Source/MeshResizingCore/Private/MeshResizing/RBFInterpolation.cpp
@@ -319,8 +319,24 @@ namespace UE::MeshResizing
 		}
 	}
 
-	void FRBFInterpolation::DeformPoints(const TArray<FVector3f>& TargetPositions, const FMeshResizingRBFInterpolationData& InterpolationData, TArray<FVector3d>& Points)
+	void FRBFInterpolation::DeformPoints(const FMeshDescription& TargetMesh, const FMeshResizingRBFInterpolationData& InterpolationData, TArray<FVector3d>& Points)
 	{
+		TArray<FVector3f> TargetPositions;
+		TargetPositions.Reserve(InterpolationData.SampleIndices.Num());
+		for (int32 Index = 0; Index < InterpolationData.SampleIndices.Num(); ++Index)
+		{
+			const int32 SampleIndex = InterpolationData.SampleIndices[Index];
+			if (ensure(TargetMesh.IsVertexValid(SampleIndex)))
+			{
+				TargetPositions.Emplace(TargetMesh.GetVertexPosition(SampleIndex));
+			}
+			else
+			{
+				// Just use Source positions
+				TargetPositions.Emplace(InterpolationData.SampleRestPositions[Index]);
+			}
+		}
+
 		TArray<FVector3f> SampleDeformations;
 		Private::UpdateMeshSamples(TargetPositions, InterpolationData, SampleDeformations);
 		for (FVector3d& Point : Points)
diff --git a/Engine/Plugins/Experimental/MeshResizing/Source/MeshResizingCore/Public/MeshResizing/CustomRegionResizing.h b/Engine/Plugins/Experimental/MeshResizing/Source/MeshResizingCore/Public/MeshResizing/CustomRegionResizing.h
new file mode 100644
--- /dev/null
+++ b/Engine/Plugins/Experimental/MeshResizing/Source/MeshResizingCore/Public/MeshResizing/CustomRegionResizing.h
@@ -0,0 +1,82 @@
+// Copyright Epic Games, Inc. All Rights Reserved.
+
+#pragma once 
+#include "UObject/NameTypes.h"
+#include "Containers/Array.h"
+#include "Math/MathFwd.h"
+#include "Math/Matrix.h"
+
+#include "CustomRegionResizing.generated.h"
+
+namespace UE::Geometry
+{
+	class FDynamicMesh3;
+}
+
+UENUM()
+enum struct EMeshResizingCustomRegionType : uint8
+{
+	TrilinearInterpolation,
+	/* Hoping to add support for these in a future review
+	RigidRivet,
+	UniformScaleRivet,
+	NonUniformScaleRivet,
+	Rivet
+	*/
+};
+
+USTRUCT()
+struct MESHRESIZINGCORE_API FMeshResizingCustomRegion
+{
+	GENERATED_BODY()
+
+	UPROPERTY()
+	TArray<int32> RegionVertices;
+	UPROPERTY()
+	TArray<FVector3f> RegionVertexCoords; // Trilinear interpolation within Bounds: (0,0,0) = RegionBoundsCentroid - RegionBoundsExtents; (1,1,1) = RegionBoundsCentroid + RegionBoundsExtents
+
+	UPROPERTY()
+	int32 SourceFaceIndex;
+	UPROPERTY()
+	FVector3f SourceBaryCoords;
+	UPROPERTY()
+	FVector3d SourceOrigin;
+	UPROPERTY()
+	FVector3f SourceAxis0;
+	UPROPERTY()
+	FVector3f SourceAxis1;
+	UPROPERTY()
+	FVector3f SourceAxis2;
+
+	// Relative to Source
+	UPROPERTY()
+	FVector3f RegionBoundsCentroid;
+
+	UPROPERTY()
+	FVector3f RegionBoundsExtents;
+
+	void Reset()
+	{
+		RegionVertices.Reset();
+		RegionVertexCoords.Reset();
+		SourceFaceIndex = INDEX_NONE;
+		RegionBoundsCentroid = FVector3f(0.f);
+		RegionBoundsExtents = FVector3f(0.f);
+	}
+
+	bool IsValid() const
+	{
+		return SourceFaceIndex != INDEX_NONE && RegionVertices.Num() > 0;
+	}
+};
+
+namespace UE::MeshResizing
+{
+	struct MESHRESIZINGCORE_API FCustomRegionResizing
+	{
+		static void GenerateCustomRegion(const TConstArrayView<FVector3f>& BoundPositions, const UE::Geometry::FDynamicMesh3& SourceMesh, const TSet<int32>& BoundVertices, FMeshResizingCustomRegion& OutData);
+		static bool CalculateFrameForCustomRegion(const UE::Geometry::FDynamicMesh3& SourceMesh, const FMeshResizingCustomRegion& BindingGroup, FVector3d& OutOrigin, FVector3f& OutTangentU, FVector3f& OutTangentV, FVector3f& OutNormal);
+		static void InterpolateCustomRegionPoints(const FMeshResizingCustomRegion& BindingGroup, const TConstArrayView<FVector3d>& BoundsCorners, TArrayView<FVector3f> BoundPositions);
+	};
+}
+
diff --git a/Engine/Plugins/Experimental/MeshResizing/Source/MeshResizingCore/Public/MeshResizing/RBFInterpolation.h b/Engine/Plugins/Experimental/MeshResizing/Source/MeshResizingCore/Public/MeshResizing/RBFInterpolation.h
--- a/Engine/Plugins/Experimental/MeshResizing/Source/MeshResizingCore/Public/MeshResizing/RBFInterpolation.h
+++ b/Engine/Plugins/Experimental/MeshResizing/Source/MeshResizingCore/Public/MeshResizing/RBFInterpolation.h
@@ -34,7 +34,7 @@ namespace UE::MeshResizing
 		static void GenerateWeights(const FMeshDescription& BaseMesh, int32 NumInterpolationPoints, FMeshResizingRBFInterpolationData& OutData);
 		static void DeformPoints(const TArray<FVector3f>& TargetPositions, const FMeshResizingRBFInterpolationData& InterpolationData, bool bInterpolateNormals,  UE::Geometry::FDynamicMesh3& DeformingMesh);
 		static void DeformCoordinateFrames(const TArray<FVector3f>& TargetPositions, const FMeshResizingRBFInterpolationData& InterpolationData, bool bNormalize, bool bOrthogonalize, TArray<FMatrix44f>& Coordinates);
-		static void DeformPoints(const TArray<FVector3f>& TargetPositions, const FMeshResizingRBFInterpolationData& InterpolationData, TArray<FVector3d>& Points);
+		static void DeformPoints(const FMeshDescription& TargetMesh, const FMeshResizingRBFInterpolationData& InterpolationData, TArray<FVector3d>& Points);
 		static void DeformPoints(const FMeshDescription& TargetMesh, const FMeshResizingRBFInterpolationData& InterpolationData, bool bInterpolateNormals, UE::Geometry::FDynamicMesh3& DeformingMesh);
 		static void DeformPoints(const FMeshDescription& TargetMesh, const FMeshResizingRBFInterpolationData& InterpolationData, TArrayView<FVector3f> Points, TArrayView<FVector3f> Normals);
 		static void DeformPoints(const FMeshDescription& TargetMesh, const FMeshResizingRBFInterpolationData& InterpolationData, TArrayView<FVector3f> Points, TArrayView<FVector3f> Normals, TArrayView<FVector3f> TangentUs, TArrayView<FVector3f> TangentVs);
diff --git a/Engine/Plugins/Experimental/MeshResizing/Source/MeshResizingNodes/Private/MeshResizing/BaseBodyDataflowNodes.cpp b/Engine/Plugins/Experimental/MeshResizing/Source/MeshResizingNodes/Private/MeshResizing/BaseBodyDataflowNodes.cpp
--- a/Engine/Plugins/Experimental/MeshResizing/Source/MeshResizingNodes/Private/MeshResizing/BaseBodyDataflowNodes.cpp
+++ b/Engine/Plugins/Experimental/MeshResizing/Source/MeshResizingNodes/Private/MeshResizing/BaseBodyDataflowNodes.cpp
@@ -48,7 +48,8 @@ void FSkeletalMeshToMeshDataflowNode::Evaluate(UE::Dataflow::FContext& Context,
 						UE::Geometry::FDynamicMesh3& DynMesh = NewMesh->GetMeshRef();
 
 						FSkeletalMeshLODModelToDynamicMesh Converter;
-						Converter.Convert(SkeletalMeshLODModel, DynMesh);
+						constexpr bool bCopyTangentsTrue = true;
+						Converter.Convert(SkeletalMeshLODModel, DynMesh, bCopyTangentsTrue); // Tangents are needed by RigidBinding
 
 						// Set ImportedVertices
 						TArray<int32> VertexToImportedVertexIDMap;
