Morph Target code cleanup.

Fixing doc comments, renaming functions, renaming structs. No functional changes.

diff --git a/Engine/Plugins/ChaosClothAsset/Source/ChaosClothAsset/Private/ChaosClothAsset/ClothCollection.cpp b/Engine/Plugins/ChaosClothAsset/Source/ChaosClothAsset/Private/ChaosClothAsset/ClothCollection.cpp
--- a/Engine/Plugins/ChaosClothAsset/Source/ChaosClothAsset/Private/ChaosClothAsset/ClothCollection.cpp
+++ b/Engine/Plugins/ChaosClothAsset/Source/ChaosClothAsset/Private/ChaosClothAsset/ClothCollection.cpp
@@ -287,9 +287,9 @@ namespace UE::Chaos::ClothAsset
 		CustomResizingRegionType = ManagedArrayCollection->FindAttribute<int32>(ClothCollectionAttribute::CustomResizingRegionType, ClothCollectionGroup::CustomResizingRegions);
 	}
 
-	bool FClothCollection::IsValid(EClothCollectionOptionalSchemas OptionalSchemas) const
+	bool FClothCollection::IsValid(EClothCollectionExtendedSchemas OptionalSchemas) const
 	{
-		if (EnumHasAnyFlags(OptionalSchemas, EClothCollectionOptionalSchemas::CookedOnly))
+		if (EnumHasAnyFlags(OptionalSchemas, EClothCollectionExtendedSchemas::CookedOnly))
 		{
 			// Sim Morph Targets Group
 			return
@@ -309,7 +309,7 @@ namespace UE::Chaos::ClothAsset
 			SkeletalMeshPathName &&
 
 			// Solvers Group
-			(!EnumHasAnyFlags(OptionalSchemas, EClothCollectionOptionalSchemas::Solvers) ||
+			(!EnumHasAnyFlags(OptionalSchemas, EClothCollectionExtendedSchemas::Solvers) ||
 				(SolverGravity &&
 					SolverAirDamping &&
 					SolverSubSteps &&
@@ -337,7 +337,7 @@ namespace UE::Chaos::ClothAsset
 			RenderFacesStart &&
 			RenderFacesEnd &&
 			RenderMaterialPathName &&
-			(!EnumHasAnyFlags(OptionalSchemas, EClothCollectionOptionalSchemas::RenderDeformer) ||
+			(!EnumHasAnyFlags(OptionalSchemas, EClothCollectionExtendedSchemas::RenderDeformer) ||
 				RenderDeformerNumInfluences) &&
 
 			// Sim Faces Group
@@ -345,7 +345,7 @@ namespace UE::Chaos::ClothAsset
 			SimIndices3D &&
 
 			// Fabrics Group
-			(!EnumHasAnyFlags(OptionalSchemas, EClothCollectionOptionalSchemas::Fabrics) ||
+			(!EnumHasAnyFlags(OptionalSchemas, EClothCollectionExtendedSchemas::Fabrics) ||
 				(FabricBendingStiffness &&
 					FabricBucklingStiffness &&
 					FabricStretchStiffness &&
@@ -370,17 +370,17 @@ namespace UE::Chaos::ClothAsset
 			TetherReferenceLength &&
 			SimVertex2DLookup &&
 			SeamStitchLookup &&
-			(!EnumHasAnyFlags(OptionalSchemas, EClothCollectionOptionalSchemas::Resizing) ||
+			(!EnumHasAnyFlags(OptionalSchemas, EClothCollectionExtendedSchemas::Resizing) ||
 				SimCustomResizingBlend) &&
 
 			// Sim Morph Targets Group
-			(!EnumHasAnyFlags(OptionalSchemas, EClothCollectionOptionalSchemas::SimMorphTargets) ||
+			(!EnumHasAnyFlags(OptionalSchemas, EClothCollectionExtendedSchemas::SimMorphTargets) ||
 				(SimMorphTargetName &&
 					SimMorphTargetVerticesStart &&
 					SimMorphTargetVerticesEnd)) &&
 
 			// Sim Morph Target Vertices Group
-			(!EnumHasAnyFlags(OptionalSchemas, EClothCollectionOptionalSchemas::SimMorphTargets) ||
+			(!EnumHasAnyFlags(OptionalSchemas, EClothCollectionExtendedSchemas::SimMorphTargets) ||
 				(SimMorphTargetPositionDelta &&
 					SimMorphTargetTangentZDelta &&
 					SimMorphTargetSimVertex3DIndex)) &&
@@ -397,23 +397,23 @@ namespace UE::Chaos::ClothAsset
 			RenderColor &&
 			RenderBoneIndices &&
 			RenderBoneWeights &&
-			(!EnumHasAnyFlags(OptionalSchemas, EClothCollectionOptionalSchemas::RenderDeformer) ||
+			(!EnumHasAnyFlags(OptionalSchemas, EClothCollectionExtendedSchemas::RenderDeformer) ||
 				(RenderDeformerPositionBaryCoordsAndDist &&
 				RenderDeformerNormalBaryCoordsAndDist &&
 				RenderDeformerTangentBaryCoordsAndDist &&
 				RenderDeformerSimIndices3D &&
 				RenderDeformerWeight &&
 				RenderDeformerSkinningBlend)) &&
-			(!EnumHasAnyFlags(OptionalSchemas, EClothCollectionOptionalSchemas::Resizing) ||
+			(!EnumHasAnyFlags(OptionalSchemas, EClothCollectionExtendedSchemas::Resizing) ||
 				RenderCustomResizingBlend) &&
 				
 			// Group Resizing Group
-			(!EnumHasAnyFlags(OptionalSchemas, EClothCollectionOptionalSchemas::Resizing) ||
+			(!EnumHasAnyFlags(OptionalSchemas, EClothCollectionExtendedSchemas::Resizing) ||
 				(CustomResizingRegionSet &&
 					CustomResizingRegionType));
 	}
 
-	void FClothCollection::DefineSchema(EClothCollectionOptionalSchemas OptionalSchemas)
+	void FClothCollection::DefineSchema(EClothCollectionExtendedSchemas OptionalSchemas)
 	{
 		using namespace UE::Chaos::ClothAsset::Private;
 
@@ -429,7 +429,7 @@ namespace UE::Chaos::ClothAsset
 		FManagedArrayCollection::FConstructionParameters SimVertices3DDependency(ClothCollectionGroup::SimVertices3D, bSaved, bAllowCircularDependency);  // Any attribute with this dependency must handle welding and splitting in FCollectionClothSeamFacade
 		FManagedArrayCollection::FConstructionParameters SimMorphTargetVerticesDependency(ClothCollectionGroup::SimMorphTargetVertices, bSaved, bAllowCircularDependency);
 
-		if (EnumHasAnyFlags(OptionalSchemas, EClothCollectionOptionalSchemas::CookedOnly))
+		if (EnumHasAnyFlags(OptionalSchemas, EClothCollectionExtendedSchemas::CookedOnly))
 		{
 			// Special case: only define cooked schema.
 
@@ -451,7 +451,7 @@ namespace UE::Chaos::ClothAsset
 		SkeletalMeshPathName = &ManagedArrayCollection->AddAttribute<FString>(ClothCollectionAttribute::SkeletalMeshPathName, ClothCollectionGroup::Lods);
 
 		// Solvers Group
-		if (EnumHasAnyFlags(OptionalSchemas, EClothCollectionOptionalSchemas::Solvers))
+		if (EnumHasAnyFlags(OptionalSchemas, EClothCollectionExtendedSchemas::Solvers))
 		{
 			SolverGravity = &ManagedArrayCollection->AddAttribute<FVector3f>(ClothCollectionAttribute::SolverGravity, ClothCollectionGroup::Solvers);
 			SolverAirDamping = &ManagedArrayCollection->AddAttribute<float>(ClothCollectionAttribute::SolverAirDamping, ClothCollectionGroup::Solvers);
@@ -482,7 +482,7 @@ namespace UE::Chaos::ClothAsset
 		RenderMaterialPathName = &ManagedArrayCollection->AddAttribute<FString>(ClothCollectionAttribute::RenderMaterialPathName, ClothCollectionGroup::RenderPatterns);
 
 		//~ Fabric Group
-		if (EnumHasAnyFlags(OptionalSchemas, EClothCollectionOptionalSchemas::Fabrics))
+		if (EnumHasAnyFlags(OptionalSchemas, EClothCollectionExtendedSchemas::Fabrics))
 		{
 			FabricBendingStiffness = &ManagedArrayCollection->AddAttribute<FVector3f>(ClothCollectionAttribute::FabricBendingStiffness, ClothCollectionGroup::Fabrics);
 			FabricBucklingStiffness = &ManagedArrayCollection->AddAttribute<FVector3f>(ClothCollectionAttribute::FabricBucklingStiffness, ClothCollectionGroup::Fabrics);
@@ -514,7 +514,7 @@ namespace UE::Chaos::ClothAsset
 		SimVertex2DLookup = &ManagedArrayCollection->AddAttribute<TArray<int32>>(ClothCollectionAttribute::SimVertex2DLookup, ClothCollectionGroup::SimVertices3D, SimVertices2DDependency);
 		SeamStitchLookup = &ManagedArrayCollection->AddAttribute<TArray<int32>>(ClothCollectionAttribute::SeamStitchLookup, ClothCollectionGroup::SimVertices3D, SeamStitchesDependency);
 
-		if (EnumHasAnyFlags(OptionalSchemas, EClothCollectionOptionalSchemas::SimMorphTargets))
+		if (EnumHasAnyFlags(OptionalSchemas, EClothCollectionExtendedSchemas::SimMorphTargets))
 		{
 			// Sim Morph Targets Group
 			SimMorphTargetName = &ManagedArrayCollection->AddAttribute<FString>(ClothCollectionAttribute::SimMorphTargetName, ClothCollectionGroup::SimMorphTargets);
@@ -540,7 +540,7 @@ namespace UE::Chaos::ClothAsset
 		RenderBoneIndices = &ManagedArrayCollection->AddAttribute<TArray<int32>>(ClothCollectionAttribute::RenderBoneIndices, ClothCollectionGroup::RenderVertices);
 		RenderBoneWeights = &ManagedArrayCollection->AddAttribute<TArray<float>>(ClothCollectionAttribute::RenderBoneWeights, ClothCollectionGroup::RenderVertices);
 
-		if (EnumHasAnyFlags(OptionalSchemas, EClothCollectionOptionalSchemas::RenderDeformer))
+		if (EnumHasAnyFlags(OptionalSchemas, EClothCollectionExtendedSchemas::RenderDeformer))
 		{
 			// Render Patterns Group
 			RenderDeformerNumInfluences = &ManagedArrayCollection->AddAttribute<int32>(ClothCollectionAttribute::RenderDeformerNumInfluences, ClothCollectionGroup::RenderPatterns);
@@ -554,7 +554,7 @@ namespace UE::Chaos::ClothAsset
 			RenderDeformerSkinningBlend = &ManagedArrayCollection->AddAttribute<float>(ClothCollectionAttribute::RenderDeformerSkinningBlend, ClothCollectionGroup::RenderVertices);
 		}
 		
-		if (EnumHasAnyFlags(OptionalSchemas, EClothCollectionOptionalSchemas::Resizing))
+		if (EnumHasAnyFlags(OptionalSchemas, EClothCollectionExtendedSchemas::Resizing))
 		{
 			// Sim Vertices 3D Group
 			SimCustomResizingBlend = &ManagedArrayCollection->AddAttribute<float>(ClothCollectionAttribute::SimCustomResizingBlend, ClothCollectionGroup::SimVertices3D);
diff --git a/Engine/Plugins/ChaosClothAsset/Source/ChaosClothAsset/Private/ChaosClothAsset/ClothCollection.h b/Engine/Plugins/ChaosClothAsset/Source/ChaosClothAsset/Private/ChaosClothAsset/ClothCollection.h
--- a/Engine/Plugins/ChaosClothAsset/Source/ChaosClothAsset/Private/ChaosClothAsset/ClothCollection.h
+++ b/Engine/Plugins/ChaosClothAsset/Source/ChaosClothAsset/Private/ChaosClothAsset/ClothCollection.h
@@ -3,7 +3,10 @@
 #pragma once
 
 #include "GeometryCollection/ManagedArray.h"
+#if UE_ENABLE_INCLUDE_ORDER_DEPRECATED_IN_5_6
 #include "ChaosClothAsset/ClothCollectionOptionalSchemas.h"
+#endif
+#include "ChaosClothAsset/ClothCollectionExtendedSchemas.h"
 #include "ChaosClothAsset/IsUserAttributeType.h"
 
 struct FManagedArrayCollection;
@@ -28,10 +31,10 @@ namespace UE::Chaos::ClothAsset
 		FClothCollection& operator=(FClothCollection&&) = delete;
 
 		/** Return whether the underlying collection is a valid cloth collection. */
-		bool IsValid(EClothCollectionOptionalSchemas OptionalSchemas = EClothCollectionOptionalSchemas::None) const;
+		bool IsValid(EClothCollectionExtendedSchemas OptionalSchemas = EClothCollectionExtendedSchemas::None) const;
 
 		/** Make the underlying collection a cloth collection. */
-		void DefineSchema(EClothCollectionOptionalSchemas OptionalSchemas = EClothCollectionOptionalSchemas::None);
+		void DefineSchema(EClothCollectionExtendedSchemas OptionalSchemas = EClothCollectionExtendedSchemas::None);
 
 		/** Get the number of elements of a group. */
 		int32 GetNumElements(const FName& GroupName) const;
diff --git a/Engine/Plugins/ChaosClothAsset/Source/ChaosClothAsset/Private/ChaosClothAsset/ClothGeometryTools.cpp b/Engine/Plugins/ChaosClothAsset/Source/ChaosClothAsset/Private/ChaosClothAsset/ClothGeometryTools.cpp
--- a/Engine/Plugins/ChaosClothAsset/Source/ChaosClothAsset/Private/ChaosClothAsset/ClothGeometryTools.cpp
+++ b/Engine/Plugins/ChaosClothAsset/Source/ChaosClothAsset/Private/ChaosClothAsset/ClothGeometryTools.cpp
@@ -474,7 +474,7 @@ namespace UE::Chaos::ClothAsset
 	{
 		static void ApplyGroupResizing(FCollectionClothFacade& ClothFacade, const FMeshDescription& TargetMeshDescription, const FMeshResizingRBFInterpolationData& InterpolationData, const TArray<FMeshResizingCustomRegion>& ResizingGroupData, TArrayView<FVector3f> Positions, TConstArrayView<float> ResizingBlend)
 		{
-			if (ClothFacade.IsValid(EClothCollectionOptionalSchemas::Resizing))
+			if (ClothFacade.IsValid(EClothCollectionExtendedSchemas::Resizing))
 			{
 				TConstArrayView<int32> ClothSetTypes = ClothFacade.GetCustomResizingRegionType();
 				check(ClothSetTypes.Num() == ResizingGroupData.Num());
@@ -1757,7 +1757,7 @@ PRAGMA_ENABLE_DEPRECATION_WARNINGS
 	void FClothGeometryTools::ApplyProxyDeformer(const TSharedRef<FManagedArrayCollection>& ClothCollection, bool bIgnoreSkinningBlend)
 	{
 		FCollectionClothFacade ClothFacade(ClothCollection);
-		if (ClothFacade.IsValid(EClothCollectionOptionalSchemas::RenderDeformer))
+		if (ClothFacade.IsValid(EClothCollectionExtendedSchemas::RenderDeformer))
 		{
 			return ApplyProxyDeformer(ClothFacade, bIgnoreSkinningBlend, ClothFacade.GetSimPosition3D(), ClothFacade.GetSimNormal());
 		}
@@ -1765,7 +1765,7 @@ PRAGMA_ENABLE_DEPRECATION_WARNINGS
 
 	void FClothGeometryTools::ApplyProxyDeformer(FCollectionClothFacade& ClothFacade, bool bIgnoreSkinningBlend, const TConstArrayView<FVector3f>& SimPositions, const TConstArrayView<FVector3f>& SimNormals)
 	{
-		if (ClothFacade.IsValid(EClothCollectionOptionalSchemas::RenderDeformer))
+		if (ClothFacade.IsValid(EClothCollectionExtendedSchemas::RenderDeformer))
 		{
 			// This follows GpuSkinVertexFactor.ush
 			for (int32 SectionIndex = 0; SectionIndex < ClothFacade.GetNumRenderPatterns(); ++SectionIndex)
@@ -1885,7 +1885,7 @@ PRAGMA_ENABLE_DEPRECATION_WARNINGS
 		check(SetNames.Num() == SetTypes.Num());
 		FCollectionClothFacade ClothFacade(ClothCollection);
 
-		ClothFacade.DefineSchema(EClothCollectionOptionalSchemas::Resizing);
+		ClothFacade.DefineSchema(EClothCollectionExtendedSchemas::Resizing);
 
 		// Clear the sim and render blends
 		TArrayView<float> SimBlend = ClothFacade.GetSimCustomResizingBlend();
@@ -1950,7 +1950,7 @@ PRAGMA_ENABLE_DEPRECATION_WARNINGS
 	void FClothGeometryTools::GenerateSimMeshResizingGroupData(const TSharedRef<FManagedArrayCollection>& ClothCollection, const FMeshDescription& SourceMeshDescription, TArray<FMeshResizingCustomRegion>& OutResizingGroupData)
 	{
 		FCollectionClothFacade ClothFacade(ClothCollection);
-		if (ClothFacade.IsValid(EClothCollectionOptionalSchemas::Resizing))
+		if (ClothFacade.IsValid(EClothCollectionExtendedSchemas::Resizing))
 		{
 			UE::Geometry::FDynamicMesh3 SourceMesh;
 			FMeshDescriptionToDynamicMesh Converter;
@@ -1973,7 +1973,7 @@ PRAGMA_ENABLE_DEPRECATION_WARNINGS
 	void FClothGeometryTools::GenerateRenderMeshResizingGroupData(const TSharedRef<FManagedArrayCollection>& ClothCollection, const FMeshDescription& SourceMeshDescription, TArray<FMeshResizingCustomRegion>& OutResizingGroupData)
 	{
 		FCollectionClothFacade ClothFacade(ClothCollection);
-		if (ClothFacade.IsValid(EClothCollectionOptionalSchemas::Resizing))
+		if (ClothFacade.IsValid(EClothCollectionExtendedSchemas::Resizing))
 		{
 			UE::Geometry::FDynamicMesh3 SourceMesh;
 			FMeshDescriptionToDynamicMesh Converter;
diff --git a/Engine/Plugins/ChaosClothAsset/Source/ChaosClothAsset/Private/ChaosClothAsset/CollectionClothFabricFacade.cpp b/Engine/Plugins/ChaosClothAsset/Source/ChaosClothAsset/Private/ChaosClothAsset/CollectionClothFabricFacade.cpp
--- a/Engine/Plugins/ChaosClothAsset/Source/ChaosClothAsset/Private/ChaosClothAsset/CollectionClothFabricFacade.cpp
+++ b/Engine/Plugins/ChaosClothAsset/Source/ChaosClothAsset/Private/ChaosClothAsset/CollectionClothFabricFacade.cpp
@@ -94,9 +94,9 @@ namespace UE::Chaos::ClothAsset
 		check((ElementIndex >= 0) && (ElementIndex < ClothCollection->GetNumElements(ClothCollectionGroup::Fabrics)));
 		
 		const TSharedRef<class FClothCollection> ClothCol = GetClothCollection();
-		if (!ClothCol->IsValid(EClothCollectionOptionalSchemas::Fabrics))
+		if (!ClothCol->IsValid(EClothCollectionExtendedSchemas::Fabrics))
 		{
-			ClothCol->DefineSchema(EClothCollectionOptionalSchemas::Fabrics);
+			ClothCol->DefineSchema(EClothCollectionExtendedSchemas::Fabrics);
 		}
 		
 		ClothCol->GetElements(ClothCol->GetFabricBendingStiffness())[ElementIndex] = BendingStiffness.GetVectorDatas();
diff --git a/Engine/Plugins/ChaosClothAsset/Source/ChaosClothAsset/Private/ChaosClothAsset/CollectionClothFacade.cpp b/Engine/Plugins/ChaosClothAsset/Source/ChaosClothAsset/Private/ChaosClothAsset/CollectionClothFacade.cpp
--- a/Engine/Plugins/ChaosClothAsset/Source/ChaosClothAsset/Private/ChaosClothAsset/CollectionClothFacade.cpp
+++ b/Engine/Plugins/ChaosClothAsset/Source/ChaosClothAsset/Private/ChaosClothAsset/CollectionClothFacade.cpp
@@ -19,9 +19,9 @@ namespace UE::Chaos::ClothAsset
 	{
 	}
 
-	bool FCollectionClothConstFacade::IsValid(EClothCollectionOptionalSchemas OptionalSchemas) const
+	bool FCollectionClothConstFacade::IsValid(EClothCollectionExtendedSchemas OptionalSchemas) const
 	{
-		return ClothCollection->IsValid(OptionalSchemas) && (EnumHasAnyFlags(OptionalSchemas, EClothCollectionOptionalSchemas::CookedOnly) || ClothCollection->GetNumElements(ClothCollectionGroup::Lods) == 1);
+		return ClothCollection->IsValid(OptionalSchemas) && (EnumHasAnyFlags(OptionalSchemas, EClothCollectionExtendedSchemas::CookedOnly) || ClothCollection->GetNumElements(ClothCollectionGroup::Lods) == 1);
 	}
 
 	bool FCollectionClothConstFacade::HasValidSimulationData() const
@@ -638,7 +638,7 @@ namespace UE::Chaos::ClothAsset
 	{
 	}
 
-	void FCollectionClothFacade::DefineSchema(EClothCollectionOptionalSchemas OptionalSchemas)
+	void FCollectionClothFacade::DefineSchema(EClothCollectionExtendedSchemas OptionalSchemas)
 	{
 		GetClothCollection()->DefineSchema(OptionalSchemas);
 		SetDefaults();
@@ -667,7 +667,7 @@ namespace UE::Chaos::ClothAsset
 		Reset();
 		
 		// Solvers Group
-		if (Other.IsValid(EClothCollectionOptionalSchemas::Solvers) && Other.HasSolverElement())
+		if (Other.IsValid(EClothCollectionExtendedSchemas::Solvers) && Other.HasSolverElement())
 		{
 			SetSolverGravity(Other.GetSolverGravity());
 			SetSolverAirDamping(Other.GetSolverAirDamping());
@@ -699,10 +699,10 @@ namespace UE::Chaos::ClothAsset
 			SetSkeletalMeshPathName(Other.GetSkeletalMeshPathName());
 		}
 
-		const bool bOtherHasResizingData = Other.IsValid(EClothCollectionOptionalSchemas::Resizing);
+		const bool bOtherHasResizingData = Other.IsValid(EClothCollectionExtendedSchemas::Resizing);
 		if (bOtherHasResizingData)
 		{
-			DefineSchema(EClothCollectionOptionalSchemas::Resizing);
+			DefineSchema(EClothCollectionExtendedSchemas::Resizing);
 		}
 
 		// Very important order of operations to ensure indices don't get messed up:
@@ -871,36 +871,36 @@ namespace UE::Chaos::ClothAsset
 	
 	void FCollectionClothFacade::SetSolverGravity(const FVector3f& SolverGravity)
 	{
-		if (!IsValid(EClothCollectionOptionalSchemas::Solvers))
+		if (!IsValid(EClothCollectionExtendedSchemas::Solvers))
 		{
-			DefineSchema(EClothCollectionOptionalSchemas::Solvers);
+			DefineSchema(EClothCollectionExtendedSchemas::Solvers);
 		}
 		(*GetClothCollection()->GetSolverGravity())[0] = SolverGravity;
 	}
 	
 	void FCollectionClothFacade::SetSolverAirDamping(const float SolverAirDamping)
 	{
-		if (!IsValid(EClothCollectionOptionalSchemas::Solvers))
+		if (!IsValid(EClothCollectionExtendedSchemas::Solvers))
 		{
-			DefineSchema(EClothCollectionOptionalSchemas::Solvers);
+			DefineSchema(EClothCollectionExtendedSchemas::Solvers);
 		}
 		(*GetClothCollection()->GetSolverAirDamping())[0] = SolverAirDamping;
 	}
 	
 	void FCollectionClothFacade::SetSolverTimeStep(const float SolverTimeStep)
 	{
-		if (!IsValid(EClothCollectionOptionalSchemas::Solvers))
+		if (!IsValid(EClothCollectionExtendedSchemas::Solvers))
 		{
-			DefineSchema(EClothCollectionOptionalSchemas::Solvers);
+			DefineSchema(EClothCollectionExtendedSchemas::Solvers);
 		}
 		(*GetClothCollection()->GetSolverTimeStep())[0] = SolverTimeStep;
 	}
 	
 	void FCollectionClothFacade::SetSolverSubSteps(const int32 SolverSubSteps)
 	{
-		if (!IsValid(EClothCollectionOptionalSchemas::Solvers))
+		if (!IsValid(EClothCollectionExtendedSchemas::Solvers))
 		{
-			DefineSchema(EClothCollectionOptionalSchemas::Solvers);
+			DefineSchema(EClothCollectionExtendedSchemas::Solvers);
 		}
 		(*GetClothCollection()->GetSolverSubSteps())[0] = SolverSubSteps;
 	}
@@ -1057,7 +1057,7 @@ namespace UE::Chaos::ClothAsset
 
 	void FCollectionClothFacade::SetNumSimMorphTargets(int32 InNumMorphTargets, bool bCookedOnly)
 	{
-		const EClothCollectionOptionalSchemas OptionalSchemas = bCookedOnly ? EClothCollectionOptionalSchemas::CookedOnly : EClothCollectionOptionalSchemas::SimMorphTargets;
+		const EClothCollectionExtendedSchemas OptionalSchemas = bCookedOnly ? EClothCollectionExtendedSchemas::CookedOnly : EClothCollectionExtendedSchemas::SimMorphTargets;
 		if (!IsValid(OptionalSchemas))
 		{
 			DefineSchema(OptionalSchemas);
@@ -1342,9 +1342,9 @@ namespace UE::Chaos::ClothAsset
 
 	void FCollectionClothFacade::SetNumCustomResizingRegions(int32 InNumGroups)
 	{
-		if (!IsValid(EClothCollectionOptionalSchemas::Resizing))
+		if (!IsValid(EClothCollectionExtendedSchemas::Resizing))
 		{
-			DefineSchema(EClothCollectionOptionalSchemas::Resizing);
+			DefineSchema(EClothCollectionExtendedSchemas::Resizing);
 		}
 
 		const int32 NumGroups = GetNumCustomResizingRegions();
@@ -1415,7 +1415,7 @@ namespace UE::Chaos::ClothAsset
 	{
 		GetClothCollection()->SetNumElements(1, ClothCollectionGroup::Lods);
 
-		if (IsValid(EClothCollectionOptionalSchemas::Solvers))
+		if (IsValid(EClothCollectionExtendedSchemas::Solvers))
 		{
 			GetClothCollection()->SetNumElements(1, ClothCollectionGroup::Solvers);
 			
diff --git a/Engine/Plugins/ChaosClothAsset/Source/ChaosClothAsset/Private/ChaosClothAsset/CollectionClothRenderPatternFacade.cpp b/Engine/Plugins/ChaosClothAsset/Source/ChaosClothAsset/Private/ChaosClothAsset/CollectionClothRenderPatternFacade.cpp
--- a/Engine/Plugins/ChaosClothAsset/Source/ChaosClothAsset/Private/ChaosClothAsset/CollectionClothRenderPatternFacade.cpp
+++ b/Engine/Plugins/ChaosClothAsset/Source/ChaosClothAsset/Private/ChaosClothAsset/CollectionClothRenderPatternFacade.cpp
@@ -235,7 +235,7 @@ namespace UE::Chaos::ClothAsset
 		FClothCollection::CopyArrayViewData(GetRenderBoneIndices(), Other.GetRenderBoneIndices());
 		FClothCollection::CopyArrayViewData(GetRenderBoneWeights(), Other.GetRenderBoneWeights());
 
-		if (GetClothCollection()->IsValid(EClothCollectionOptionalSchemas::RenderDeformer) && Other.ClothCollection->IsValid(EClothCollectionOptionalSchemas::RenderDeformer))
+		if (GetClothCollection()->IsValid(EClothCollectionExtendedSchemas::RenderDeformer) && Other.ClothCollection->IsValid(EClothCollectionExtendedSchemas::RenderDeformer))
 		{
 			FClothCollection::CopyArrayViewData(GetRenderDeformerPositionBaryCoordsAndDist(), Other.GetRenderDeformerPositionBaryCoordsAndDist());
 			FClothCollection::CopyArrayViewData(GetRenderDeformerNormalBaryCoordsAndDist(), Other.GetRenderDeformerNormalBaryCoordsAndDist());
@@ -245,7 +245,7 @@ namespace UE::Chaos::ClothAsset
 			FClothCollection::CopyArrayViewData(GetRenderDeformerSkinningBlend(), Other.GetRenderDeformerSkinningBlend());
 		}
 
-		if (GetClothCollection()->IsValid(EClothCollectionOptionalSchemas::Resizing) && Other.ClothCollection->IsValid(EClothCollectionOptionalSchemas::Resizing))
+		if (GetClothCollection()->IsValid(EClothCollectionExtendedSchemas::Resizing) && Other.ClothCollection->IsValid(EClothCollectionExtendedSchemas::Resizing))
 		{
 			FClothCollection::CopyArrayViewData(GetRenderCustomResizingBlend(), Other.GetRenderCustomResizingBlend());
 		}
@@ -259,13 +259,13 @@ namespace UE::Chaos::ClothAsset
 	void FCollectionClothRenderPatternFacade::SetRenderDeformerNumInfluences(int32 NumInfluences)
 	{
 		check(NumInfluences >= 0);
-		if (!GetClothCollection()->IsValid(EClothCollectionOptionalSchemas::RenderDeformer))
+		if (!GetClothCollection()->IsValid(EClothCollectionExtendedSchemas::RenderDeformer))
 		{
 			if (NumInfluences == 0)
 			{
 				return;  // Not having the schema is the same as having zero influences
 			}
-			GetClothCollection()->DefineSchema(EClothCollectionOptionalSchemas::RenderDeformer);
+			GetClothCollection()->DefineSchema(EClothCollectionExtendedSchemas::RenderDeformer);
 		}
 		(*GetClothCollection()->GetRenderDeformerNumInfluences())[GetElementIndex()] = NumInfluences;
 	}
diff --git a/Engine/Plugins/ChaosClothAsset/Source/ChaosClothAsset/Private/ChaosClothAsset/CollectionClothSimMorphTargetFacade.cpp b/Engine/Plugins/ChaosClothAsset/Source/ChaosClothAsset/Private/ChaosClothAsset/CollectionClothSimMorphTargetFacade.cpp
--- a/Engine/Plugins/ChaosClothAsset/Source/ChaosClothAsset/Private/ChaosClothAsset/CollectionClothSimMorphTargetFacade.cpp
+++ b/Engine/Plugins/ChaosClothAsset/Source/ChaosClothAsset/Private/ChaosClothAsset/CollectionClothSimMorphTargetFacade.cpp
@@ -21,7 +21,7 @@ namespace UE::Chaos::ClothAsset
 			GetElementIndex());
 	}
 
-	int32 FCollectionClothSimMorphTargetConstFacade::GetimMorphTargetVerticesOffset() const
+	int32 FCollectionClothSimMorphTargetConstFacade::GetSimMorphTargetVerticesOffset() const
 	{
 		return ClothCollection->GetElementsOffset(
 			ClothCollection->GetSimMorphTargetVerticesStart(),
@@ -65,7 +65,7 @@ namespace UE::Chaos::ClothAsset
 		: ClothCollection(ClothCollection)
 		, MorphTargetIndex(InMorphTargetIndex)
 	{
-		check(ClothCollection->IsValid(EClothCollectionOptionalSchemas::CookedOnly | EClothCollectionOptionalSchemas::SimMorphTargets));
+		check(ClothCollection->IsValid(EClothCollectionExtendedSchemas::CookedOnly | EClothCollectionExtendedSchemas::SimMorphTargets));
 		check(MorphTargetIndex >= 0 && MorphTargetIndex < ClothCollection->GetNumElements(ClothCollectionGroup::SimMorphTargets));
 	}
 
@@ -123,7 +123,7 @@ namespace UE::Chaos::ClothAsset
 	void FCollectionClothSimMorphTargetFacade::RemoveSimMorphTargetVertices(const TArray<int32>& SortedDeletionList)
 	{
 		TArray<int32> GlobalIndexSortedDeletionList;
-		const int32 Offset = GetimMorphTargetVerticesOffset();
+		const int32 Offset = GetSimMorphTargetVerticesOffset();
 		GlobalIndexSortedDeletionList.SetNumUninitialized(SortedDeletionList.Num());
 		for (int32 Idx = 0; Idx < SortedDeletionList.Num(); ++Idx)
 		{
diff --git a/Engine/Plugins/ChaosClothAsset/Source/ChaosClothAsset/Public/ChaosClothAsset/ClothCollectionExtendedSchemas.h b/Engine/Plugins/ChaosClothAsset/Source/ChaosClothAsset/Public/ChaosClothAsset/ClothCollectionExtendedSchemas.h
new file mode 100644
--- /dev/null
+++ b/Engine/Plugins/ChaosClothAsset/Source/ChaosClothAsset/Public/ChaosClothAsset/ClothCollectionExtendedSchemas.h
@@ -0,0 +1,22 @@
+// Copyright Epic Games, Inc. All Rights Reserved.
+
+#pragma once
+
+#include "HAL/Platform.h"
+#include "Misc/EnumClassFlags.h"
+
+namespace UE::Chaos::ClothAsset
+{
+	/** Cloth collection extended schemas. */
+	enum class EClothCollectionExtendedSchemas : uint8
+	{
+		None = 0,
+		RenderDeformer = 1 << 1,
+		Solvers = 1 << 2,
+		Fabrics = 1 << 3,
+		SimMorphTargets = 1 << 4,
+		CookedOnly = 1 << 5, // If this flag is found, only cooked schemas will be considered.
+		Resizing = 1 << 6,
+	};
+	ENUM_CLASS_FLAGS(EClothCollectionExtendedSchemas)
+}  // End namespace UE::Chaos::ClothAsset
diff --git a/Engine/Plugins/ChaosClothAsset/Source/ChaosClothAsset/Public/ChaosClothAsset/ClothCollectionOptionalSchemas.h b/Engine/Plugins/ChaosClothAsset/Source/ChaosClothAsset/Public/ChaosClothAsset/ClothCollectionOptionalSchemas.h
--- a/Engine/Plugins/ChaosClothAsset/Source/ChaosClothAsset/Public/ChaosClothAsset/ClothCollectionOptionalSchemas.h
+++ b/Engine/Plugins/ChaosClothAsset/Source/ChaosClothAsset/Public/ChaosClothAsset/ClothCollectionOptionalSchemas.h
@@ -1,22 +1,13 @@
 // Copyright Epic Games, Inc. All Rights Reserved.
 
 #pragma once
+#include "Misc/CoreMiscDefines.h"
 
-#include "HAL/Platform.h"
-#include "Misc/EnumClassFlags.h"
+UE_DEPRECATED_HEADER(5.6, "Use ClothCollectionExtendedSchemas.h instead of ClothCollectionOptionalSchemas.h")
+
+#include "ChaosClothAsset/ClothCollectionExtendedSchemas.h"
 
 namespace UE::Chaos::ClothAsset
 {
-	/** Cloth collection optional schemas. */
-	enum class EClothCollectionOptionalSchemas : uint8
-	{
-		None = 0,
-		RenderDeformer = 1 << 1,
-		Solvers = 1 << 2,
-		Fabrics = 1 << 3,
-		SimMorphTargets = 1 << 4,
-		CookedOnly = 1 << 5, // If this flag is found, only cooked schemas will be considered.
-		Resizing = 1 << 6,
-	};
-	ENUM_CLASS_FLAGS(EClothCollectionOptionalSchemas)
+	UE_DEPRECATED(5.6, "Use EClothCollectionExtendedSchemas instead") typedef enum EClothCollectionExtendedSchemas EClothCollectionOptionalSchemas;
 }  // End namespace UE::Chaos::ClothAsset
diff --git a/Engine/Plugins/ChaosClothAsset/Source/ChaosClothAsset/Public/ChaosClothAsset/CollectionClothFacade.h b/Engine/Plugins/ChaosClothAsset/Source/ChaosClothAsset/Public/ChaosClothAsset/CollectionClothFacade.h
--- a/Engine/Plugins/ChaosClothAsset/Source/ChaosClothAsset/Public/ChaosClothAsset/CollectionClothFacade.h
+++ b/Engine/Plugins/ChaosClothAsset/Source/ChaosClothAsset/Public/ChaosClothAsset/CollectionClothFacade.h
@@ -3,7 +3,10 @@
 #pragma once
 
 #include "Math/Vector.h"
+#if UE_ENABLE_INCLUDE_ORDER_DEPRECATED_IN_5_6
 #include "ChaosClothAsset/ClothCollectionOptionalSchemas.h"
+#endif
+#include "ChaosClothAsset/ClothCollectionExtendedSchemas.h"
 #include "ChaosClothAsset/CollectionClothFabricFacade.h"
 #include "ChaosClothAsset/CollectionClothRenderPatternFacade.h"
 #include "ChaosClothAsset/CollectionClothSeamFacade.h"
@@ -48,7 +51,7 @@ namespace UE::Chaos::ClothAsset
 		virtual ~FCollectionClothConstFacade() = default;
 
 		/** Return whether the facade is defined on the collection. */
-		UE_API bool IsValid(EClothCollectionOptionalSchemas OptionalSchemas = EClothCollectionOptionalSchemas::None) const;
+		UE_API bool IsValid(EClothCollectionExtendedSchemas OptionalSchemas = EClothCollectionExtendedSchemas::None) const;
 
 		/** Return whether the facade has a non-empty simulation mesh data. */
 		UE_API bool HasValidSimulationData() const;
@@ -244,7 +247,7 @@ namespace UE::Chaos::ClothAsset
 		virtual ~FCollectionClothFacade() override = default;
 
 		/** Create this facade's groups and attributes. */
-		UE_API void DefineSchema(EClothCollectionOptionalSchemas OptionalSchemas = EClothCollectionOptionalSchemas::None);
+		UE_API void DefineSchema(EClothCollectionExtendedSchemas OptionalSchemas = EClothCollectionExtendedSchemas::None);
 
 		/** Remove all LODs from this cloth. */
 		UE_API void Reset();
diff --git a/Engine/Plugins/ChaosClothAsset/Source/ChaosClothAsset/Public/ChaosClothAsset/CollectionClothSimMorphTargetFacade.h b/Engine/Plugins/ChaosClothAsset/Source/ChaosClothAsset/Public/ChaosClothAsset/CollectionClothSimMorphTargetFacade.h
--- a/Engine/Plugins/ChaosClothAsset/Source/ChaosClothAsset/Public/ChaosClothAsset/CollectionClothSimMorphTargetFacade.h
+++ b/Engine/Plugins/ChaosClothAsset/Source/ChaosClothAsset/Public/ChaosClothAsset/CollectionClothSimMorphTargetFacade.h
@@ -35,7 +35,7 @@ namespace UE::Chaos::ClothAsset
 		/** Return the total number of vertices for this morph target. */
 		UE_API int32 GetNumSimMorphTargetVertices() const;
 		/** Return the morph target vertices offset for this morph target in the morph target vertices for the collection. */
-		UE_API int32 GetimMorphTargetVerticesOffset() const;
+		UE_API int32 GetSimMorphTargetVerticesOffset() const;
 		UE_API TConstArrayView<FVector3f> GetSimMorphTargetPositionDelta() const;
 		UE_API TConstArrayView<FVector3f> GetSimMorphTargetTangentZDelta() const;
 		UE_API TConstArrayView<int32> GetSimMorphTargetSimVertex3DIndex() const;
@@ -43,7 +43,7 @@ namespace UE::Chaos::ClothAsset
 		UE_API bool IsEmpty() const;
 
 		/** Return the Morph Target index this facade has been created with. */
-		int32 GetMorphTargetIndex() const { return MorphTargetIndex; }
+		int32 GetSimMorphTargetIndex() const { return MorphTargetIndex; }
 
 	protected:
 		friend class FCollectionClothSimMorphTargetFacade;  // For other instances access
diff --git a/Engine/Plugins/ChaosClothAsset/Source/ChaosClothAssetEngine/Private/ChaosClothAsset/ClothAsset.cpp b/Engine/Plugins/ChaosClothAsset/Source/ChaosClothAssetEngine/Private/ChaosClothAsset/ClothAsset.cpp
--- a/Engine/Plugins/ChaosClothAsset/Source/ChaosClothAssetEngine/Private/ChaosClothAsset/ClothAsset.cpp
+++ b/Engine/Plugins/ChaosClothAsset/Source/ChaosClothAssetEngine/Private/ChaosClothAsset/ClothAsset.cpp
@@ -166,10 +166,10 @@ namespace UE::Chaos::ClothAsset::Private
 
 				// Morph targets
 				const UE::Chaos::ClothAsset::FCollectionClothConstFacade InClothFacade(ClothCollection);
-				if (InClothFacade.IsValid(UE::Chaos::ClothAsset::EClothCollectionOptionalSchemas::CookedOnly))
+				if (InClothFacade.IsValid(UE::Chaos::ClothAsset::EClothCollectionExtendedSchemas::CookedOnly))
 				{
 					UE::Chaos::ClothAsset::FCollectionClothFacade ClothFacade(PropertyCollection);
-					ClothFacade.DefineSchema(UE::Chaos::ClothAsset::EClothCollectionOptionalSchemas::CookedOnly);
+					ClothFacade.DefineSchema(UE::Chaos::ClothAsset::EClothCollectionExtendedSchemas::CookedOnly);
 					ClothFacade.InitializeCookedOnly(InClothFacade);
 				}
 
diff --git a/Engine/Plugins/ChaosClothAssetEditor/Source/ChaosClothAssetDataflowNodes/Private/ChaosClothAsset/ClothCollectionQueryNode.cpp b/Engine/Plugins/ChaosClothAssetEditor/Source/ChaosClothAssetDataflowNodes/Private/ChaosClothAsset/ClothCollectionQueryNode.cpp
--- a/Engine/Plugins/ChaosClothAssetEditor/Source/ChaosClothAssetDataflowNodes/Private/ChaosClothAsset/ClothCollectionQueryNode.cpp
+++ b/Engine/Plugins/ChaosClothAssetEditor/Source/ChaosClothAssetDataflowNodes/Private/ChaosClothAsset/ClothCollectionQueryNode.cpp
@@ -42,7 +42,7 @@ void FChaosClothAssetCollectionQueryNode::Evaluate(UE::Dataflow::FContext& Conte
 		SetValue(Context, ClothFacade.IsValid(), &bIsClothCollection);
 		SetValue(Context, FClothGeometryTools::HasSimMesh(ClothCollection), &bHasClothSimMesh);
 		SetValue(Context, FClothGeometryTools::HasRenderMesh(ClothCollection), &bHasClothRenderMesh);
-		SetValue(Context, ClothFacade.IsValid(EClothCollectionOptionalSchemas::RenderDeformer), &bHasClothProxyDeformer);
+		SetValue(Context, ClothFacade.IsValid(EClothCollectionExtendedSchemas::RenderDeformer), &bHasClothProxyDeformer);
 
 		Chaos::Softs::FCollectionPropertyConstFacade PropertyFacade(ClothCollection);
 		bool bBooleanPropertyResult = bBooleanPropertyValue;
diff --git a/Engine/Plugins/ChaosClothAssetEditor/Source/ChaosClothAssetDataflowNodes/Private/ChaosClothAsset/ProxyDeformerNode.cpp b/Engine/Plugins/ChaosClothAssetEditor/Source/ChaosClothAssetDataflowNodes/Private/ChaosClothAsset/ProxyDeformerNode.cpp
--- a/Engine/Plugins/ChaosClothAssetEditor/Source/ChaosClothAssetDataflowNodes/Private/ChaosClothAsset/ProxyDeformerNode.cpp
+++ b/Engine/Plugins/ChaosClothAssetEditor/Source/ChaosClothAssetDataflowNodes/Private/ChaosClothAsset/ProxyDeformerNode.cpp
@@ -266,9 +266,9 @@ void FChaosClothAssetProxyDeformerNode_v2::Evaluate(UE::Dataflow::FContext& Cont
 			FCollectionClothSelectionFacade SelectionFacade(ClothCollection);
 
 			// Add the optional render deformer schema
-			if (!ClothFacade.IsValid(EClothCollectionOptionalSchemas::RenderDeformer))
+			if (!ClothFacade.IsValid(EClothCollectionExtendedSchemas::RenderDeformer))
 			{
-				ClothFacade.DefineSchema(EClothCollectionOptionalSchemas::RenderDeformer);
+				ClothFacade.DefineSchema(EClothCollectionExtendedSchemas::RenderDeformer);
 			}
 
 			// Create the render weight map for storing the skinning blend weights
@@ -474,9 +474,9 @@ void FChaosClothAssetProxyDeformerNode::Evaluate(UE::Dataflow::FContext& Context
 			}
 
 			// Add the optional render deformer schema
-			if (!ClothFacade.IsValid(EClothCollectionOptionalSchemas::RenderDeformer))
+			if (!ClothFacade.IsValid(EClothCollectionExtendedSchemas::RenderDeformer))
 			{
-				ClothFacade.DefineSchema(EClothCollectionOptionalSchemas::RenderDeformer);
+				ClothFacade.DefineSchema(EClothCollectionExtendedSchemas::RenderDeformer);
 			}
 
 			// Create the render weight map for storing the skinning blend weights
diff --git a/Engine/Plugins/ChaosClothAssetEditor/Source/ChaosClothAssetDataflowNodes/Private/ChaosClothAsset/RemeshNode.cpp b/Engine/Plugins/ChaosClothAssetEditor/Source/ChaosClothAssetDataflowNodes/Private/ChaosClothAsset/RemeshNode.cpp
--- a/Engine/Plugins/ChaosClothAssetEditor/Source/ChaosClothAssetDataflowNodes/Private/ChaosClothAsset/RemeshNode.cpp
+++ b/Engine/Plugins/ChaosClothAssetEditor/Source/ChaosClothAssetDataflowNodes/Private/ChaosClothAsset/RemeshNode.cpp
@@ -1782,7 +1782,7 @@ void FChaosClothAssetRemeshNode_v2::Evaluate(UE::Dataflow::FContext& Context, co
 		// Copy collection to output
 		TSharedRef<FManagedArrayCollection> OutputClothCollection = MakeShared<FManagedArrayCollection>();
 
-		if ((bRemeshSim || bRemeshRender) && ClothFacade.IsValid(EClothCollectionOptionalSchemas::RenderDeformer))
+		if ((bRemeshSim || bRemeshRender) && ClothFacade.IsValid(EClothCollectionExtendedSchemas::RenderDeformer))
 		{
 			FClothDataflowTools::LogAndToastWarning(*this,
 				LOCTEXT("InputHasDeformerDataHeadline", "Proxy Deformer Data Found"),
@@ -1884,7 +1884,7 @@ void FChaosClothAssetRemeshNode::Evaluate(UE::Dataflow::FContext& Context, const
 		// Copy collection to output
 		TSharedPtr<FManagedArrayCollection> OutputClothCollection = MakeShared<FManagedArrayCollection>();
 
-		if ((bRemeshSim || bRemeshRender) && ClothFacade.IsValid(EClothCollectionOptionalSchemas::RenderDeformer))
+		if ((bRemeshSim || bRemeshRender) && ClothFacade.IsValid(EClothCollectionExtendedSchemas::RenderDeformer))
 		{
 			FClothDataflowTools::LogAndToastWarning(*this, 
 				LOCTEXT("InputHasDeformerDataHeadline", "Proxy Deformer Data Found"),
diff --git a/Engine/Plugins/ChaosClothAssetEditor/Source/ChaosClothAssetDataflowNodes/Private/ChaosClothAsset/SimulationBaseConfigNode.cpp b/Engine/Plugins/ChaosClothAssetEditor/Source/ChaosClothAssetDataflowNodes/Private/ChaosClothAsset/SimulationBaseConfigNode.cpp
--- a/Engine/Plugins/ChaosClothAssetEditor/Source/ChaosClothAssetDataflowNodes/Private/ChaosClothAsset/SimulationBaseConfigNode.cpp
+++ b/Engine/Plugins/ChaosClothAssetEditor/Source/ChaosClothAssetDataflowNodes/Private/ChaosClothAsset/SimulationBaseConfigNode.cpp
@@ -409,7 +409,7 @@ void FChaosClothAssetSimulationBaseConfigNode::FPropertyHelper::SetSolverPropert
 	const TArray<FName>& SimilarPropertyNames, ECollectionPropertyFlags PropertyFlags)
 {
 	UE::Chaos::ClothAsset::FCollectionClothFacade ClothFacade(GetClothCollection());
-	if(PropertyValue.bUseImportedValue && ClothFacade.IsValid(UE::Chaos::ClothAsset::EClothCollectionOptionalSchemas::Solvers) && ClothFacade.HasSolverElement())
+	if(PropertyValue.bUseImportedValue && ClothFacade.IsValid(UE::Chaos::ClothAsset::EClothCollectionExtendedSchemas::Solvers) && ClothFacade.HasSolverElement())
 	{
 		PropertyValue.ImportedValue = SolverValueFunction(ClothFacade);
 	}
@@ -422,7 +422,7 @@ void FChaosClothAssetSimulationBaseConfigNode::FPropertyHelper::SetFabricPropert
 	const TArray<FName>& SimilarPropertyNames, ECollectionPropertyFlags PropertyFlags)
 {
 	UE::Chaos::ClothAsset::FCollectionClothFacade ClothFacade(GetClothCollection());
-	if(PropertyValue.bUseImportedValue && ClothFacade.IsValid(UE::Chaos::ClothAsset::EClothCollectionOptionalSchemas::Fabrics) && ClothFacade.GetNumFabrics() > 0)
+	if(PropertyValue.bUseImportedValue && ClothFacade.IsValid(UE::Chaos::ClothAsset::EClothCollectionExtendedSchemas::Fabrics) && ClothFacade.GetNumFabrics() > 0)
 	{
 		const int32 NumFabrics = ClothFacade.GetNumFabrics();
 		
@@ -445,7 +445,7 @@ void FChaosClothAssetSimulationBaseConfigNode::FPropertyHelper::SetFabricPropert
 	const TArray<FName>& SimilarPropertyNames, ECollectionPropertyFlags PropertyFlags)
 {
 	UE::Chaos::ClothAsset::FCollectionClothFacade ClothFacade(GetClothCollection());
-	if(PropertyValue.bCouldUseFabrics && ClothFacade.IsValid(UE::Chaos::ClothAsset::EClothCollectionOptionalSchemas::Fabrics) && (ClothFacade.GetNumFabrics() > 0) &&
+	if(PropertyValue.bCouldUseFabrics && ClothFacade.IsValid(UE::Chaos::ClothAsset::EClothCollectionExtendedSchemas::Fabrics) && (ClothFacade.GetNumFabrics() > 0) &&
 			(PropertyValue.bImportFabricBounds || PropertyValue.bBuildFabricMaps))
 	{
 		UE::Chaos::ClothAsset::FWeightedValueBounds WeightedValueBounds;
@@ -477,7 +477,7 @@ void FChaosClothAssetSimulationBaseConfigNode::FPropertyHelper::SetFabricPropert
 	const TArray<FName>& SimilarPropertyNames, ECollectionPropertyFlags PropertyFlags, const FName& GroupName)
 {
 	UE::Chaos::ClothAsset::FCollectionClothFacade ClothFacade(GetClothCollection());
-	if(PropertyValue.bCouldUseFabrics && ClothFacade.IsValid(UE::Chaos::ClothAsset::EClothCollectionOptionalSchemas::Fabrics) && (ClothFacade.GetNumFabrics() > 0) && PropertyValue.bBuildFabricMaps)
+	if(PropertyValue.bCouldUseFabrics && ClothFacade.IsValid(UE::Chaos::ClothAsset::EClothCollectionExtendedSchemas::Fabrics) && (ClothFacade.GetNumFabrics() > 0) && PropertyValue.bBuildFabricMaps)
 	{
 		const int32 NumPatterns = ClothFacade.GetNumSimPatterns();
 
@@ -531,7 +531,7 @@ void FChaosClothAssetSimulationBaseConfigNode::FPropertyHelper::SetSolverPropert
 	const TArray<FName>& SimilarPropertyNames, ECollectionPropertyFlags PropertyFlags)
 {
 	UE::Chaos::ClothAsset::FCollectionClothFacade ClothFacade(GetClothCollection());
-	if(PropertyValue.bCouldUseFabrics && (ClothFacade.IsValid(UE::Chaos::ClothAsset::EClothCollectionOptionalSchemas::Solvers) && ClothFacade.HasSolverElement()) && PropertyValue.bImportFabricBounds)
+	if(PropertyValue.bCouldUseFabrics && (ClothFacade.IsValid(UE::Chaos::ClothAsset::EClothCollectionExtendedSchemas::Solvers) && ClothFacade.HasSolverElement()) && PropertyValue.bImportFabricBounds)
 	{
 		const float SolverValue = SolverValueFunction(ClothFacade);
 	
diff --git a/Engine/Plugins/ChaosClothAssetEditor/Source/ChaosClothAssetDataflowNodes/Private/ChaosClothAsset/SkinningBlendNode.cpp b/Engine/Plugins/ChaosClothAssetEditor/Source/ChaosClothAssetDataflowNodes/Private/ChaosClothAsset/SkinningBlendNode.cpp
--- a/Engine/Plugins/ChaosClothAssetEditor/Source/ChaosClothAssetDataflowNodes/Private/ChaosClothAsset/SkinningBlendNode.cpp
+++ b/Engine/Plugins/ChaosClothAssetEditor/Source/ChaosClothAssetDataflowNodes/Private/ChaosClothAsset/SkinningBlendNode.cpp
@@ -131,7 +131,7 @@ void FChaosClothAssetSkinningBlendNode::Evaluate(UE::Dataflow::FContext& Context
 		{
 			FCollectionClothSelectionConstFacade SelectionFacade(ClothCollection);
 			// Check for the optional render deformer schema
-			if (!ClothFacade.IsValid(EClothCollectionOptionalSchemas::RenderDeformer))
+			if (!ClothFacade.IsValid(EClothCollectionExtendedSchemas::RenderDeformer))
 			{
 				FClothDataflowTools::LogAndToastWarning(*this,
 					LOCTEXT("NoProxyDefromerHeadline", "No Proxy Deformer data."),
