Chaos Outfit Asset - Fixed the Protruding property on the body size match computation that is causing protuding garments to be selected over non protuding bodies. Added a new experimental RecalculateNormals node to recalculate the normals if necessary.

#rnx

diff --git a/Engine/Plugins/ChaosClothAsset/Source/ChaosClothAsset/Private/ChaosClothAsset/ClothGeometryTools.cpp b/Engine/Plugins/ChaosClothAsset/Source/ChaosClothAsset/Private/ChaosClothAsset/ClothGeometryTools.cpp
--- a/Engine/Plugins/ChaosClothAsset/Source/ChaosClothAsset/Private/ChaosClothAsset/ClothGeometryTools.cpp
+++ b/Engine/Plugins/ChaosClothAsset/Source/ChaosClothAsset/Private/ChaosClothAsset/ClothGeometryTools.cpp
@@ -644,7 +644,7 @@ namespace UE::Chaos::ClothAsset
 			const FVector2f UV02 = SimPosition2D[Face[2]] - SimPosition2D[Face[0]];
 
 			const float Denom = UV01.X * UV02.Y - UV01.Y * UV02.X;
-			const float InvDenom = (Denom < SMALL_NUMBER) ? 0.f : 1.f / Denom;
+			const float InvDenom = (FMath::Abs(Denom) < SMALL_NUMBER) ? 0.f : 1.f / Denom;
 			const FVector3f TangentU = (Pos01 * UV02.Y - Pos02 * UV01.Y) * InvDenom;
 			const FVector3f TangentV = (Pos02 * UV01.X - Pos01 * UV02.X) * InvDenom;
 
@@ -804,6 +804,72 @@ namespace UE::Chaos::ClothAsset
 		}
 	}
 
+	void FClothGeometryTools::RecalculateRenderMeshNormals(const TSharedRef<FManagedArrayCollection>& ClothCollection)
+	{
+		FCollectionClothFacade ClothFacade(ClothCollection);
+
+		const TConstArrayView<FIntVector3> FaceIndices = ClothFacade.GetRenderIndices();
+		const TConstArrayView<FVector3f> Positions = ClothFacade.GetRenderPosition();
+		const TConstArrayView<TArray<FVector2f>> UVs = ClothFacade.GetRenderUVs();
+		const TArrayView<FVector3f> Normals = ClothFacade.GetRenderNormal();
+		const TArrayView<FVector3f> TangentUs = ClothFacade.GetRenderTangentU();
+		const TArrayView<FVector3f> TangentVs = ClothFacade.GetRenderTangentV();
+
+		check(Positions.Num() == UVs.Num());
+		check(Positions.Num() == Normals.Num());
+		check(Positions.Num() == TangentUs.Num());
+		check(Positions.Num() == TangentVs.Num());
+
+		FMemory::Memzero(Normals.GetData(), Normals.NumBytes());
+		FMemory::Memzero(TangentUs.GetData(), TangentUs.NumBytes());
+
+		for (int32 FaceIndex = 0; FaceIndex < FaceIndices.Num(); ++FaceIndex)
+		{
+			const FIntVector3& Face = FaceIndices[FaceIndex];
+
+			const FVector3f& Pos0 = Positions[Face[0]];
+			const FVector3f& Pos1 = Positions[Face[1]];
+			const FVector3f& Pos2 = Positions[Face[2]];
+			const FVector3f Pos01 = Pos1 - Pos0;
+			const FVector3f Pos02 = Pos2 - Pos0;
+			const FVector3f Normal = FVector3f::CrossProduct(Pos02, Pos01);
+
+			FVector3f TangentU;
+			if (UVs[Face[0]].Num() && UVs[Face[1]].Num() && UVs[Face[2]].Num())
+			{
+				const FVector2f UV01 = UVs[Face[1]][0] - UVs[Face[0]][0];
+				const FVector2f UV02 = UVs[Face[2]][0] - UVs[Face[0]][0];
+				const float Denom = UV01.X * UV02.Y - UV01.Y * UV02.X;
+				const float InvDenom = (FMath::Abs(Denom) < SMALL_NUMBER) ? 0.f : 1.f / Denom;
+				TangentU = (Pos01 * UV02.Y - Pos02 * UV01.Y) * InvDenom;
+			}
+			else
+			{
+				TangentU = (Pos1 + Pos2) * 0.5f - Pos0;
+			}
+
+			for (const int32 PointIndex : { 0, 1, 2 })
+			{
+				Normals[Face[PointIndex]] += Normal;
+				TangentUs[Face[PointIndex]] += TangentU;
+			}
+		}
+
+		const int32 NumVertices = Positions.Num();
+		ParallelFor(NumVertices, [&Normals, &TangentUs, &TangentVs](const int32 VertexIndex)
+			{
+				FVector3f& Normal = Normals[VertexIndex];
+				FVector3f& TangentU = TangentUs[VertexIndex];
+				FVector3f& TangentV = TangentVs[VertexIndex];
+
+				Normal= Normal.GetSafeNormal();
+				TangentU = TangentU.GetSafeNormal();
+				TangentU = TangentU - FVector3f::DotProduct(TangentU, Normal) * Normal;
+				TangentV = FVector3f::CrossProduct(TangentU, Normal);
+				TangentV = TangentV.GetSafeNormal();
+			}, NumVertices < 2000);
+	}
+
 	void FClothGeometryTools::BindMeshToRootBone(
 		const TSharedRef<FManagedArrayCollection>& ClothCollection,
 		bool bBindSimMesh,
diff --git a/Engine/Plugins/ChaosClothAsset/Source/ChaosClothAsset/Public/ChaosClothAsset/ClothGeometryTools.h b/Engine/Plugins/ChaosClothAsset/Source/ChaosClothAsset/Public/ChaosClothAsset/ClothGeometryTools.h
--- a/Engine/Plugins/ChaosClothAsset/Source/ChaosClothAsset/Public/ChaosClothAsset/ClothGeometryTools.h
+++ b/Engine/Plugins/ChaosClothAsset/Source/ChaosClothAsset/Public/ChaosClothAsset/ClothGeometryTools.h
@@ -53,6 +53,9 @@ namespace UE::Chaos::ClothAsset
 			const TArray<int32>& SimPatternSelection,
 			const TArray<int32>& RenderPatternSelection);
 
+		/** Recalculate the render mesh normals. */
+		static UE_API void RecalculateRenderMeshNormals(const TSharedRef<FManagedArrayCollection>& ClothCollection);
+
 		/**
 		 * Set the skinning weights for all of the sim/render vertices in ClothCollection to be bound to the root node.
 		 *
@@ -63,7 +66,6 @@ namespace UE::Chaos::ClothAsset
 			bool bBindSimMesh,
 			bool bBindRenderMesh);
 
-
 		/**
 		* Build (or add to) a ClothCollection Sim Mesh from the given 2D and 3D mesh data. Uses a Polygroup Attribute Layer to specify Pattern topology.
 		* 
diff --git a/Engine/Plugins/ChaosClothAssetEditor/Source/ChaosClothAssetDataflowNodes/Private/ChaosClothAsset/ChaosClothAssetDataflowNodesModule.cpp b/Engine/Plugins/ChaosClothAssetEditor/Source/ChaosClothAssetDataflowNodes/Private/ChaosClothAsset/ChaosClothAssetDataflowNodesModule.cpp
--- a/Engine/Plugins/ChaosClothAssetEditor/Source/ChaosClothAssetDataflowNodes/Private/ChaosClothAsset/ChaosClothAssetDataflowNodesModule.cpp
+++ b/Engine/Plugins/ChaosClothAssetEditor/Source/ChaosClothAssetDataflowNodes/Private/ChaosClothAsset/ChaosClothAssetDataflowNodesModule.cpp
@@ -27,6 +27,7 @@
 #include "ChaosClothAsset/ProxyDeformerNode.h"
 #include "ChaosClothAsset/RemeshNode.h"
 #include "ChaosClothAsset/ReverseNormalsNode.h"
+#include "ChaosClothAsset/RecalculateNormalsNode.h"
 #include "ChaosClothAsset/SelectionNode.h"
 #include "ChaosClothAsset/SelectionToIntMapNode.h"
 #include "ChaosClothAsset/SelectionToWeightMapNode.h"
@@ -124,6 +125,7 @@ namespace UE::Chaos::ClothAsset
 			DATAFLOW_NODE_REGISTER_CREATION_FACTORY(FChaosClothAssetProxyDeformerNode_v3);
 			DATAFLOW_NODE_REGISTER_CREATION_FACTORY(FChaosClothAssetRemeshNode_v2);
 			DATAFLOW_NODE_REGISTER_CREATION_FACTORY(FChaosClothAssetReverseNormalsNode);
+			DATAFLOW_NODE_REGISTER_CREATION_FACTORY(FChaosClothAssetRecalculateNormalsNode);
 			DATAFLOW_NODE_REGISTER_CREATION_FACTORY(FChaosClothAssetSelectionNode_v2);
 			DATAFLOW_NODE_REGISTER_CREATION_FACTORY(FChaosClothAssetSelectionToIntMapNode);
 			DATAFLOW_NODE_REGISTER_CREATION_FACTORY(FChaosClothAssetSelectionToWeightMapNode);
diff --git a/Engine/Plugins/ChaosClothAssetEditor/Source/ChaosClothAssetDataflowNodes/Private/ChaosClothAsset/RecalculateNormalsNode.cpp b/Engine/Plugins/ChaosClothAssetEditor/Source/ChaosClothAssetDataflowNodes/Private/ChaosClothAsset/RecalculateNormalsNode.cpp
new file mode 100644
--- /dev/null
+++ b/Engine/Plugins/ChaosClothAssetEditor/Source/ChaosClothAssetDataflowNodes/Private/ChaosClothAsset/RecalculateNormalsNode.cpp
@@ -0,0 +1,38 @@
+// Copyright Epic Games, Inc. All Rights Reserved.
+
+#include "ChaosClothAsset/RecalculateNormalsNode.h"
+#include "ChaosClothAsset/ClothGeometryTools.h"
+#include "ChaosClothAsset/CollectionClothFacade.h"
+#include "Dataflow/DataflowInputOutput.h"
+
+#include UE_INLINE_GENERATED_CPP_BY_NAME(RecalculateNormalsNode)
+
+#define LOCTEXT_NAMESPACE "ChaosClothAssetRecalculateNormalsNode"
+
+FChaosClothAssetRecalculateNormalsNode::FChaosClothAssetRecalculateNormalsNode(const UE::Dataflow::FNodeParameters& InParam, FGuid InGuid)
+	: FDataflowNode(InParam, InGuid)
+{
+	RegisterInputConnection(&Collection);
+	RegisterOutputConnection(&Collection, &Collection);
+}
+
+void FChaosClothAssetRecalculateNormalsNode::Evaluate(UE::Dataflow::FContext& Context, const FDataflowOutput* Out) const
+{
+	if (Out->IsA<FManagedArrayCollection>(&Collection))
+	{
+		using namespace UE::Chaos::ClothAsset;
+
+		// Evaluate in collection
+		FManagedArrayCollection InCollection = GetValue<FManagedArrayCollection>(Context, &Collection);
+		const TSharedRef<FManagedArrayCollection> ClothCollection = MakeShared<FManagedArrayCollection>(MoveTemp(InCollection));
+
+		if (FCollectionClothFacade(ClothCollection).IsValid())  // Can only act on the collection if it is a valid cloth collection
+		{
+			FClothGeometryTools::RecalculateRenderMeshNormals(ClothCollection);
+		}
+
+		SetValue(Context, MoveTemp(*ClothCollection), &Collection);
+	}
+}
+
+#undef LOCTEXT_NAMESPACE
diff --git a/Engine/Plugins/ChaosClothAssetEditor/Source/ChaosClothAssetDataflowNodes/Public/ChaosClothAsset/RecalculateNormalsNode.h b/Engine/Plugins/ChaosClothAssetEditor/Source/ChaosClothAssetDataflowNodes/Public/ChaosClothAsset/RecalculateNormalsNode.h
new file mode 100644
--- /dev/null
+++ b/Engine/Plugins/ChaosClothAssetEditor/Source/ChaosClothAssetDataflowNodes/Public/ChaosClothAsset/RecalculateNormalsNode.h
@@ -0,0 +1,27 @@
+// Copyright Epic Games, Inc. All Rights Reserved.
+
+#pragma once 
+
+#include "Dataflow/DataflowNode.h"
+#include "GeometryCollection/ManagedArrayCollection.h"
+#include "RecalculateNormalsNode.generated.h"
+
+/** Recalculate the geometry's normals. Experimental, only does the render geometry. */
+USTRUCT(Meta = (DataflowCloth, Experimental))
+struct FChaosClothAssetRecalculateNormalsNode : public FDataflowNode
+{
+	GENERATED_USTRUCT_BODY()
+	DATAFLOW_NODE_DEFINE_INTERNAL(FChaosClothAssetRecalculateNormalsNode, "RecalculateNormals", "Cloth", "Cloth Recalculate Simulation Render Mesh Normals")
+	DATAFLOW_NODE_RENDER_TYPE("SurfaceRender", FName("FClothCollection"), "Collection")
+
+public:
+	FChaosClothAssetRecalculateNormalsNode(const UE::Dataflow::FNodeParameters& InParam, FGuid InGuid = FGuid::NewGuid());
+
+private:
+	//~ Begin FDataflowNode Interface
+	virtual void Evaluate(UE::Dataflow::FContext& Context, const FDataflowOutput* Out) const override;
+	//~ End FDataflowNode Interface
+
+	UPROPERTY(Meta = (Dataflowinput, DataflowOutput, DataflowPassthrough = "Collection"))
+	FManagedArrayCollection Collection;
+};
diff --git a/Engine/Plugins/Experimental/ChaosOutfitAsset/Source/ChaosOutfitAssetEngine/Private/ChaosOutfitAsset/CollectionOutfitFacade.cpp b/Engine/Plugins/Experimental/ChaosOutfitAsset/Source/ChaosOutfitAssetEngine/Private/ChaosOutfitAsset/CollectionOutfitFacade.cpp
--- a/Engine/Plugins/Experimental/ChaosOutfitAsset/Source/ChaosOutfitAssetEngine/Private/ChaosOutfitAsset/CollectionOutfitFacade.cpp
+++ b/Engine/Plugins/Experimental/ChaosOutfitAsset/Source/ChaosOutfitAssetEngine/Private/ChaosOutfitAsset/CollectionOutfitFacade.cpp
@@ -208,7 +208,7 @@ namespace UE::Chaos::OutfitAsset
 					true)
 				{
 					// Protruding
-					bIsProtruding = Underbust > 5.f;
+					bIsProtruding = (Bust - Underbust) >= 15.f;
 
 					// Torso shape
 					const float HipBust = Hip / Bust;
