Chaos Cloth Asset - Fixed missing bounds when reloading old assets that cause them to disappear when loading the level. This will also fix the missing cloth asset thumbnails.

#rnx

diff --git a/Engine/Plugins/ChaosClothAsset/Source/ChaosClothAssetEngine/Private/ChaosClothAsset/ClothAsset.cpp b/Engine/Plugins/ChaosClothAsset/Source/ChaosClothAssetEngine/Private/ChaosClothAsset/ClothAsset.cpp
--- a/Engine/Plugins/ChaosClothAsset/Source/ChaosClothAssetEngine/Private/ChaosClothAsset/ClothAsset.cpp
+++ b/Engine/Plugins/ChaosClothAsset/Source/ChaosClothAssetEngine/Private/ChaosClothAsset/ClothAsset.cpp
@@ -221,6 +221,7 @@ void UChaosClothAsset::Serialize(FArchive& Ar)
 	Super::Serialize(Ar);
 
 	Ar.UsingCustomVersion(FFortniteMainBranchObjectVersion::GUID);
+	Ar.UsingCustomVersion(FUE5ReleaseStreamObjectVersion::GUID);
 
 	bool bCooked = Ar.IsCooking();
 	Ar << bCooked;
@@ -238,6 +239,13 @@ void UChaosClothAsset::Serialize(FArchive& Ar)
 		Private::Serialize(ChaosArchive, GetClothCollections());
 	}
 
+#if WITH_EDITOR
+	if (Ar.CustomVer(FUE5ReleaseStreamObjectVersion::GUID) < FUE5ReleaseStreamObjectVersion::RecalculateClothAssetSerializedBounds)
+	{
+		CalculateBounds();
+	}
+#endif
+
 	if (Ar.CustomVer(FFortniteMainBranchObjectVersion::GUID) < FFortniteMainBranchObjectVersion::AddClothAssetBase)
 	{
 		Ar << GetRefSkeleton(); // Moved to Cloth Asset Base serialization
diff --git a/Engine/Source/Runtime/Core/Public/UObject/UE5ReleaseStreamObjectVersions.inl b/Engine/Source/Runtime/Core/Public/UObject/UE5ReleaseStreamObjectVersions.inl
--- a/Engine/Source/Runtime/Core/Public/UObject/UE5ReleaseStreamObjectVersions.inl
+++ b/Engine/Source/Runtime/Core/Public/UObject/UE5ReleaseStreamObjectVersions.inl
@@ -174,5 +174,8 @@ UE5_RELEASE_VERSION(AnimModelGuidGenerationSettings, 54)
 // Add support for Standalone HLOD
 UE5_RELEASE_VERSION(WorldPartitionAddStandaloneHLODSupport, 55)
 
+// Fixed the missing bounds for cloth assets that don't have them serialized
+UE5_RELEASE_VERSION(RecalculateClothAssetSerializedBounds, 56)
+
 // -----<new versions can be added above this line>-------------------------------------------------
 #undef UE5_RELEASE_VERSION
\ No newline at end of file
