Chaos Outfit Asset - Fixed crash (removed check) when merging empty cloth assets into an outfit.

#rnx

diff --git a/Engine/Plugins/ChaosClothAsset/Source/ChaosClothAssetEngine/Public/ChaosClothAsset/ClothAsset.h b/Engine/Plugins/ChaosClothAsset/Source/ChaosClothAssetEngine/Public/ChaosClothAsset/ClothAsset.h
--- a/Engine/Plugins/ChaosClothAsset/Source/ChaosClothAssetEngine/Public/ChaosClothAsset/ClothAsset.h
+++ b/Engine/Plugins/ChaosClothAsset/Source/ChaosClothAssetEngine/Public/ChaosClothAsset/ClothAsset.h
@@ -164,7 +164,7 @@ PRAGMA_ENABLE_DEPRECATION_WARNINGS
 	UE_API virtual bool HasValidClothSimulationModels() const override;
 	virtual int32 GetNumClothSimulationModels() const override
 	{
-		return 1;
+		return ClothSimulationModel.IsValid() ? 1 : 0;
 	}
 	virtual FName GetClothSimulationModelName(int32 /*ModelIndex*/) const override
 	{
diff --git a/Engine/Plugins/Experimental/ChaosOutfitAsset/Source/ChaosOutfitAssetEngine/Private/ChaosOutfitAsset/Outfit.cpp b/Engine/Plugins/Experimental/ChaosOutfitAsset/Source/ChaosOutfitAssetEngine/Private/ChaosOutfitAsset/Outfit.cpp
--- a/Engine/Plugins/Experimental/ChaosOutfitAsset/Source/ChaosOutfitAssetEngine/Private/ChaosOutfitAsset/Outfit.cpp
+++ b/Engine/Plugins/Experimental/ChaosOutfitAsset/Source/ChaosOutfitAssetEngine/Private/ChaosOutfitAsset/Outfit.cpp
@@ -1133,8 +1133,6 @@ void UChaosOutfit::Merge(
 {
 	using namespace UE::Chaos::OutfitAsset;
 
-	check(InSkeletalMeshRenderData);
-
 	// Filter the body size pieces to merge, or select all pieces
 	TArray<FGuid> AssetGuidsFilter;
 
