Chaos Cloth Asset - Fixed ability to paint the RenderDeformerSkinningBlend weight map broken in CL 42453857.

#rnx

diff --git a/Engine/Plugins/ChaosClothAsset/Source/ChaosClothAsset/Private/ChaosClothAsset/ClothCollection.cpp b/Engine/Plugins/ChaosClothAsset/Source/ChaosClothAsset/Private/ChaosClothAsset/ClothCollection.cpp
--- a/Engine/Plugins/ChaosClothAsset/Source/ChaosClothAsset/Private/ChaosClothAsset/ClothCollection.cpp
+++ b/Engine/Plugins/ChaosClothAsset/Source/ChaosClothAsset/Private/ChaosClothAsset/ClothCollection.cpp
@@ -177,8 +177,7 @@ namespace UE::Chaos::ClothAsset::Private
 	static const TMap<FName, TArray<FName>> UserAccessibleAttributeNamesMap =
 	{
 		{ ClothCollectionGroup::SimVertices3D, { ClothCollectionAttribute::SimCustomResizingBlend } },  // Can be painted if required
-		{ ClothCollectionGroup::RenderVertices, { ClothCollectionAttribute::RenderDeformerSkinningBlend } },  // Can be painted if required
-		{ ClothCollectionGroup::RenderVertices, { ClothCollectionAttribute::RenderCustomResizingBlend } },  // Can be painted if required
+		{ ClothCollectionGroup::RenderVertices, { ClothCollectionAttribute::RenderDeformerSkinningBlend, ClothCollectionAttribute::RenderCustomResizingBlend } },  // Can be painted if required
 	};
 }  // End namespace UE::Chaos::ClothAsset::Private
 
