ChaosCloth: Support BindToLeaderComponent to use leader component simulation to drive the follower render mesh.

Currently only supported when using bUseAttachedParentAsPoseComponent

diff --git a/Engine/Plugins/ChaosClothAsset/Source/ChaosClothAssetEngine/Private/ChaosClothAsset/ClothComponent.cpp b/Engine/Plugins/ChaosClothAsset/Source/ChaosClothAssetEngine/Private/ChaosClothAsset/ClothComponent.cpp
--- a/Engine/Plugins/ChaosClothAsset/Source/ChaosClothAssetEngine/Private/ChaosClothAsset/ClothComponent.cpp
+++ b/Engine/Plugins/ChaosClothAsset/Source/ChaosClothAssetEngine/Private/ChaosClothAsset/ClothComponent.cpp
@@ -290,7 +290,15 @@ void UChaosClothComponent::OnUnregister()
 
 bool UChaosClothComponent::IsComponentTickEnabled() const
 {
-	return bEnableSimulation && Super::IsComponentTickEnabled();
+	bool bLeaderClothComponentTickEnabled = false;
+	if (bBindToLeaderComponent)
+	{
+		if (UChaosClothComponent* const LeaderComponent = Cast<UChaosClothComponent>(LeaderPoseComponent.Get()))
+		{
+			bLeaderClothComponentTickEnabled = LeaderComponent->IsComponentTickEnabled();
+		}
+	}
+	return (bEnableSimulation || bLeaderClothComponentTickEnabled) && Super::IsComponentTickEnabled();
 }
 
 void UChaosClothComponent::TickComponent(float DeltaTime, enum ELevelTick TickType, FActorComponentTickFunction* ThisTickFunction)
@@ -340,6 +348,13 @@ bool UChaosClothComponent::RequiresPreEndOfFrameSync() const
 		// then we will wait for the cloth task in SendAllEndOfFrameUpdates.
 		return true;
 	}
+	if (bBindToLeaderComponent)
+	{
+		if (UChaosClothComponent* const LeaderComponent = LeaderClothComponent.Get())
+		{
+			return LeaderComponent->RequiresPreEndOfFrameSync();
+		}
+	}
 	return Super::RequiresPreEndOfFrameSync();
 }
 
@@ -438,6 +453,11 @@ void UChaosClothComponent::OnAttachmentChanged()
 	if (bUseAttachedParentAsPoseComponent)
 	{
 		USkinnedMeshComponent* const AttachParentComponent = Cast<USkinnedMeshComponent>(GetAttachParent());
+		if (UChaosClothComponent* const AttachClothComponent = Cast<UChaosClothComponent>(AttachParentComponent))
+		{
+			LeaderClothComponent = AttachClothComponent;
+		}
+
 		SetLeaderPoseComponent(AttachParentComponent);  // If the cast fail, remove the current leader
 
 		// When parented to a skeletal mesh, the anim setup needs re-initializing in order to use the follower's bones requirement
@@ -450,6 +470,11 @@ void UChaosClothComponent::OnAttachmentChanged()
 	Super::OnAttachmentChanged();
 }
 
+bool UChaosClothComponent::IsVisible() const
+{
+	return bHasValidRenderDataForVisibility && Super::IsVisible();
+}
+
 void UChaosClothComponent::RefreshBoneTransforms(FActorComponentTickFunction* /*TickFunction*/)
 {
 	MarkRenderDynamicDataDirty();
@@ -464,7 +489,7 @@ void UChaosClothComponent::GetUpdateClothSimulationData_AnyThread(TMap<int32, FC
 {
 	OutLocalToWorld = GetComponentToWorld().ToMatrixWithScale();
 
-	const UChaosClothComponent* const LeaderPoseClothComponent = Cast<UChaosClothComponent>(LeaderPoseComponent.Get());
+	const UChaosClothComponent* const LeaderPoseClothComponent = LeaderClothComponent.Get();
 	if (LeaderPoseClothComponent && LeaderPoseClothComponent->ClothSimulationProxy && bBindToLeaderComponent)
 	{
 		OutBlendWeight = BlendWeight;
@@ -672,7 +697,7 @@ void UChaosClothComponent::HandleExistingParallelSimulation()
 {
 	if (bBindToLeaderComponent)
 	{
-		if (UChaosClothComponent* const LeaderComponent = Cast<UChaosClothComponent>(LeaderPoseComponent.Get()))
+		if (UChaosClothComponent* const LeaderComponent = LeaderClothComponent.Get())
 		{
 			LeaderComponent->HandleExistingParallelSimulation();
 		}
@@ -713,11 +738,11 @@ void UChaosClothComponent::UpdateVisibility()
 	{
 		const FSkeletalMeshRenderData* const SkeletalMeshRenderData = GetAsset()->GetResourceForRendering();
 		const int32 FirstValidLODIdx = SkeletalMeshRenderData ? SkeletalMeshRenderData->GetFirstValidLODIdx(0) : INDEX_NONE;
-		SetVisibility(FirstValidLODIdx != INDEX_NONE);
+		bHasValidRenderDataForVisibility = FirstValidLODIdx != INDEX_NONE;
 	}
 	else
 	{
-		SetVisibility(false);
+		bHasValidRenderDataForVisibility = false;
 	}
 }
 
diff --git a/Engine/Plugins/ChaosClothAsset/Source/ChaosClothAssetEngine/Public/ChaosClothAsset/ClothComponent.h b/Engine/Plugins/ChaosClothAsset/Source/ChaosClothAssetEngine/Public/ChaosClothAsset/ClothComponent.h
--- a/Engine/Plugins/ChaosClothAsset/Source/ChaosClothAssetEngine/Public/ChaosClothAsset/ClothComponent.h
+++ b/Engine/Plugins/ChaosClothAsset/Source/ChaosClothAssetEngine/Public/ChaosClothAsset/ClothComponent.h
@@ -237,6 +237,16 @@ public:
 	UFUNCTION(BlueprintGetter, Category = ClothComponent)
 	UE_API float GetTeleportRotationThreshold() const;
 
+	void SetBindToLeaderComponent(bool bValue)
+	{
+		bBindToLeaderComponent = bValue;
+	}
+
+	bool GetBindToLeaderComponent() const
+	{
+		return bBindToLeaderComponent;
+	}
+
 	/**
 	* Sets the teleportation rotation threshold.
 	*
@@ -287,6 +297,7 @@ protected:
 	//~ Begin USceneComponent Interface.
 	UE_API virtual FBoxSphereBounds CalcBounds(const FTransform& LocalToWorld) const override;
 	UE_API virtual void OnAttachmentChanged() override;
+	UE_API virtual bool IsVisible() const override;
 	//~ End USceneComponent Interface
 
 	//~ Begin USkinnedMeshComponent Interface
@@ -357,8 +368,14 @@ private:
 	UPROPERTY(Interp, Category = ClothComponent)
 	uint8 bSuspendSimulation : 1;
 
-	/** Whether to use the leader component pose. */
-	UPROPERTY()
+	/** Whether to use the leader component simulation result. 
+	 * Currently only supported when using UseAttachedParentAsPoseComponent
+	 * 
+	 * If this component has a valid LeaderPoseComponent then this makes cloth items on the follower component
+	 * take the transforms of the cloth items on the leader component instead of simulating separately.
+	 * @Note The meshes used in the components must be identical for the cloth to bind correctly
+	 */
+	UPROPERTY(EditAnywhere, Category = ClothComponent)
 	uint8 bBindToLeaderComponent : 1;
 
 	/** Whether to teleport the cloth prior to advancing the simulation. */
@@ -384,6 +401,10 @@ private:
 	TObjectPtr<UChaosClothAssetBase> Asset;
 #endif
 
+	/** Cloth leader component used to drive simulation. This may have its own leader component */
+	UPROPERTY()
+	TWeakObjectPtr<UChaosClothComponent> LeaderClothComponent;
+
 	/**
 	* Conduct teleportation if the character's movement is greater than this threshold in 1 frame.
 	* Zero or negative values will skip the check.
@@ -418,6 +439,8 @@ private:
 	/** Like bReset, but cleared every frame */
 	bool bResetOnce = false;
 
+	bool bHasValidRenderDataForVisibility = false;
+
 	/**
 	 * Simulation properties per cloth model. 
 	 * Transient because it doesn't need to be serialized but contains pointers to UObject that needs not be garbage collected.
