Set ChaosCloth to do async initialization.

diff --git a/Engine/Plugins/ChaosClothAsset/Source/ChaosClothAssetEngine/Private/ChaosClothAsset/ClothSimulationProxy.cpp b/Engine/Plugins/ChaosClothAsset/Source/ChaosClothAssetEngine/Private/ChaosClothAsset/ClothSimulationProxy.cpp
--- a/Engine/Plugins/ChaosClothAsset/Source/ChaosClothAssetEngine/Private/ChaosClothAsset/ClothSimulationProxy.cpp
+++ b/Engine/Plugins/ChaosClothAsset/Source/ChaosClothAssetEngine/Private/ChaosClothAsset/ClothSimulationProxy.cpp
@@ -223,6 +223,7 @@ PRAGMA_ENABLE_DEPRECATION_WARNINGS
 
 	void FClothSimulationProxy::ExecuteInitialization()
 	{
+		TRACE_CPUPROFILER_EVENT_SCOPE(ClothSimulationProxy_ExecuteInitialization);
 		LLM_SCOPE_BYNAME(TEXT("Physics/Cloth"));
 		int32 LocalNumCloths = 0;
 		int32 LocalNumKinematicParticles = 0;
@@ -248,7 +249,7 @@ PRAGMA_ENABLE_DEPRECATION_WARNINGS
 		using namespace ::Chaos;
 		if (IsValidRef(ParallelInitializationTask))
 		{
-			TRACE_CPUPROFILER_EVENT_SCOPE(STAT_ClothSimulationProxy_CompleteInitialization);
+			TRACE_CPUPROFILER_EVENT_SCOPE(ClothSimulationProxy_CompleteInitialization);
 			CSV_SCOPED_SET_WAIT_STAT(Cloth);
 
 			FTaskGraphInterface::Get().WaitUntilTaskCompletes(ParallelInitializationTask, ENamedThreads::GameThread);
