Chaos Cloth: Remove code stripping unused maps.

This allows maps to be set at runtime via blueprint.

diff --git a/Engine/Plugins/ChaosClothAsset/Source/ChaosClothAssetEngine/Private/ChaosClothAsset/ClothSimulationModel.cpp b/Engine/Plugins/ChaosClothAsset/Source/ChaosClothAssetEngine/Private/ChaosClothAsset/ClothSimulationModel.cpp
--- a/Engine/Plugins/ChaosClothAsset/Source/ChaosClothAssetEngine/Private/ChaosClothAsset/ClothSimulationModel.cpp
+++ b/Engine/Plugins/ChaosClothAsset/Source/ChaosClothAssetEngine/Private/ChaosClothAsset/ClothSimulationModel.cpp
@@ -288,10 +288,8 @@ FChaosClothSimulationModel::FChaosClothSimulationModel(const TArray<TSharedRef<c
 		LodModel.WeightMaps.Reserve(WeightMapNames.Num());
 		for (const FName& WeightMapName : WeightMapNames)
 		{
-			if (IsNameUsedInProperties(WeightMapName) || IsNameReserved(WeightMapName))  // Only copy weight map used by properties - TODO: Add a switch or flag to disable this behavior when enabling runtime weightmap switching
-			{
-				LodModel.WeightMaps.Add(WeightMapName) = Cloth.GetWeightMap(WeightMapName);
-			}
+			// Keep all weight maps to allow runtime switching.
+			LodModel.WeightMaps.Add(WeightMapName) = Cloth.GetWeightMap(WeightMapName);
 		}
 
 		// Copy vertex and face sets
