ChaosCloth: add option to check for non-finite values in the sim data getting sent to render (default on). Add option to reset the cloth simulation if this is found (default on).

New cvars are:
p.ChaosClothAsset.SimulationDataCheckFinite
p.ChaosClothAsset.SimulationDataResetOnNonFinite

diff --git a/Engine/Plugins/ChaosClothAsset/Source/ChaosClothAssetEngine/Private/ChaosClothAsset/ClothSimulationProxy.cpp b/Engine/Plugins/ChaosClothAsset/Source/ChaosClothAssetEngine/Private/ChaosClothAsset/ClothSimulationProxy.cpp
--- a/Engine/Plugins/ChaosClothAsset/Source/ChaosClothAssetEngine/Private/ChaosClothAsset/ClothSimulationProxy.cpp
+++ b/Engine/Plugins/ChaosClothAsset/Source/ChaosClothAssetEngine/Private/ChaosClothAsset/ClothSimulationProxy.cpp
@@ -49,6 +49,11 @@ namespace UE::Chaos::ClothAsset
 	bool bWaitForAsyncClothInitialization = true;
 	FAutoConsoleVariableRef CVarWaitForAsyncClothInitialization(TEXT("p.ChaosClothAsset.WaitForAsyncClothInitialization"), bWaitForAsyncClothInitialization, TEXT("When asynchronous cloth proxy initialization is enabled, wait for initialization to complete to start up cloth simulation. Otherwise, cloth simulation will be disabled until initialization has completed."));
 
+	bool bSimulationDataCheckFinite = true;
+	FAutoConsoleVariableRef CVarSimulationDataCheckFinite(TEXT("p.ChaosClothAsset.SimulationDataCheckFinite"), bSimulationDataCheckFinite, TEXT("Check for nans/infs in cloth simulation data"));
+	bool bSimulationDataResetOnNonFinite = true;
+	FAutoConsoleVariableRef CVarSimulationDataResetOnNonFinite(TEXT("p.ChaosClothAsset.SimulationDataResetOnNonFinite"), bSimulationDataResetOnNonFinite, TEXT("Reset cloth simulation if a nan/inf is detected. Requires p.ChaosClothAsset.SimulationDataCheckFinite"));
+
 	static FAutoConsoleTaskPriority CPrio_ClothSimulationProxyParallelTask(
 		TEXT("TaskGraph.TaskPriorities.ClothSimulationProxyParallelTask"),
 		TEXT("Task and thread priority for the cloth simulation proxy."),
@@ -460,7 +465,7 @@ PRAGMA_ENABLE_DEPRECATION_WARNINGS
 		const double StartTime = FPlatformTime::Seconds();
 		const float PrevSimulationTime = SimulationTime;  // Copy the atomic to prevent a re-read
 
-		const bool bNeedsReset = (ClothSimulationContext->bReset || PrevSimulationTime == 0.f);  // Reset on the first frame too since the simulation is created in bind pose, and not in start pose
+		const bool bNeedsReset = (bNeedResetClothFromNanCheck || ClothSimulationContext->bReset || PrevSimulationTime == 0.f);  // Reset on the first frame too since the simulation is created in bind pose, and not in start pose
 		const bool bNeedsTeleport = ClothSimulationContext->bTeleport;
 		bIsTeleported = bNeedsTeleport;
 
@@ -533,6 +538,8 @@ PRAGMA_ENABLE_DEPRECATION_WARNINGS
 			LastLinearSolveIterations = 0;
 		}
 
+		bNeedResetClothFromNanCheck = false;
+
 		// Visualization
 #if CHAOS_DEBUG_DRAW
 		static const TConsoleVariableData<bool>* const DebugDrawLocalSpaceCVar = IConsoleManager::Get().FindTConsoleVariableDataBool(TEXT("p.ChaosCloth.DebugDrawLocalSpace"));
@@ -748,8 +755,10 @@ PRAGMA_ENABLE_DEPRECATION_WARNINGS
 			Data.Positions = Cloth->GetParticlePositions(Solver.Get());
 			Data.Normals = Cloth->GetParticleNormals(Solver.Get());
 
+			check(Data.Positions.Num() == Data.Normals.Num());
+
 			// Transform into the cloth reference simulation space used at the time of simulation
-			check(Data.Positions.Num() == Data.Normals.Num())
+			bool bSimDataIsFinite = true;
 #if INTEL_ISPC
 			if (bTransformClothSimulData_ISPC_Enabled)
 			{
@@ -757,20 +766,61 @@ PRAGMA_ENABLE_DEPRECATION_WARNINGS
 				check(sizeof(ispc::FVector3f) == Data.Positions.GetTypeSize());
 				check(sizeof(ispc::FVector3f) == Data.Normals.GetTypeSize());
 
-				ispc::TransformClothSimulData(
-					(ispc::FVector3f*)Data.Positions.GetData(),
-					(ispc::FVector3f*)Data.Normals.GetData(),
-					(ispc::FTransform&)ReferenceSpaceTransform,
-					LocalSpaceScale,
-					Data.Positions.Num());
+				if (bSimulationDataCheckFinite)
+				{
+					bSimDataIsFinite = ispc::TransformClothSimulDataAndCheckFinite(
+						(ispc::FVector3f*)Data.Positions.GetData(),
+						(ispc::FVector3f*)Data.Normals.GetData(),
+						(ispc::FTransform&)ReferenceSpaceTransform,
+						LocalSpaceScale,
+						Data.Positions.Num());
+				}
+				else
+				{
+					ispc::TransformClothSimulData(
+						(ispc::FVector3f*)Data.Positions.GetData(),
+						(ispc::FVector3f*)Data.Normals.GetData(),
+						(ispc::FTransform&)ReferenceSpaceTransform,
+						LocalSpaceScale,
+						Data.Positions.Num());
+				}
 			}
 			else
 #endif
 			{
-				for (int32 Index = 0; Index < Data.Positions.Num(); ++Index)
+				if (bSimulationDataCheckFinite)
+				{
+					auto VectorIsFinite = [](const FSolverVec3& Vector)
+						{
+							return FMath::IsFinite(Vector.X) && FMath::IsFinite(Vector.Y) && FMath::IsFinite(Vector.Z);
+						};
+					for (int32 Index = 0; Index < Data.Positions.Num(); ++Index)
+					{
+						Data.Positions[Index] = FVec3f(ReferenceSpaceTransform.InverseTransformPosition(LocalSpaceScale * FVec3(Data.Positions[Index])));
+						Data.Normals[Index] = FVec3f(ReferenceSpaceTransform.InverseTransformVector(FVec3(-Data.Normals[Index])));
+						bSimDataIsFinite = bSimDataIsFinite && VectorIsFinite(Data.Positions[Index]) && VectorIsFinite(Data.Normals[Index]);
+					}
+				}
+				else
+				{
+					for (int32 Index = 0; Index < Data.Positions.Num(); ++Index)
+					{
+						Data.Positions[Index] = FVec3f(ReferenceSpaceTransform.InverseTransformPosition(LocalSpaceScale * FVec3(Data.Positions[Index])));
+						Data.Normals[Index] = FVec3f(ReferenceSpaceTransform.InverseTransformVector(FVec3(-Data.Normals[Index])));
+					}
+				}
+			}
+
+
+			if (bSimulationDataCheckFinite)
+			{
+				UE_CLOG(!bSimDataIsFinite, LogChaosClothAsset, Warning, TEXT("Failed to write back clothing simulation data for component %s as simulation contains non-finite values."), *ClothComponent.GetName());
+				ensure(bSimDataIsFinite);
+				if (!bSimDataIsFinite && bSimulationDataResetOnNonFinite)
 				{
-					Data.Positions[Index] = FVec3f(ReferenceSpaceTransform.InverseTransformPosition(LocalSpaceScale * FVec3(Data.Positions[Index])));
-					Data.Normals[Index] = FVec3f(ReferenceSpaceTransform.InverseTransformVector(FVec3(-Data.Normals[Index])));
+					bNeedResetClothFromNanCheck = true;
+					CurrentSimulationData.Reset();
+					return;
 				}
 			}
 
diff --git a/Engine/Plugins/ChaosClothAsset/Source/ChaosClothAssetEngine/Private/ChaosClothAsset/ClothSimulationProxy.ispc b/Engine/Plugins/ChaosClothAsset/Source/ChaosClothAssetEngine/Private/ChaosClothAsset/ClothSimulationProxy.ispc
--- a/Engine/Plugins/ChaosClothAsset/Source/ChaosClothAssetEngine/Private/ChaosClothAsset/ClothSimulationProxy.ispc
+++ b/Engine/Plugins/ChaosClothAsset/Source/ChaosClothAssetEngine/Private/ChaosClothAsset/ClothSimulationProxy.ispc
@@ -43,3 +43,53 @@ export void TransformClothSimulData(uniform FVector3f Positions[],
 		VectorStore(&Normals[Index], OutNormalFloat);
 	}
 }
+
+inline bool VectorIsFinite(const FVector3f& Vec)
+{
+	return (intbits(Vec.V[0]) & 0x7f800000) != 0x7f800000 
+		&& (intbits(Vec.V[1]) & 0x7f800000) != 0x7f800000 
+		&& (intbits(Vec.V[2]) & 0x7f800000) != 0x7f800000;
+}
+
+export uniform bool TransformClothSimulDataAndCheckFinite(uniform FVector3f Positions[],
+									uniform FVector3f Normals[],
+									const uniform FTransform &ReferenceSpaceTransform,
+									const uniform double LocalSpaceScale,
+									const uniform int NumPositions)
+{
+	const uniform FVector4 RotationInverse = QuatInverse(ReferenceSpaceTransform.Rotation);
+	const uniform FVector Translation = SetVector(ReferenceSpaceTransform.Translation);
+	const uniform FVector SafeReciprocal = TransformGetSafeScaleReciprocal(SetVector(ReferenceSpaceTransform.Scale3D), SMALL_NUMBER);
+
+	varying bool bIsFinite = true;
+	foreach(i = 0 ... NumPositions)
+	{
+		const uniform int Index = extract(i, 0);
+		const FVector3f PositionFloat = VectorLoad(&Positions[Index]);
+		const FVector3f NormalFloat = VectorLoad(&Normals[Index]);
+
+
+		const FVector InputPositionVector = ConvertVector3fTo3Native(PositionFloat);
+
+		const FVector3f InputNormalVectorFloat = NormalFloat * -FLOAT_ONE;
+		const FVector InputNormalVector = ConvertVector3fTo3Native(InputNormalVectorFloat);
+
+		const FVector ScaledPosition = InputPositionVector * LocalSpaceScale;
+		const FVector TranslatedPosition = ScaledPosition - Translation;
+
+		const FVector VRPosition = VectorQuaternionRotateVector(RotationInverse, TranslatedPosition);
+		const FVector VRNormal = VectorQuaternionRotateVector(RotationInverse, InputNormalVector);
+
+		const FVector OutPosition = VRPosition * SafeReciprocal;
+		const FVector OutNormal = VRNormal * SafeReciprocal;
+
+		const FVector3f OutPositionFloat = ConvertVector3NativeTo3f(OutPosition);
+		const FVector3f OutNormalFloat = ConvertVector3NativeTo3f(OutNormal);
+		
+		bIsFinite = bIsFinite && VectorIsFinite(OutPositionFloat) && VectorIsFinite(OutNormalFloat);
+
+		VectorStore(&Positions[Index], OutPositionFloat);
+		VectorStore(&Normals[Index], OutNormalFloat);
+	}
+	return all(bIsFinite);
+}
diff --git a/Engine/Plugins/ChaosClothAsset/Source/ChaosClothAssetEngine/Public/ChaosClothAsset/ClothSimulationProxy.h b/Engine/Plugins/ChaosClothAsset/Source/ChaosClothAssetEngine/Public/ChaosClothAsset/ClothSimulationProxy.h
--- a/Engine/Plugins/ChaosClothAsset/Source/ChaosClothAssetEngine/Public/ChaosClothAsset/ClothSimulationProxy.h
+++ b/Engine/Plugins/ChaosClothAsset/Source/ChaosClothAssetEngine/Public/ChaosClothAsset/ClothSimulationProxy.h
@@ -199,6 +199,7 @@ PRAGMA_ENABLE_DEPRECATION_WARNINGS
 		std::atomic<float> LastLinearSolveError; // For single cloth only.
 
 		mutable bool bHasInvalidReferenceBoneTransforms = false;
+		bool bNeedResetClothFromNanCheck = false;
 
 		// Cached value of the MaxPhysicsDeltaTime setting for the life of this proxy
 		const float MaxDeltaTime;
