Chaos Cloth: make game thread tick nearly free when cloth is not simulating. Previously, we were doing all of the setup for a simulation even though we weren't simulating.

Also fix some bounding box issues:
* Handle case where the bbox from the simulation is non-finite
* Handle BindToLeaderComponent.

diff --git a/Engine/Plugins/ChaosClothAsset/Source/ChaosClothAssetEngine/Private/ChaosClothAsset/ClothComponent.cpp b/Engine/Plugins/ChaosClothAsset/Source/ChaosClothAssetEngine/Private/ChaosClothAsset/ClothComponent.cpp
--- a/Engine/Plugins/ChaosClothAsset/Source/ChaosClothAssetEngine/Private/ChaosClothAsset/ClothComponent.cpp
+++ b/Engine/Plugins/ChaosClothAsset/Source/ChaosClothAssetEngine/Private/ChaosClothAsset/ClothComponent.cpp
@@ -312,7 +312,7 @@ void UChaosClothComponent::TickComponent(float DeltaTime, enum ELevelTick TickTy
 	// Make sure that the previous frame simulation has completed
 	HandleExistingParallelSimulation();
 
-	if (ClothSimulationProxy.IsValid())
+	if (IsSimulationEnabled())
 	{
 		UpdateClothTeleport();
 	}
@@ -423,14 +423,29 @@ FBoxSphereBounds UChaosClothComponent::CalcBounds(const FTransform& LocalToWorld
 		}
 
 		// Add sim mesh bounds
-		if (ClothSimulationProxy.IsValid())
+		FBoxSphereBounds SimulationBounds(ForceInit);
+		if (IsSimulationEnabled())
+		{
+			SimulationBounds = ClothSimulationProxy->CalculateBounds_AnyThread();
+		}
+		else if (bBindToLeaderComponent)
 		{
-			const FBoxSphereBounds SimulationBounds = ClothSimulationProxy->CalculateBounds_AnyThread();
-			if (SimulationBounds.SphereRadius > UE_SMALL_NUMBER)  // Don't add the simulation bounds if there are empty, otherwise it could unwillingly add the component's location
+			if (const UChaosClothComponent* const LeaderComponent = LeaderClothComponent.Get())
 			{
-				NewBounds = NewBounds + SimulationBounds.TransformBy(CachedBoundsTransform);
+				if (LeaderComponent->IsSimulationEnabled())
+				{
+					const UE::Chaos::ClothAsset::FClothSimulationProxy* const LeaderProxy = LeaderComponent->GetClothSimulationProxy();
+					check(LeaderProxy); // valid pointer is included in IsSimulationEnabled
+					SimulationBounds = LeaderProxy->CalculateBounds_AnyThread();
+				}
 			}
 		}
+		if (SimulationBounds.SphereRadius > UE_SMALL_NUMBER && FMath::IsFinite(SimulationBounds.Origin.X) && FMath::IsFinite(SimulationBounds.Origin.Y) && FMath::IsFinite(SimulationBounds.Origin.Z) &&
+			FMath::IsFinite(SimulationBounds.BoxExtent.X) && FMath::IsFinite(SimulationBounds.BoxExtent.Y) && FMath::IsFinite(SimulationBounds.BoxExtent.Z) &&
+			FMath::IsFinite(SimulationBounds.SphereRadius))  // Don't add the simulation bounds if there are empty, otherwise it could unwillingly add the component's location
+		{
+			NewBounds = NewBounds + SimulationBounds.TransformBy(CachedBoundsTransform);
+		}
 
 		CachedWorldOrLocalSpaceBounds = NewBounds;
 		bCachedLocalBoundsUpToDate = bCacheLocalSpaceBounds;
diff --git a/Engine/Plugins/ChaosClothAsset/Source/ChaosClothAssetEngine/Private/ChaosClothAsset/ClothSimulationContext.cpp b/Engine/Plugins/ChaosClothAsset/Source/ChaosClothAssetEngine/Private/ChaosClothAsset/ClothSimulationContext.cpp
--- a/Engine/Plugins/ChaosClothAsset/Source/ChaosClothAssetEngine/Private/ChaosClothAsset/ClothSimulationContext.cpp
+++ b/Engine/Plugins/ChaosClothAsset/Source/ChaosClothAssetEngine/Private/ChaosClothAsset/ClothSimulationContext.cpp
@@ -19,6 +19,15 @@ namespace UE::Chaos::ClothAsset
 
 		// Set the current LOD index
 		LodIndex = ClothComponent.GetPredictedLODLevel();
+		const UChaosClothAssetBase* const Asset = ClothComponent.GetAsset();
+		if (!Asset ||
+			!Asset->GetResourceForRendering() ||
+			!Asset->GetResourceForRendering()->LODRenderData.IsValidIndex(LodIndex) ||
+			!Asset->GetResourceForRendering()->LODRenderData[LodIndex].HasClothData())
+		{
+			// This Lod will not simulate, so there's no need to update the context.
+			return;
+		}
 
 		// Set the teleport mode
 		static const IConsoleVariable* const CVarMaxDeltaTimeTeleportMultiplier = IConsoleManager::Get().FindConsoleVariable(TEXT("p.Cloth.MaxDeltaTimeTeleportMultiplier"));
@@ -48,7 +57,6 @@ namespace UE::Chaos::ClothAsset
 		SolverGeometryScale = ClothComponent.GetClothGeometryScale();;
 
 		// Update bone transforms
-		const UChaosClothAssetBase* const Asset = ClothComponent.GetAsset();
 		const FReferenceSkeleton* const ReferenceSkeleton = Asset ? &Asset->GetRefSkeleton() : nullptr;
 		int32 NumBones = ReferenceSkeleton ? ReferenceSkeleton->GetNum() : 0;
 
diff --git a/Engine/Plugins/ChaosClothAsset/Source/ChaosClothAssetEngine/Private/ChaosClothAsset/ClothSimulationProxy.cpp b/Engine/Plugins/ChaosClothAsset/Source/ChaosClothAssetEngine/Private/ChaosClothAsset/ClothSimulationProxy.cpp
--- a/Engine/Plugins/ChaosClothAsset/Source/ChaosClothAssetEngine/Private/ChaosClothAsset/ClothSimulationProxy.cpp
+++ b/Engine/Plugins/ChaosClothAsset/Source/ChaosClothAssetEngine/Private/ChaosClothAsset/ClothSimulationProxy.cpp
@@ -311,6 +311,11 @@ PRAGMA_ENABLE_DEPRECATION_WARNINGS
 
 		bIsPreProcessed = true;
 
+		if (!ClothComponent.IsSimulationEnabled())
+		{
+			return;
+		}
+
 		// Fill a new context, note the context is also needed when the simulation is suspended or playing back the cache
 		constexpr bool bIsInitializationFalse = false;
 		PRAGMA_DISABLE_DEPRECATION_WARNINGS
@@ -399,12 +404,12 @@ PRAGMA_ENABLE_DEPRECATION_WARNINGS
 		}
 		if (!bIsSimulating)
 		{
-			// Take care of the LOD switching, as the simulation won't do it
-			UpdateClothLODs();
 
 			// If the simulation is enabled, then it is suspended or reading from cache and the simulation data still needs updating (transforms and LODs)
 			if (ClothComponent.IsSimulationEnabled())
 			{
+				// Take care of the LOD switching, as the simulation won't do it
+				UpdateClothLODs();
 				WriteSimulationData();
 			}
 			else
@@ -846,9 +851,9 @@ PRAGMA_ENABLE_DEPRECATION_WARNINGS
 	{
 		SCOPE_CYCLE_COUNTER(STAT_ClothSimulationProxy_CalculateBounds);
 
-		check(Solver);
 		if (bIsPreProcessed && (!IsValidRef(ParallelTask) || ParallelTask->IsComplete()))
 		{
+			check(Solver);
 			FBoxSphereBounds Bounds = Solver->CalculateBounds();
 
 			// The component could be moving while the simulation is suspended so getting the bounds
