Cloth Resize : Update ClothGeometryTools RecomputeNormals to retain the BasisDeterminant sign of previously computed Tangents.

diff --git a/Engine/Plugins/ChaosClothAsset/Source/ChaosClothAsset/Private/ChaosClothAsset/ClothGeometryTools.cpp b/Engine/Plugins/ChaosClothAsset/Source/ChaosClothAsset/Private/ChaosClothAsset/ClothGeometryTools.cpp
--- a/Engine/Plugins/ChaosClothAsset/Source/ChaosClothAsset/Private/ChaosClothAsset/ClothGeometryTools.cpp
+++ b/Engine/Plugins/ChaosClothAsset/Source/ChaosClothAsset/Private/ChaosClothAsset/ClothGeometryTools.cpp
@@ -755,6 +755,14 @@ namespace UE::Chaos::ClothAsset
 		check(Positions.Num() == TangentUs.Num());
 		check(Positions.Num() == TangentVs.Num());
 
+		// Record BasisDeterminantSigns to preserve them when calculating TangentVs.
+		TArray<float> BasisDeterminantSigns;
+		BasisDeterminantSigns.SetNumUninitialized(Normals.Num());
+		for (int32 VertexIndex = 0; VertexIndex < Normals.Num(); ++VertexIndex)
+		{
+			BasisDeterminantSigns[VertexIndex] = GetBasisDeterminantSign(FVector3d(TangentUs[VertexIndex]), FVector3d(TangentVs[VertexIndex]), FVector3d(Normals[VertexIndex]));
+		}
+
 		FMemory::Memzero(Normals.GetData(), Normals.NumBytes());
 		FMemory::Memzero(TangentUs.GetData(), TangentUs.NumBytes());
 
@@ -791,7 +799,7 @@ namespace UE::Chaos::ClothAsset
 		}
 
 		const int32 NumVertices = Positions.Num();
-		ParallelFor(NumVertices, [&Normals, &TangentUs, &TangentVs](const int32 VertexIndex)
+		ParallelFor(NumVertices, [&Normals, &TangentUs, &TangentVs, &BasisDeterminantSigns](const int32 VertexIndex)
 			{
 				FVector3f& Normal = Normals[VertexIndex];
 				FVector3f& TangentU = TangentUs[VertexIndex];
@@ -800,7 +808,7 @@ namespace UE::Chaos::ClothAsset
 				Normal= Normal.GetSafeNormal();
 				TangentU = TangentU.GetSafeNormal();
 				TangentU = TangentU - FVector3f::DotProduct(TangentU, Normal) * Normal;
-				TangentV = FVector3f::CrossProduct(TangentU, Normal);
+				TangentV = BasisDeterminantSigns[VertexIndex] * FVector3f::CrossProduct(Normal, TangentU);
 				TangentV = TangentV.GetSafeNormal();
 			}, NumVertices < 2000);
 	}
