Fix crash when first initializing Chaos Cloth with a non-simulating LOD.

This crash was caused by a regression from the optimization to make non-simulating cloth very cheap. We still need valid context data when initializing.

diff --git a/Engine/Plugins/ChaosClothAsset/Source/ChaosClothAssetEngine/Private/ChaosClothAsset/ClothSimulationContext.cpp b/Engine/Plugins/ChaosClothAsset/Source/ChaosClothAssetEngine/Private/ChaosClothAsset/ClothSimulationContext.cpp
--- a/Engine/Plugins/ChaosClothAsset/Source/ChaosClothAssetEngine/Private/ChaosClothAsset/ClothSimulationContext.cpp
+++ b/Engine/Plugins/ChaosClothAsset/Source/ChaosClothAssetEngine/Private/ChaosClothAsset/ClothSimulationContext.cpp
@@ -20,10 +20,11 @@ namespace UE::Chaos::ClothAsset
 		// Set the current LOD index
 		LodIndex = ClothComponent.GetPredictedLODLevel();
 		const UChaosClothAssetBase* const Asset = ClothComponent.GetAsset();
-		if (!Asset ||
+		if (!bIsInitialization &&
+			(!Asset ||
 			!Asset->GetResourceForRendering() ||
 			!Asset->GetResourceForRendering()->LODRenderData.IsValidIndex(LodIndex) ||
-			!Asset->GetResourceForRendering()->LODRenderData[LodIndex].HasClothData())
+			!Asset->GetResourceForRendering()->LODRenderData[LodIndex].HasClothData()))
 		{
 			// This Lod will not simulate, so there's no need to update the context.
 			return;
