Chaos Cloth Asset - Made the cloth collection construct a const version.

#rnx

diff --git a/Engine/Plugins/ChaosClothAsset/Source/ChaosClothAsset/Private/ChaosClothAsset/ClothCollection.cpp b/Engine/Plugins/ChaosClothAsset/Source/ChaosClothAsset/Private/ChaosClothAsset/ClothCollection.cpp
--- a/Engine/Plugins/ChaosClothAsset/Source/ChaosClothAsset/Private/ChaosClothAsset/ClothCollection.cpp
+++ b/Engine/Plugins/ChaosClothAsset/Source/ChaosClothAsset/Private/ChaosClothAsset/ClothCollection.cpp
@@ -183,11 +183,10 @@ namespace UE::Chaos::ClothAsset::Private
 
 namespace UE::Chaos::ClothAsset
 {
-	FClothCollection::FClothCollection(const TSharedRef<FManagedArrayCollection>& InManagedArrayCollection)
+	template<typename T UE_REQUIRES_DEFINITION(std::is_same_v<T, FManagedArrayCollection> || std::is_same_v<T, const FManagedArrayCollection>)>
+	FConstClothCollection::FConstClothCollection(const TSharedRef<T>& InManagedArrayCollection, TProtectedInit)
 		: ManagedArrayCollection(InManagedArrayCollection)
 	{
-		using namespace UE::Chaos::ClothAsset::Private;
-
 		// LODs Group
 		PhysicsAssetPathName = ManagedArrayCollection->FindAttribute<FString>(ClothCollectionAttribute::PhysicsAssetPathName, ClothCollectionGroup::Lods);
 		SkeletalMeshPathName = ManagedArrayCollection->FindAttribute<FString>(ClothCollectionAttribute::SkeletalMeshPathName, ClothCollectionGroup::Lods);
@@ -288,7 +287,12 @@ namespace UE::Chaos::ClothAsset
 		CustomResizingRegionType = ManagedArrayCollection->FindAttribute<int32>(ClothCollectionAttribute::CustomResizingRegionType, ClothCollectionGroup::CustomResizingRegions);
 	}
 
-	bool FClothCollection::IsValid(EClothCollectionExtendedSchemas OptionalSchemas) const
+	FConstClothCollection::FConstClothCollection(const TSharedRef<const FManagedArrayCollection>& InManagedArrayCollection)
+		: FConstClothCollection(InManagedArrayCollection, TProtectedInit{})
+	{
+	}
+
+	bool FConstClothCollection::IsValid(EClothCollectionExtendedSchemas OptionalSchemas) const
 	{
 		if (EnumHasAnyFlags(OptionalSchemas, EClothCollectionExtendedSchemas::CookedOnly))
 		{
@@ -417,6 +421,11 @@ namespace UE::Chaos::ClothAsset
 					CustomResizingRegionType));
 	}
 
+	FClothCollection::FClothCollection(const TSharedRef<FManagedArrayCollection>& InManagedArrayCollection)
+		: FConstClothCollection(InManagedArrayCollection, TProtectedInit{})
+	{
+	}
+
 	void FClothCollection::DefineSchema(EClothCollectionExtendedSchemas OptionalSchemas)
 	{
 		using namespace UE::Chaos::ClothAsset::Private;
@@ -438,150 +447,150 @@ namespace UE::Chaos::ClothAsset
 			// Special case: only define cooked schema.
 
 			// Sim Morph Targets Group
-			SimMorphTargetName = &ManagedArrayCollection->AddAttribute<FString>(ClothCollectionAttribute::SimMorphTargetName, ClothCollectionGroup::SimMorphTargets);
-			SimMorphTargetVerticesStart = &ManagedArrayCollection->AddAttribute<int32>(ClothCollectionAttribute::SimMorphTargetVerticesStart, ClothCollectionGroup::SimMorphTargets);
-			SimMorphTargetVerticesEnd = &ManagedArrayCollection->AddAttribute<int32>(ClothCollectionAttribute::SimMorphTargetVerticesEnd, ClothCollectionGroup::SimMorphTargets);
+			SimMorphTargetName = &GetManagedArrayCollection()->AddAttribute<FString>(ClothCollectionAttribute::SimMorphTargetName, ClothCollectionGroup::SimMorphTargets);
+			SimMorphTargetVerticesStart = &GetManagedArrayCollection()->AddAttribute<int32>(ClothCollectionAttribute::SimMorphTargetVerticesStart, ClothCollectionGroup::SimMorphTargets);
+			SimMorphTargetVerticesEnd = &GetManagedArrayCollection()->AddAttribute<int32>(ClothCollectionAttribute::SimMorphTargetVerticesEnd, ClothCollectionGroup::SimMorphTargets);
 
 			// Sim Morph Target Vertices Group
-			SimMorphTargetPositionDelta = &ManagedArrayCollection->AddAttribute<FVector3f>(ClothCollectionAttribute::SimMorphTargetPositionDelta, ClothCollectionGroup::SimMorphTargetVertices);
-			SimMorphTargetTangentZDelta = &ManagedArrayCollection->AddAttribute<FVector3f>(ClothCollectionAttribute::SimMorphTargetTangentZDelta, ClothCollectionGroup::SimMorphTargetVertices);
-			SimMorphTargetSimVertex3DIndex = &ManagedArrayCollection->AddAttribute<int32>(ClothCollectionAttribute::SimMorphTargetSimVertex3DIndex, ClothCollectionGroup::SimMorphTargetVertices);
+			SimMorphTargetPositionDelta = &GetManagedArrayCollection()->AddAttribute<FVector3f>(ClothCollectionAttribute::SimMorphTargetPositionDelta, ClothCollectionGroup::SimMorphTargetVertices);
+			SimMorphTargetTangentZDelta = &GetManagedArrayCollection()->AddAttribute<FVector3f>(ClothCollectionAttribute::SimMorphTargetTangentZDelta, ClothCollectionGroup::SimMorphTargetVertices);
+			SimMorphTargetSimVertex3DIndex = &GetManagedArrayCollection()->AddAttribute<int32>(ClothCollectionAttribute::SimMorphTargetSimVertex3DIndex, ClothCollectionGroup::SimMorphTargetVertices);
 			
 			// Sim Vertices 3D Group
-			PreResizedSimPosition3D = &ManagedArrayCollection->AddAttribute<FVector3f>(ClothCollectionAttribute::PreResizedSimPosition3D, ClothCollectionGroup::SimVertices3D);
+			PreResizedSimPosition3D = &GetManagedArrayCollection()->AddAttribute<FVector3f>(ClothCollectionAttribute::PreResizedSimPosition3D, ClothCollectionGroup::SimVertices3D);
 
 			return;
 		}
 
 		// LODs Group
-		PhysicsAssetPathName = &ManagedArrayCollection->AddAttribute<FString>(ClothCollectionAttribute::PhysicsAssetPathName, ClothCollectionGroup::Lods);
-		SkeletalMeshPathName = &ManagedArrayCollection->AddAttribute<FString>(ClothCollectionAttribute::SkeletalMeshPathName, ClothCollectionGroup::Lods);
+		PhysicsAssetPathName = &GetManagedArrayCollection()->AddAttribute<FString>(ClothCollectionAttribute::PhysicsAssetPathName, ClothCollectionGroup::Lods);
+		SkeletalMeshPathName = &GetManagedArrayCollection()->AddAttribute<FString>(ClothCollectionAttribute::SkeletalMeshPathName, ClothCollectionGroup::Lods);
 
 		// Solvers Group
 		if (EnumHasAnyFlags(OptionalSchemas, EClothCollectionExtendedSchemas::Solvers))
 		{
-			SolverGravity = &ManagedArrayCollection->AddAttribute<FVector3f>(ClothCollectionAttribute::SolverGravity, ClothCollectionGroup::Solvers);
-			SolverAirDamping = &ManagedArrayCollection->AddAttribute<float>(ClothCollectionAttribute::SolverAirDamping, ClothCollectionGroup::Solvers);
-			SolverSubSteps = &ManagedArrayCollection->AddAttribute<int32>(ClothCollectionAttribute::SolverSubSteps, ClothCollectionGroup::Solvers);
-			SolverTimeStep = &ManagedArrayCollection->AddAttribute<float>(ClothCollectionAttribute::SolverTimeStep, ClothCollectionGroup::Solvers);
+			SolverGravity = &GetManagedArrayCollection()->AddAttribute<FVector3f>(ClothCollectionAttribute::SolverGravity, ClothCollectionGroup::Solvers);
+			SolverAirDamping = &GetManagedArrayCollection()->AddAttribute<float>(ClothCollectionAttribute::SolverAirDamping, ClothCollectionGroup::Solvers);
+			SolverSubSteps = &GetManagedArrayCollection()->AddAttribute<int32>(ClothCollectionAttribute::SolverSubSteps, ClothCollectionGroup::Solvers);
+			SolverTimeStep = &GetManagedArrayCollection()->AddAttribute<float>(ClothCollectionAttribute::SolverTimeStep, ClothCollectionGroup::Solvers);
 		}
 
 		// Seams Group
-		SeamStitchStart = &ManagedArrayCollection->AddAttribute<int32>(ClothCollectionAttribute::SeamStitchStart, ClothCollectionGroup::Seams, SeamStitchesDependency);
-		SeamStitchEnd = &ManagedArrayCollection->AddAttribute<int32>(ClothCollectionAttribute::SeamStitchEnd, ClothCollectionGroup::Seams, SeamStitchesDependency);
+		SeamStitchStart = &GetManagedArrayCollection()->AddAttribute<int32>(ClothCollectionAttribute::SeamStitchStart, ClothCollectionGroup::Seams, SeamStitchesDependency);
+		SeamStitchEnd = &GetManagedArrayCollection()->AddAttribute<int32>(ClothCollectionAttribute::SeamStitchEnd, ClothCollectionGroup::Seams, SeamStitchesDependency);
 
 		// Seam Stitches Group
-		SeamStitch2DEndIndices = &ManagedArrayCollection->AddAttribute<FIntVector2>(ClothCollectionAttribute::SeamStitch2DEndIndices, ClothCollectionGroup::SeamStitches, SimVertices2DDependency);
-		SeamStitch3DIndex = &ManagedArrayCollection->AddAttribute<int32>(ClothCollectionAttribute::SeamStitch3DIndex, ClothCollectionGroup::SeamStitches, SimVertices3DDependency);
+		SeamStitch2DEndIndices = &GetManagedArrayCollection()->AddAttribute<FIntVector2>(ClothCollectionAttribute::SeamStitch2DEndIndices, ClothCollectionGroup::SeamStitches, SimVertices2DDependency);
+		SeamStitch3DIndex = &GetManagedArrayCollection()->AddAttribute<int32>(ClothCollectionAttribute::SeamStitch3DIndex, ClothCollectionGroup::SeamStitches, SimVertices3DDependency);
 
 		// Sim Patterns Group
-		SimVertices2DStart = &ManagedArrayCollection->AddAttribute<int32>(ClothCollectionAttribute::SimVertices2DStart, ClothCollectionGroup::SimPatterns, SimVertices2DDependency);
-		SimVertices2DEnd = &ManagedArrayCollection->AddAttribute<int32>(ClothCollectionAttribute::SimVertices2DEnd, ClothCollectionGroup::SimPatterns, SimVertices2DDependency);
-		SimFacesStart = &ManagedArrayCollection->AddAttribute<int32>(ClothCollectionAttribute::SimFacesStart, ClothCollectionGroup::SimPatterns, SimFacesDependency);
-		SimFacesEnd = &ManagedArrayCollection->AddAttribute<int32>(ClothCollectionAttribute::SimFacesEnd, ClothCollectionGroup::SimPatterns, SimFacesDependency);
-		SimPatternFabric = &ManagedArrayCollection->AddAttribute<int32>(ClothCollectionAttribute::SimPatternFabric, ClothCollectionGroup::SimPatterns, SimFabricsDependency);
+		SimVertices2DStart = &GetManagedArrayCollection()->AddAttribute<int32>(ClothCollectionAttribute::SimVertices2DStart, ClothCollectionGroup::SimPatterns, SimVertices2DDependency);
+		SimVertices2DEnd = &GetManagedArrayCollection()->AddAttribute<int32>(ClothCollectionAttribute::SimVertices2DEnd, ClothCollectionGroup::SimPatterns, SimVertices2DDependency);
+		SimFacesStart = &GetManagedArrayCollection()->AddAttribute<int32>(ClothCollectionAttribute::SimFacesStart, ClothCollectionGroup::SimPatterns, SimFacesDependency);
+		SimFacesEnd = &GetManagedArrayCollection()->AddAttribute<int32>(ClothCollectionAttribute::SimFacesEnd, ClothCollectionGroup::SimPatterns, SimFacesDependency);
+		SimPatternFabric = &GetManagedArrayCollection()->AddAttribute<int32>(ClothCollectionAttribute::SimPatternFabric, ClothCollectionGroup::SimPatterns, SimFabricsDependency);
 
 		// Render Patterns Group
-		RenderVerticesStart = &ManagedArrayCollection->AddAttribute<int32>(ClothCollectionAttribute::RenderVerticesStart, ClothCollectionGroup::RenderPatterns, RenderVerticesDependency);
-		RenderVerticesEnd = &ManagedArrayCollection->AddAttribute<int32>(ClothCollectionAttribute::RenderVerticesEnd, ClothCollectionGroup::RenderPatterns, RenderVerticesDependency);
-		RenderFacesStart = &ManagedArrayCollection->AddAttribute<int32>(ClothCollectionAttribute::RenderFacesStart, ClothCollectionGroup::RenderPatterns, RenderFacesDependency);
-		RenderFacesEnd = &ManagedArrayCollection->AddAttribute<int32>(ClothCollectionAttribute::RenderFacesEnd, ClothCollectionGroup::RenderPatterns, RenderFacesDependency);
-		RenderMaterialPathName = &ManagedArrayCollection->AddAttribute<FString>(ClothCollectionAttribute::RenderMaterialPathName, ClothCollectionGroup::RenderPatterns);
+		RenderVerticesStart = &GetManagedArrayCollection()->AddAttribute<int32>(ClothCollectionAttribute::RenderVerticesStart, ClothCollectionGroup::RenderPatterns, RenderVerticesDependency);
+		RenderVerticesEnd = &GetManagedArrayCollection()->AddAttribute<int32>(ClothCollectionAttribute::RenderVerticesEnd, ClothCollectionGroup::RenderPatterns, RenderVerticesDependency);
+		RenderFacesStart = &GetManagedArrayCollection()->AddAttribute<int32>(ClothCollectionAttribute::RenderFacesStart, ClothCollectionGroup::RenderPatterns, RenderFacesDependency);
+		RenderFacesEnd = &GetManagedArrayCollection()->AddAttribute<int32>(ClothCollectionAttribute::RenderFacesEnd, ClothCollectionGroup::RenderPatterns, RenderFacesDependency);
+		RenderMaterialPathName = &GetManagedArrayCollection()->AddAttribute<FString>(ClothCollectionAttribute::RenderMaterialPathName, ClothCollectionGroup::RenderPatterns);
 
 		//~ Fabric Group
 		if (EnumHasAnyFlags(OptionalSchemas, EClothCollectionExtendedSchemas::Fabrics))
 		{
-			FabricBendingStiffness = &ManagedArrayCollection->AddAttribute<FVector3f>(ClothCollectionAttribute::FabricBendingStiffness, ClothCollectionGroup::Fabrics);
-			FabricBucklingStiffness = &ManagedArrayCollection->AddAttribute<FVector3f>(ClothCollectionAttribute::FabricBucklingStiffness, ClothCollectionGroup::Fabrics);
-			FabricStretchStiffness = &ManagedArrayCollection->AddAttribute<FVector3f>(ClothCollectionAttribute::FabricStretchStiffness, ClothCollectionGroup::Fabrics);
-			FabricBucklingRatio = &ManagedArrayCollection->AddAttribute<float>(ClothCollectionAttribute::FabricBucklingRatio, ClothCollectionGroup::Fabrics);
-			FabricDensity = &ManagedArrayCollection->AddAttribute<float>(ClothCollectionAttribute::FabricDensity, ClothCollectionGroup::Fabrics);
-			FabricFriction = &ManagedArrayCollection->AddAttribute<float>(ClothCollectionAttribute::FabricFriction, ClothCollectionGroup::Fabrics);
-			FabricDamping = &ManagedArrayCollection->AddAttribute<float>(ClothCollectionAttribute::FabricDamping, ClothCollectionGroup::Fabrics);
-			FabricPressure = &ManagedArrayCollection->AddAttribute<float>(ClothCollectionAttribute::FabricPressure, ClothCollectionGroup::Fabrics);
-			FabricLayer = &ManagedArrayCollection->AddAttribute<int32>(ClothCollectionAttribute::FabricLayer, ClothCollectionGroup::Fabrics);
-			FabricCollisionThickness = &ManagedArrayCollection->AddAttribute<float>(ClothCollectionAttribute::FabricCollisionThickness, ClothCollectionGroup::Fabrics);
+			FabricBendingStiffness = &GetManagedArrayCollection()->AddAttribute<FVector3f>(ClothCollectionAttribute::FabricBendingStiffness, ClothCollectionGroup::Fabrics);
+			FabricBucklingStiffness = &GetManagedArrayCollection()->AddAttribute<FVector3f>(ClothCollectionAttribute::FabricBucklingStiffness, ClothCollectionGroup::Fabrics);
+			FabricStretchStiffness = &GetManagedArrayCollection()->AddAttribute<FVector3f>(ClothCollectionAttribute::FabricStretchStiffness, ClothCollectionGroup::Fabrics);
+			FabricBucklingRatio = &GetManagedArrayCollection()->AddAttribute<float>(ClothCollectionAttribute::FabricBucklingRatio, ClothCollectionGroup::Fabrics);
+			FabricDensity = &GetManagedArrayCollection()->AddAttribute<float>(ClothCollectionAttribute::FabricDensity, ClothCollectionGroup::Fabrics);
+			FabricFriction = &GetManagedArrayCollection()->AddAttribute<float>(ClothCollectionAttribute::FabricFriction, ClothCollectionGroup::Fabrics);
+			FabricDamping = &GetManagedArrayCollection()->AddAttribute<float>(ClothCollectionAttribute::FabricDamping, ClothCollectionGroup::Fabrics);
+			FabricPressure = &GetManagedArrayCollection()->AddAttribute<float>(ClothCollectionAttribute::FabricPressure, ClothCollectionGroup::Fabrics);
+			FabricLayer = &GetManagedArrayCollection()->AddAttribute<int32>(ClothCollectionAttribute::FabricLayer, ClothCollectionGroup::Fabrics);
+			FabricCollisionThickness = &GetManagedArrayCollection()->AddAttribute<float>(ClothCollectionAttribute::FabricCollisionThickness, ClothCollectionGroup::Fabrics);
 		}
 
 		// Sim Faces Group
-		SimIndices2D = &ManagedArrayCollection->AddAttribute<FIntVector3>(ClothCollectionAttribute::SimIndices2D, ClothCollectionGroup::SimFaces, SimVertices2DDependency);
-		SimIndices3D = &ManagedArrayCollection->AddAttribute<FIntVector3>(ClothCollectionAttribute::SimIndices3D, ClothCollectionGroup::SimFaces, SimVertices3DDependency);
+		SimIndices2D = &GetManagedArrayCollection()->AddAttribute<FIntVector3>(ClothCollectionAttribute::SimIndices2D, ClothCollectionGroup::SimFaces, SimVertices2DDependency);
+		SimIndices3D = &GetManagedArrayCollection()->AddAttribute<FIntVector3>(ClothCollectionAttribute::SimIndices3D, ClothCollectionGroup::SimFaces, SimVertices3DDependency);
 
 		// Sim Vertices 2D Group
-		SimPosition2D = &ManagedArrayCollection->AddAttribute<FVector2f>(ClothCollectionAttribute::SimPosition2D, ClothCollectionGroup::SimVertices2D);
-		SimVertex3DLookup = &ManagedArrayCollection->AddAttribute<int32>(ClothCollectionAttribute::SimVertex3DLookup, ClothCollectionGroup::SimVertices2D, SimVertices3DDependency);
+		SimPosition2D = &GetManagedArrayCollection()->AddAttribute<FVector2f>(ClothCollectionAttribute::SimPosition2D, ClothCollectionGroup::SimVertices2D);
+		SimVertex3DLookup = &GetManagedArrayCollection()->AddAttribute<int32>(ClothCollectionAttribute::SimVertex3DLookup, ClothCollectionGroup::SimVertices2D, SimVertices3DDependency);
 
 		// Sim Vertices 3D Group
-		SimPosition3D = &ManagedArrayCollection->AddAttribute<FVector3f>(ClothCollectionAttribute::SimPosition3D, ClothCollectionGroup::SimVertices3D);
-		SimNormal = &ManagedArrayCollection->AddAttribute<FVector3f>(ClothCollectionAttribute::SimNormal, ClothCollectionGroup::SimVertices3D);
-		SimBoneIndices = &ManagedArrayCollection->AddAttribute<TArray<int32>>(ClothCollectionAttribute::SimBoneIndices, ClothCollectionGroup::SimVertices3D);
-		SimBoneWeights = &ManagedArrayCollection->AddAttribute<TArray<float>>(ClothCollectionAttribute::SimBoneWeights, ClothCollectionGroup::SimVertices3D);
-		TetherKinematicIndex = &ManagedArrayCollection->AddAttribute<TArray<int32>>(ClothCollectionAttribute::TetherKinematicIndex, ClothCollectionGroup::SimVertices3D, SimVertices3DDependency);
-		TetherReferenceLength = &ManagedArrayCollection->AddAttribute<TArray<float>>(ClothCollectionAttribute::TetherReferenceLength, ClothCollectionGroup::SimVertices3D);
-		SimVertex2DLookup = &ManagedArrayCollection->AddAttribute<TArray<int32>>(ClothCollectionAttribute::SimVertex2DLookup, ClothCollectionGroup::SimVertices3D, SimVertices2DDependency);
-		SeamStitchLookup = &ManagedArrayCollection->AddAttribute<TArray<int32>>(ClothCollectionAttribute::SeamStitchLookup, ClothCollectionGroup::SimVertices3D, SeamStitchesDependency);
+		SimPosition3D = &GetManagedArrayCollection()->AddAttribute<FVector3f>(ClothCollectionAttribute::SimPosition3D, ClothCollectionGroup::SimVertices3D);
+		SimNormal = &GetManagedArrayCollection()->AddAttribute<FVector3f>(ClothCollectionAttribute::SimNormal, ClothCollectionGroup::SimVertices3D);
+		SimBoneIndices = &GetManagedArrayCollection()->AddAttribute<TArray<int32>>(ClothCollectionAttribute::SimBoneIndices, ClothCollectionGroup::SimVertices3D);
+		SimBoneWeights = &GetManagedArrayCollection()->AddAttribute<TArray<float>>(ClothCollectionAttribute::SimBoneWeights, ClothCollectionGroup::SimVertices3D);
+		TetherKinematicIndex = &GetManagedArrayCollection()->AddAttribute<TArray<int32>>(ClothCollectionAttribute::TetherKinematicIndex, ClothCollectionGroup::SimVertices3D, SimVertices3DDependency);
+		TetherReferenceLength = &GetManagedArrayCollection()->AddAttribute<TArray<float>>(ClothCollectionAttribute::TetherReferenceLength, ClothCollectionGroup::SimVertices3D);
+		SimVertex2DLookup = &GetManagedArrayCollection()->AddAttribute<TArray<int32>>(ClothCollectionAttribute::SimVertex2DLookup, ClothCollectionGroup::SimVertices3D, SimVertices2DDependency);
+		SeamStitchLookup = &GetManagedArrayCollection()->AddAttribute<TArray<int32>>(ClothCollectionAttribute::SeamStitchLookup, ClothCollectionGroup::SimVertices3D, SeamStitchesDependency);
 
 		if (EnumHasAnyFlags(OptionalSchemas, EClothCollectionExtendedSchemas::SimMorphTargets))
 		{
 			// Sim Morph Targets Group
-			SimMorphTargetName = &ManagedArrayCollection->AddAttribute<FString>(ClothCollectionAttribute::SimMorphTargetName, ClothCollectionGroup::SimMorphTargets);
-			SimMorphTargetVerticesStart = &ManagedArrayCollection->AddAttribute<int32>(ClothCollectionAttribute::SimMorphTargetVerticesStart, ClothCollectionGroup::SimMorphTargets, SimMorphTargetVerticesDependency);
-			SimMorphTargetVerticesEnd = &ManagedArrayCollection->AddAttribute<int32>(ClothCollectionAttribute::SimMorphTargetVerticesEnd, ClothCollectionGroup::SimMorphTargets, SimMorphTargetVerticesDependency);
+			SimMorphTargetName = &GetManagedArrayCollection()->AddAttribute<FString>(ClothCollectionAttribute::SimMorphTargetName, ClothCollectionGroup::SimMorphTargets);
+			SimMorphTargetVerticesStart = &GetManagedArrayCollection()->AddAttribute<int32>(ClothCollectionAttribute::SimMorphTargetVerticesStart, ClothCollectionGroup::SimMorphTargets, SimMorphTargetVerticesDependency);
+			SimMorphTargetVerticesEnd = &GetManagedArrayCollection()->AddAttribute<int32>(ClothCollectionAttribute::SimMorphTargetVerticesEnd, ClothCollectionGroup::SimMorphTargets, SimMorphTargetVerticesDependency);
 
 			// Sim Morph Target Vertices Group
-			SimMorphTargetPositionDelta = &ManagedArrayCollection->AddAttribute<FVector3f>(ClothCollectionAttribute::SimMorphTargetPositionDelta, ClothCollectionGroup::SimMorphTargetVertices);
-			SimMorphTargetTangentZDelta = &ManagedArrayCollection->AddAttribute<FVector3f>(ClothCollectionAttribute::SimMorphTargetTangentZDelta, ClothCollectionGroup::SimMorphTargetVertices);
-			SimMorphTargetSimVertex3DIndex = &ManagedArrayCollection->AddAttribute<int32>(ClothCollectionAttribute::SimMorphTargetSimVertex3DIndex, ClothCollectionGroup::SimMorphTargetVertices, SimVertices3DDependency);
+			SimMorphTargetPositionDelta = &GetManagedArrayCollection()->AddAttribute<FVector3f>(ClothCollectionAttribute::SimMorphTargetPositionDelta, ClothCollectionGroup::SimMorphTargetVertices);
+			SimMorphTargetTangentZDelta = &GetManagedArrayCollection()->AddAttribute<FVector3f>(ClothCollectionAttribute::SimMorphTargetTangentZDelta, ClothCollectionGroup::SimMorphTargetVertices);
+			SimMorphTargetSimVertex3DIndex = &GetManagedArrayCollection()->AddAttribute<int32>(ClothCollectionAttribute::SimMorphTargetSimVertex3DIndex, ClothCollectionGroup::SimMorphTargetVertices, SimVertices3DDependency);
 		}
 
 		// Render Faces Group
-		RenderIndices = &ManagedArrayCollection->AddAttribute<FIntVector3>(ClothCollectionAttribute::RenderIndices, ClothCollectionGroup::RenderFaces, RenderVerticesDependency);
+		RenderIndices = &GetManagedArrayCollection()->AddAttribute<FIntVector3>(ClothCollectionAttribute::RenderIndices, ClothCollectionGroup::RenderFaces, RenderVerticesDependency);
 
 		// Render Vertices Group
-		RenderPosition = &ManagedArrayCollection->AddAttribute<FVector3f>(ClothCollectionAttribute::RenderPosition, ClothCollectionGroup::RenderVertices);
-		RenderNormal = &ManagedArrayCollection->AddAttribute<FVector3f>(ClothCollectionAttribute::RenderNormal, ClothCollectionGroup::RenderVertices);
-		RenderTangentU = &ManagedArrayCollection->AddAttribute<FVector3f>(ClothCollectionAttribute::RenderTangentU, ClothCollectionGroup::RenderVertices);
-		RenderTangentV = &ManagedArrayCollection->AddAttribute<FVector3f>(ClothCollectionAttribute::RenderTangentV, ClothCollectionGroup::RenderVertices);
-		RenderUVs = &ManagedArrayCollection->AddAttribute<TArray<FVector2f>>(ClothCollectionAttribute::RenderUVs, ClothCollectionGroup::RenderVertices);
-		RenderColor = &ManagedArrayCollection->AddAttribute<FLinearColor>(ClothCollectionAttribute::RenderColor, ClothCollectionGroup::RenderVertices);
-		RenderBoneIndices = &ManagedArrayCollection->AddAttribute<TArray<int32>>(ClothCollectionAttribute::RenderBoneIndices, ClothCollectionGroup::RenderVertices);
-		RenderBoneWeights = &ManagedArrayCollection->AddAttribute<TArray<float>>(ClothCollectionAttribute::RenderBoneWeights, ClothCollectionGroup::RenderVertices);
+		RenderPosition = &GetManagedArrayCollection()->AddAttribute<FVector3f>(ClothCollectionAttribute::RenderPosition, ClothCollectionGroup::RenderVertices);
+		RenderNormal = &GetManagedArrayCollection()->AddAttribute<FVector3f>(ClothCollectionAttribute::RenderNormal, ClothCollectionGroup::RenderVertices);
+		RenderTangentU = &GetManagedArrayCollection()->AddAttribute<FVector3f>(ClothCollectionAttribute::RenderTangentU, ClothCollectionGroup::RenderVertices);
+		RenderTangentV = &GetManagedArrayCollection()->AddAttribute<FVector3f>(ClothCollectionAttribute::RenderTangentV, ClothCollectionGroup::RenderVertices);
+		RenderUVs = &GetManagedArrayCollection()->AddAttribute<TArray<FVector2f>>(ClothCollectionAttribute::RenderUVs, ClothCollectionGroup::RenderVertices);
+		RenderColor = &GetManagedArrayCollection()->AddAttribute<FLinearColor>(ClothCollectionAttribute::RenderColor, ClothCollectionGroup::RenderVertices);
+		RenderBoneIndices = &GetManagedArrayCollection()->AddAttribute<TArray<int32>>(ClothCollectionAttribute::RenderBoneIndices, ClothCollectionGroup::RenderVertices);
+		RenderBoneWeights = &GetManagedArrayCollection()->AddAttribute<TArray<float>>(ClothCollectionAttribute::RenderBoneWeights, ClothCollectionGroup::RenderVertices);
 
 		if (EnumHasAnyFlags(OptionalSchemas, EClothCollectionExtendedSchemas::RenderDeformer))
 		{
 			// Render Patterns Group
-			RenderDeformerNumInfluences = &ManagedArrayCollection->AddAttribute<int32>(ClothCollectionAttribute::RenderDeformerNumInfluences, ClothCollectionGroup::RenderPatterns);
+			RenderDeformerNumInfluences = &GetManagedArrayCollection()->AddAttribute<int32>(ClothCollectionAttribute::RenderDeformerNumInfluences, ClothCollectionGroup::RenderPatterns);
 
 			// Render Vertices Group
-			RenderDeformerPositionBaryCoordsAndDist = &ManagedArrayCollection->AddAttribute<TArray<FVector4f>>(ClothCollectionAttribute::RenderDeformerPositionBaryCoordsAndDist, ClothCollectionGroup::RenderVertices);
-			RenderDeformerNormalBaryCoordsAndDist = &ManagedArrayCollection->AddAttribute<TArray<FVector4f>>(ClothCollectionAttribute::RenderDeformerNormalBaryCoordsAndDist, ClothCollectionGroup::RenderVertices);
-			RenderDeformerTangentBaryCoordsAndDist = &ManagedArrayCollection->AddAttribute<TArray<FVector4f>>(ClothCollectionAttribute::RenderDeformerTangentBaryCoordsAndDist, ClothCollectionGroup::RenderVertices);
-			RenderDeformerSimIndices3D = &ManagedArrayCollection->AddAttribute<TArray<FIntVector3>>(ClothCollectionAttribute::RenderDeformerSimIndices3D, ClothCollectionGroup::RenderVertices, SimVertices3DDependency);
-			RenderDeformerWeight = &ManagedArrayCollection->AddAttribute<TArray<float>>(ClothCollectionAttribute::RenderDeformerWeight, ClothCollectionGroup::RenderVertices);
-			RenderDeformerSkinningBlend = &ManagedArrayCollection->AddAttribute<float>(ClothCollectionAttribute::RenderDeformerSkinningBlend, ClothCollectionGroup::RenderVertices);
+			RenderDeformerPositionBaryCoordsAndDist = &GetManagedArrayCollection()->AddAttribute<TArray<FVector4f>>(ClothCollectionAttribute::RenderDeformerPositionBaryCoordsAndDist, ClothCollectionGroup::RenderVertices);
+			RenderDeformerNormalBaryCoordsAndDist = &GetManagedArrayCollection()->AddAttribute<TArray<FVector4f>>(ClothCollectionAttribute::RenderDeformerNormalBaryCoordsAndDist, ClothCollectionGroup::RenderVertices);
+			RenderDeformerTangentBaryCoordsAndDist = &GetManagedArrayCollection()->AddAttribute<TArray<FVector4f>>(ClothCollectionAttribute::RenderDeformerTangentBaryCoordsAndDist, ClothCollectionGroup::RenderVertices);
+			RenderDeformerSimIndices3D = &GetManagedArrayCollection()->AddAttribute<TArray<FIntVector3>>(ClothCollectionAttribute::RenderDeformerSimIndices3D, ClothCollectionGroup::RenderVertices, SimVertices3DDependency);
+			RenderDeformerWeight = &GetManagedArrayCollection()->AddAttribute<TArray<float>>(ClothCollectionAttribute::RenderDeformerWeight, ClothCollectionGroup::RenderVertices);
+			RenderDeformerSkinningBlend = &GetManagedArrayCollection()->AddAttribute<float>(ClothCollectionAttribute::RenderDeformerSkinningBlend, ClothCollectionGroup::RenderVertices);
 		}
 		
 		if (EnumHasAnyFlags(OptionalSchemas, EClothCollectionExtendedSchemas::Resizing))
 		{
 			// Sim Vertices 3D Group
-			SimCustomResizingBlend = &ManagedArrayCollection->AddAttribute<float>(ClothCollectionAttribute::SimCustomResizingBlend, ClothCollectionGroup::SimVertices3D);
-			PreResizedSimPosition3D = &ManagedArrayCollection->AddAttribute<FVector3f>(ClothCollectionAttribute::PreResizedSimPosition3D, ClothCollectionGroup::SimVertices3D);
+			SimCustomResizingBlend = &GetManagedArrayCollection()->AddAttribute<float>(ClothCollectionAttribute::SimCustomResizingBlend, ClothCollectionGroup::SimVertices3D);
+			PreResizedSimPosition3D = &GetManagedArrayCollection()->AddAttribute<FVector3f>(ClothCollectionAttribute::PreResizedSimPosition3D, ClothCollectionGroup::SimVertices3D);
 
 			// Render Vertices Group
-			RenderCustomResizingBlend = &ManagedArrayCollection->AddAttribute<float>(ClothCollectionAttribute::RenderCustomResizingBlend, ClothCollectionGroup::RenderVertices);
+			RenderCustomResizingBlend = &GetManagedArrayCollection()->AddAttribute<float>(ClothCollectionAttribute::RenderCustomResizingBlend, ClothCollectionGroup::RenderVertices);
 
 			// Group Resizing Group
-			CustomResizingRegionSet = &ManagedArrayCollection->AddAttribute<FString>(ClothCollectionAttribute::CustomResizingRegionSet, ClothCollectionGroup::CustomResizingRegions);
-			CustomResizingRegionType = &ManagedArrayCollection->AddAttribute<int32>(ClothCollectionAttribute::CustomResizingRegionType, ClothCollectionGroup::CustomResizingRegions);
+			CustomResizingRegionSet = &GetManagedArrayCollection()->AddAttribute<FString>(ClothCollectionAttribute::CustomResizingRegionSet, ClothCollectionGroup::CustomResizingRegions);
+			CustomResizingRegionType = &GetManagedArrayCollection()->AddAttribute<int32>(ClothCollectionAttribute::CustomResizingRegionType, ClothCollectionGroup::CustomResizingRegions);
 		}
 	}
 
-	int32 FClothCollection::GetNumElements(const FName& GroupName) const
+	int32 FConstClothCollection::GetNumElements(const FName& GroupName) const
 	{
 		return ManagedArrayCollection->NumElements(GroupName);
 	}
 
-	int32 FClothCollection::GetNumElements(const TManagedArray<int32>* StartArray, const TManagedArray<int32>* EndArray, int32 ArrayIndex) const
+	int32 FConstClothCollection::GetNumElements(const TManagedArray<int32>* StartArray, const TManagedArray<int32>* EndArray, int32 ArrayIndex) const
 	{
 		if (StartArray && StartArray->GetConstArray().IsValidIndex(ArrayIndex) &&
 			EndArray && EndArray->GetConstArray().IsValidIndex(ArrayIndex))
@@ -601,17 +610,17 @@ namespace UE::Chaos::ClothAsset
 	{
 		check(InNumElements >= 0);
 		
-		const int32 NumElements = ManagedArrayCollection->NumElements(GroupName);
+		const int32 NumElements = GetManagedArrayCollection()->NumElements(GroupName);
 
 		if (const int32 Delta = InNumElements - NumElements)
 		{
 			if (Delta > 0)
 			{
-				ManagedArrayCollection->AddElements(Delta, GroupName);
+				GetManagedArrayCollection()->AddElements(Delta, GroupName);
 			}
 			else
 			{
-				ManagedArrayCollection->RemoveElements(GroupName, -Delta, InNumElements);
+				GetManagedArrayCollection()->RemoveElements(GroupName, -Delta, InNumElements);
 			}
 		}
 	}
@@ -648,7 +657,7 @@ namespace UE::Chaos::ClothAsset
 
 				// Grow the array
 				const int32 Position = ComputeEnd() + 1;
-				ManagedArrayCollection->InsertElements(Delta, Position, GroupName);
+				GetManagedArrayCollection()->InsertElements(Delta, Position, GroupName);
 
 				// Update Start/End
 				if (!NumElements)
@@ -661,7 +670,7 @@ namespace UE::Chaos::ClothAsset
 			{
 				// Shrink the array
 				const int32 Position = Start + InNumElements;
-				ManagedArrayCollection->RemoveElements(GroupName, -Delta, Position);
+				GetManagedArrayCollection()->RemoveElements(GroupName, -Delta, Position);
 
 				// Update Start/End
 				if (InNumElements)
@@ -679,7 +688,7 @@ namespace UE::Chaos::ClothAsset
 
 	void FClothCollection::RemoveElements(const FName& Group, const TArray<int32>& SortedDeletionList)
 	{
-		ManagedArrayCollection->RemoveElements(Group, SortedDeletionList);
+		GetManagedArrayCollection()->RemoveElements(Group, SortedDeletionList);
 	}
 
 	void FClothCollection::RemoveElements(const FName& GroupName, const TArray<int32>& SortedDeletionList, TManagedArray<int32>* StartArray, TManagedArray<int32>* EndArray, int32 ArrayIndex)
@@ -703,7 +712,7 @@ namespace UE::Chaos::ClothAsset
 		check(SortedDeletionList.Last() <= End);
 		check(OrigNumElements >= SortedDeletionList.Num());
 
-		ManagedArrayCollection->RemoveElements(GroupName, SortedDeletionList);
+		GetManagedArrayCollection()->RemoveElements(GroupName, SortedDeletionList);
 
 		if (SortedDeletionList.Num() == OrigNumElements)
 		{
@@ -720,7 +729,7 @@ namespace UE::Chaos::ClothAsset
 		}
 	}
 
-	int32 FClothCollection::GetElementsOffset(const TManagedArray<int32>* StartArray, int32 BaseElementIndex, int32 ElementIndex)
+	/*static*/ int32 FConstClothCollection::GetElementsOffset(const TManagedArray<int32>* StartArray, int32 BaseElementIndex, int32 ElementIndex)
 	{
 		while ((*StartArray)[BaseElementIndex] == INDEX_NONE && BaseElementIndex < ElementIndex)
 		{
@@ -729,7 +738,7 @@ namespace UE::Chaos::ClothAsset
 		return (*StartArray)[ElementIndex] - (*StartArray)[BaseElementIndex];
 	}
 
-	int32 FClothCollection::GetArrayIndexForContainedElement(
+	/*static*/ int32 FConstClothCollection::GetArrayIndexForContainedElement(
 		const TManagedArray<int32>* StartArray,
 		const TManagedArray<int32>* EndArray,
 		int32 ElementIndex)
@@ -744,7 +753,7 @@ namespace UE::Chaos::ClothAsset
 		return INDEX_NONE;
 	}
 
-	int32 FClothCollection::GetNumSubElements(
+	/*static*/ int32 FConstClothCollection::GetNumSubElements(
 		const TManagedArray<int32>* StartArray,
 		const TManagedArray<int32>* EndArray,
 		const TManagedArray<int32>* StartSubArray,
@@ -763,7 +772,7 @@ namespace UE::Chaos::ClothAsset
 	}
 
 	template<bool bStart, bool bEnd>
-	TTuple<int32, int32> FClothCollection::GetSubElementsStartEnd(
+	/*static*/ TTuple<int32, int32> FConstClothCollection::GetSubElementsStartEnd(
 		const TManagedArray<int32>* StartArray,
 		const TManagedArray<int32>* EndArray,
 		const TManagedArray<int32>* StartSubArray,
@@ -800,12 +809,12 @@ namespace UE::Chaos::ClothAsset
 		}
 		return TTuple<int32, int32>(Start, End);
 	}
-	template CHAOSCLOTHASSET_API TTuple<int32, int32> FClothCollection::GetSubElementsStartEnd<true, false>(const TManagedArray<int32>* StartArray, const TManagedArray<int32>* EndArray, const TManagedArray<int32>* StartSubArray, const TManagedArray<int32>* EndSubArray, int32 ArrayIndex);
-	template CHAOSCLOTHASSET_API TTuple<int32, int32> FClothCollection::GetSubElementsStartEnd<false, true>(const TManagedArray<int32>* StartArray, const TManagedArray<int32>* EndArray, const TManagedArray<int32>* StartSubArray, const TManagedArray<int32>* EndSubArray, int32 ArrayIndex);
-	template CHAOSCLOTHASSET_API TTuple<int32, int32> FClothCollection::GetSubElementsStartEnd<true, true>(const TManagedArray<int32>* StartArray, const TManagedArray<int32>* EndArray, const TManagedArray<int32>* StartSubArray, const TManagedArray<int32>* EndSubArray, int32 ArrayIndex);
+	template CHAOSCLOTHASSET_API TTuple<int32, int32> FConstClothCollection::GetSubElementsStartEnd<true, false>(const TManagedArray<int32>* StartArray, const TManagedArray<int32>* EndArray, const TManagedArray<int32>* StartSubArray, const TManagedArray<int32>* EndSubArray, int32 ArrayIndex);
+	template CHAOSCLOTHASSET_API TTuple<int32, int32> FConstClothCollection::GetSubElementsStartEnd<false, true>(const TManagedArray<int32>* StartArray, const TManagedArray<int32>* EndArray, const TManagedArray<int32>* StartSubArray, const TManagedArray<int32>* EndSubArray, int32 ArrayIndex);
+	template CHAOSCLOTHASSET_API TTuple<int32, int32> FConstClothCollection::GetSubElementsStartEnd<true, true>(const TManagedArray<int32>* StartArray, const TManagedArray<int32>* EndArray, const TManagedArray<int32>* StartSubArray, const TManagedArray<int32>* EndSubArray, int32 ArrayIndex);
 
 	template<typename T UE_REQUIRES_DEFINITION(TIsUserAttributeType<T>::Value)>
-	TArray<FName> FClothCollection::GetUserDefinedAttributeNames(const FName& GroupName) const
+	TArray<FName> FConstClothCollection::GetUserDefinedAttributeNames(const FName& GroupName) const
 	{
 		using namespace UE::Chaos::ClothAsset::Private;
 
@@ -826,11 +835,11 @@ namespace UE::Chaos::ClothAsset
 
 		return UserDefinedAttributeNames;
 	}
-	template CHAOSCLOTHASSET_API TArray<FName> FClothCollection::GetUserDefinedAttributeNames<bool>(const FName& GroupName) const;
-	template CHAOSCLOTHASSET_API TArray<FName> FClothCollection::GetUserDefinedAttributeNames<int32>(const FName& GroupName) const;
-	template CHAOSCLOTHASSET_API TArray<FName> FClothCollection::GetUserDefinedAttributeNames<float>(const FName& GroupName) const;
-	template CHAOSCLOTHASSET_API TArray<FName> FClothCollection::GetUserDefinedAttributeNames<FVector3f>(const FName& GroupName) const;
-	template CHAOSCLOTHASSET_API TArray<FName> FClothCollection::GetUserDefinedAttributeNames<TArray<int32>>(const FName& GroupName) const;
+	template CHAOSCLOTHASSET_API TArray<FName> FConstClothCollection::GetUserDefinedAttributeNames<bool>(const FName& GroupName) const;
+	template CHAOSCLOTHASSET_API TArray<FName> FConstClothCollection::GetUserDefinedAttributeNames<int32>(const FName& GroupName) const;
+	template CHAOSCLOTHASSET_API TArray<FName> FConstClothCollection::GetUserDefinedAttributeNames<float>(const FName& GroupName) const;
+	template CHAOSCLOTHASSET_API TArray<FName> FConstClothCollection::GetUserDefinedAttributeNames<FVector3f>(const FName& GroupName) const;
+	template CHAOSCLOTHASSET_API TArray<FName> FConstClothCollection::GetUserDefinedAttributeNames<TArray<int32>>(const FName& GroupName) const;
 
 	template<typename T UE_REQUIRES_DEFINITION(TIsUserAttributeType<T>::Value)>
 	TManagedArray<T>* FClothCollection::FindOrAddUserDefinedAttribute(const FName& Name, const FName& GroupName, const FName& GroupDependency)
@@ -841,7 +850,7 @@ namespace UE::Chaos::ClothAsset
 			// Only allow adding to known groups. Do not allow adding a name reserved as a fixed attribute.
 			if (!FixedAttributeNames->Contains(Name))
 			{
-				return ManagedArrayCollection->FindOrAddAttributeTyped<T>(Name, GroupName, FManagedArrayCollection::FConstructionParameters(GroupDependency));
+				return GetManagedArrayCollection()->FindOrAddAttributeTyped<T>(Name, GroupName, FManagedArrayCollection::FConstructionParameters(GroupDependency));
 			}
 		}
 		if (const TArray<FName>* const UserAccessibleAttributeNames = UserAccessibleAttributeNamesMap.Find(GroupName))
@@ -849,7 +858,7 @@ namespace UE::Chaos::ClothAsset
 			// Paintable attribute can only be found, not created (they are part of the schema)
 			if (UserAccessibleAttributeNames->Contains(Name))
 			{
-				return ManagedArrayCollection->FindAttributeTyped<T>(Name, GroupName);
+				return GetManagedArrayCollection()->FindAttributeTyped<T>(Name, GroupName);
 			}
 		}
 		return nullptr;
@@ -868,13 +877,13 @@ namespace UE::Chaos::ClothAsset
 			// User defined attributes are only allowed in known group names. Do not allow removing fixed attributes through this method.
 			if (!FixedAttributeNames->Contains(Name))
 			{
-				ManagedArrayCollection->RemoveAttribute(Name, GroupName);
+				GetManagedArrayCollection()->RemoveAttribute(Name, GroupName);
 			}
 		}
 	}
 
 	template<typename T UE_REQUIRES_DEFINITION(TIsUserAttributeType<T>::Value)>
-	bool FClothCollection::HasUserDefinedAttribute(const FName& Name, const FName& GroupName) const
+	bool FConstClothCollection::HasUserDefinedAttribute(const FName& Name, const FName& GroupName) const
 	{
 		using namespace UE::Chaos::ClothAsset::Private;
 		if (const TArray<FName>* const FixedAttributeNames = FixedAttributeNamesMap.Find(GroupName))
@@ -895,14 +904,14 @@ namespace UE::Chaos::ClothAsset
 		}
 		return false;
 	}
-	template CHAOSCLOTHASSET_API bool FClothCollection::HasUserDefinedAttribute<bool>(const FName& Name, const FName& GroupName) const;
-	template CHAOSCLOTHASSET_API bool FClothCollection::HasUserDefinedAttribute<int32>(const FName& Name, const FName& GroupName) const;
-	template CHAOSCLOTHASSET_API bool FClothCollection::HasUserDefinedAttribute<float>(const FName& Name, const FName& GroupName) const;
-	template CHAOSCLOTHASSET_API bool FClothCollection::HasUserDefinedAttribute<FVector3f>(const FName& Name, const FName& GroupName) const;
-	template CHAOSCLOTHASSET_API bool FClothCollection::HasUserDefinedAttribute<TArray<int32>>(const FName& Name, const FName& GroupName) const;
+	template CHAOSCLOTHASSET_API bool FConstClothCollection::HasUserDefinedAttribute<bool>(const FName& Name, const FName& GroupName) const;
+	template CHAOSCLOTHASSET_API bool FConstClothCollection::HasUserDefinedAttribute<int32>(const FName& Name, const FName& GroupName) const;
+	template CHAOSCLOTHASSET_API bool FConstClothCollection::HasUserDefinedAttribute<float>(const FName& Name, const FName& GroupName) const;
+	template CHAOSCLOTHASSET_API bool FConstClothCollection::HasUserDefinedAttribute<FVector3f>(const FName& Name, const FName& GroupName) const;
+	template CHAOSCLOTHASSET_API bool FConstClothCollection::HasUserDefinedAttribute<TArray<int32>>(const FName& Name, const FName& GroupName) const;
 
 	template<typename T UE_REQUIRES_DEFINITION(TIsUserAttributeType<T>::Value)>
-	const TManagedArray<T>* FClothCollection::GetUserDefinedAttribute(const FName& Name, const FName& GroupName) const
+	const TManagedArray<T>* FConstClothCollection::GetUserDefinedAttribute(const FName& Name, const FName& GroupName) const
 	{
 		using namespace UE::Chaos::ClothAsset::Private;
 		if (const TArray<FName>* const FixedAttributeNames = FixedAttributeNamesMap.Find(GroupName))
@@ -923,11 +932,11 @@ namespace UE::Chaos::ClothAsset
 		}
 		return nullptr;
 	}
-	template CHAOSCLOTHASSET_API const TManagedArray<bool>* FClothCollection::GetUserDefinedAttribute<bool>(const FName& Name, const FName& GroupName) const;
-	template CHAOSCLOTHASSET_API const TManagedArray<int32>* FClothCollection::GetUserDefinedAttribute<int32>(const FName& Name, const FName& GroupName) const;
-	template CHAOSCLOTHASSET_API const TManagedArray<float>* FClothCollection::GetUserDefinedAttribute<float>(const FName& Name, const FName& GroupName) const;
-	template CHAOSCLOTHASSET_API const TManagedArray<FVector3f>* FClothCollection::GetUserDefinedAttribute<FVector3f>(const FName& Name, const FName& GroupName) const;
-	template CHAOSCLOTHASSET_API const TManagedArray<TArray<int32>>* FClothCollection::GetUserDefinedAttribute<TArray<int32>>(const FName& Name, const FName& GroupName) const;
+	template CHAOSCLOTHASSET_API const TManagedArray<bool>* FConstClothCollection::GetUserDefinedAttribute<bool>(const FName& Name, const FName& GroupName) const;
+	template CHAOSCLOTHASSET_API const TManagedArray<int32>* FConstClothCollection::GetUserDefinedAttribute<int32>(const FName& Name, const FName& GroupName) const;
+	template CHAOSCLOTHASSET_API const TManagedArray<float>* FConstClothCollection::GetUserDefinedAttribute<float>(const FName& Name, const FName& GroupName) const;
+	template CHAOSCLOTHASSET_API const TManagedArray<FVector3f>* FConstClothCollection::GetUserDefinedAttribute<FVector3f>(const FName& Name, const FName& GroupName) const;
+	template CHAOSCLOTHASSET_API const TManagedArray<TArray<int32>>* FConstClothCollection::GetUserDefinedAttribute<TArray<int32>>(const FName& Name, const FName& GroupName) const;
 
 	template<typename T UE_REQUIRES_DEFINITION(TIsUserAttributeType<T>::Value)>
 	TManagedArray<T>* FClothCollection::GetUserDefinedAttribute(const FName& Name, const FName& GroupName)
@@ -938,7 +947,7 @@ namespace UE::Chaos::ClothAsset
 			// User defined attributes are only allowed in known group names.
 			if (!FixedAttributeNames->Contains(Name))
 			{
-				return ManagedArrayCollection->FindAttributeTyped<T>(Name, GroupName);
+				return GetManagedArrayCollection()->FindAttributeTyped<T>(Name, GroupName);
 			}
 		}
 		if (const TArray<FName>* const UserAccessibleAttributeNames = UserAccessibleAttributeNamesMap.Find(GroupName))
@@ -946,7 +955,7 @@ namespace UE::Chaos::ClothAsset
 			// Paintable attribute can only be found, not created (they are part of the schema)
 			if (UserAccessibleAttributeNames->Contains(Name))
 			{
-				return ManagedArrayCollection->FindAttributeTyped<T>(Name, GroupName);
+				return GetManagedArrayCollection()->FindAttributeTyped<T>(Name, GroupName);
 			}
 		}
 		return nullptr;
@@ -957,7 +966,7 @@ namespace UE::Chaos::ClothAsset
 	template CHAOSCLOTHASSET_API TManagedArray<FVector3f>* FClothCollection::GetUserDefinedAttribute<FVector3f>(const FName& Name, const FName& GroupName);
 	template CHAOSCLOTHASSET_API TManagedArray<TArray<int32>>* FClothCollection::GetUserDefinedAttribute<TArray<int32>>(const FName& Name, const FName& GroupName);
 
-	TArray<FName> FClothCollection::GetValidClothCollectionGroupName()
+	/*static*/ TArray<FName> FConstClothCollection::GetValidClothCollectionGroupName()
 	{
 		using namespace UE::Chaos::ClothAsset::Private;
 		TArray<FName> ValidClothCollectionGroupName;
@@ -965,13 +974,13 @@ namespace UE::Chaos::ClothAsset
 		return ValidClothCollectionGroupName;
 	}
 
-	bool FClothCollection::IsValidClothCollectionGroupName(const FName& GroupName)
+	/*static*/ bool FConstClothCollection::IsValidClothCollectionGroupName(const FName& GroupName)
 	{
 		using namespace UE::Chaos::ClothAsset::Private;
 		return FixedAttributeNamesMap.Contains(GroupName);
 	}
 
-	bool FClothCollection::IsValidUserDefinedAttributeName(const FName& Name, const FName& GroupName)
+	bool FConstClothCollection::IsValidUserDefinedAttributeName(const FName& Name, const FName& GroupName) const
 	{
 		using namespace UE::Chaos::ClothAsset::Private;
 		if (const TArray<FName>* const FixedAttributeNames = FixedAttributeNamesMap.Find(GroupName))
diff --git a/Engine/Plugins/ChaosClothAsset/Source/ChaosClothAsset/Private/ChaosClothAsset/ClothCollection.h b/Engine/Plugins/ChaosClothAsset/Source/ChaosClothAsset/Private/ChaosClothAsset/ClothCollection.h
--- a/Engine/Plugins/ChaosClothAsset/Source/ChaosClothAsset/Private/ChaosClothAsset/ClothCollection.h
+++ b/Engine/Plugins/ChaosClothAsset/Source/ChaosClothAsset/Private/ChaosClothAsset/ClothCollection.h
@@ -16,44 +16,29 @@ namespace UE::Chaos::ClothAsset
 	/**
 	 * Cloth collection facade data.
 	 */
-	class FClothCollection final
+	class FConstClothCollection
 	{
 	public:
 		static constexpr int8 MaxNumBoneInfluences = 12; // This should be <= MAX_TOTAL_INFLUENCES defined in GPUSkinPublicDefs.h 
 		static constexpr int8 MaxNumTetherAttachments = 4; // This should be <= FClothTetherDataPrivate::MaxNumAttachments
 
-		explicit FClothCollection(const TSharedRef<FManagedArrayCollection>& InManagedArrayCollection);
-		~FClothCollection() = default;
+		explicit FConstClothCollection(const TSharedRef<const FManagedArrayCollection>& InManagedArrayCollection);
+		~FConstClothCollection() = default;
 
-		FClothCollection(const FClothCollection&) = delete;
-		FClothCollection& operator=(const FClothCollection&) = delete;
-		FClothCollection(FClothCollection&&) = delete;
-		FClothCollection& operator=(FClothCollection&&) = delete;
+		FConstClothCollection(const FConstClothCollection&) = delete;
+		FConstClothCollection& operator=(const FConstClothCollection&) = delete;
+		FConstClothCollection(FConstClothCollection&&) = delete;
+		FConstClothCollection& operator=(FConstClothCollection&&) = delete;
 
 		/** Return whether the underlying collection is a valid cloth collection. */
 		bool IsValid(EClothCollectionExtendedSchemas OptionalSchemas = EClothCollectionExtendedSchemas::None) const;
 
-		/** Make the underlying collection a cloth collection. */
-		void DefineSchema(EClothCollectionExtendedSchemas OptionalSchemas = EClothCollectionExtendedSchemas::None);
-
 		/** Get the number of elements of a group. */
 		int32 GetNumElements(const FName& GroupName) const;
 
 		/** Get the number of elements of one of the sub groups that have start/end indices. */
 		int32 GetNumElements(const TManagedArray<int32>* StartArray, const TManagedArray<int32>* EndArray, int32 Index) const;
 
-		/** Set the number of elements of a group. */
-		void SetNumElements(int32 InNumElements, const FName& GroupName);
-
-		/** Set the number of elements to one of the sub groups while maintaining the correct order of the data, and return the first index of the range. */
-		int32 SetNumElements(int32 InNumElements, const FName& GroupName, TManagedArray<int32>* StartArray, TManagedArray<int32>* EndArray, int32 Index);
-
-		/** Remove Elements */
-		void RemoveElements(const FName& Group, const TArray<int32>& SortedDeletionList);
-
-		/** Remove Elements. SortedDeletionList should be global indices. */
-		void RemoveElements(const FName& GroupName, const TArray<int32>& SortedDeletionList, TManagedArray<int32>* StartArray, TManagedArray<int32>* EndArray, int32 Index);
-
 		template<typename T>
 		static inline TConstArrayView<T> GetElements(const TManagedArray<T>* ElementArray, const TManagedArray<int32>* StartArray, const TManagedArray<int32>* EndArray, int32 ArrayIndex);
 
@@ -134,29 +119,15 @@ namespace UE::Chaos::ClothAsset
 		template<typename T UE_REQUIRES(TIsUserAttributeType<T>::Value)>
 		TArray<FName> GetUserDefinedAttributeNames(const FName& GroupName) const;
 
-		template<typename T UE_REQUIRES(TIsUserAttributeType<T>::Value)>
-		TManagedArray<T>* FindOrAddUserDefinedAttribute(const FName& Name, const FName& GroupName, const FName& GroupDependency = NAME_None);
-
-		template<typename T UE_REQUIRES(TIsUserAttributeType<T>::Value)>
-		void AddUserDefinedAttribute(const FName& Name, const FName& GroupName, const FName& GroupDependency = NAME_None)
-		{
-			FindOrAddUserDefinedAttribute<T>(Name, GroupName, GroupDependency);
-		}
-
-		void RemoveUserDefinedAttribute(const FName& Name, const FName& GroupName);
-
 		template<typename T UE_REQUIRES(TIsUserAttributeType<T>::Value)>
 		bool HasUserDefinedAttribute(const FName& Name, const FName& GroupName) const;
 
 		template<typename T UE_REQUIRES(TIsUserAttributeType<T>::Value)>
 		const TManagedArray<T>* GetUserDefinedAttribute(const FName& Name, const FName& GroupName) const;
 
-		template<typename T UE_REQUIRES(TIsUserAttributeType<T>::Value)>
-		TManagedArray<T>* GetUserDefinedAttribute(const FName& Name, const FName& GroupName);
-
 		static TArray<FName> GetValidClothCollectionGroupName();
 		static bool IsValidClothCollectionGroupName(const FName& GroupName);
-		bool IsValidUserDefinedAttributeName(const FName& Name, const FName& GroupName);
+		bool IsValidUserDefinedAttributeName(const FName& Name, const FName& GroupName) const;
 
 		//~ LODs Group (There should be only one LOD per ClothCollection)
 		const TManagedArray<FString>* GetPhysicsAssetPathName() const { return PhysicsAssetPathName; }
@@ -257,211 +228,266 @@ namespace UE::Chaos::ClothAsset
 		const TManagedArray<FString>* GetCustomResizingRegionSet() const { return CustomResizingRegionSet; }
 		const TManagedArray<int32>* GetCustomResizingRegionType() const { return CustomResizingRegionType; }
 
-		//~ LODs Group (There should be only one LOD per ClothCollection)
-		TManagedArray<FString>* GetPhysicsAssetPathName(){ return PhysicsAssetPathName; }
-		TManagedArray<FString>* GetSkeletalMeshPathName() { return SkeletalMeshPathName; }
-		
-		//~ Solver Group
-		TManagedArray<FVector3f>* GetSolverGravity() { return SolverGravity; }
-		TManagedArray<float>* GetSolverAirDamping() { return SolverAirDamping; }
-		TManagedArray<int32>* GetSolverSubSteps() { return SolverSubSteps; }
-		TManagedArray<float>* GetSolverTimeStep() { return SolverTimeStep; }
-		
-		//~ Fabric Group
-		TManagedArray<FVector3f>* GetFabricBendingStiffness() { return FabricBendingStiffness; }
-		TManagedArray<FVector3f>* GetFabricBucklingStiffness() { return FabricBucklingStiffness; }
-		TManagedArray<FVector3f>* GetFabricStretchStiffness() { return FabricStretchStiffness; }
-		TManagedArray<float>* GetFabricBucklingRatio() { return FabricBucklingRatio; }
-		TManagedArray<float>* GetFabricDensity() { return FabricDensity; }
-		TManagedArray<float>* GetFabricFriction() { return FabricFriction; }
-		TManagedArray<float>* GetFabricDamping() { return FabricDamping; }
-		TManagedArray<float>* GetFabricPressure() { return FabricPressure; }
-		TManagedArray<int32>* GetFabricLayer() { return FabricLayer; }
-		TManagedArray<float>* GetFabricCollisionThickness() { return FabricCollisionThickness; }
-		
+	protected:
+		struct TProtectedInit {};
+
+		template<typename T UE_REQUIRES(std::is_same_v<T, FManagedArrayCollection> || std::is_same_v<T, const FManagedArrayCollection>)>
+		FConstClothCollection(const TSharedRef<T>& InManagedArrayCollection, TProtectedInit);
+
+		//~ Cloth collection
+		TSharedRef<const FManagedArrayCollection> ManagedArrayCollection;
+
+		//~ LODs Group
+		const TManagedArray<FString>* PhysicsAssetPathName;
+		const TManagedArray<FString>* SkeletalMeshPathName;
+
+		//~ Solvers group
+		const TManagedArray<FVector3f>* SolverGravity;
+		const TManagedArray<float>* SolverAirDamping;
+		const TManagedArray<int32>* SolverSubSteps;
+		const TManagedArray<float>* SolverTimeStep;
+
+		//~ Fabrics Group
+		const TManagedArray<FVector3f>* FabricBendingStiffness;
+		const TManagedArray<FVector3f>* FabricBucklingStiffness;
+		const TManagedArray<FVector3f>* FabricStretchStiffness;
+		const TManagedArray<float>* FabricBucklingRatio;
+		const TManagedArray<float>* FabricDensity;
+		const TManagedArray<float>* FabricFriction;
+		const TManagedArray<float>* FabricDamping;
+		const TManagedArray<float>* FabricPressure;
+		const TManagedArray<int32>* FabricLayer;
+		const TManagedArray<float>* FabricCollisionThickness;
+
 		//~ Seam Group
-		TManagedArray<int32>* GetSeamStitchStart() { return SeamStitchStart; }
-		TManagedArray<int32>* GetSeamStitchEnd() { return SeamStitchEnd; }
+		const TManagedArray<int32>* SeamStitchStart;
+		const TManagedArray<int32>* SeamStitchEnd;
 
 		//~ Seam Stitches Group
-		TManagedArray<FIntVector2>* GetSeamStitch2DEndIndices() { return SeamStitch2DEndIndices; }
-		TManagedArray<int32>* GetSeamStitch3DIndex() { return SeamStitch3DIndex; }
+		const TManagedArray<FIntVector2>* SeamStitch2DEndIndices;  // Stitched 2D vertex indices pair
+		const TManagedArray<int32>* SeamStitch3DIndex;  // Corresponding stitched 3D vertex
 
 		//~ Sim Patterns Group
-		TManagedArray<int32>* GetSimVertices2DStart() { return SimVertices2DStart; }
-		TManagedArray<int32>* GetSimVertices2DEnd() { return SimVertices2DEnd; }
-		TManagedArray<int32>* GetSimFacesStart() { return SimFacesStart; }
-		TManagedArray<int32>* GetSimFacesEnd() { return SimFacesEnd; }
-		TManagedArray<int32>* GetSimPatternFabric() { return SimPatternFabric; }
+		const TManagedArray<int32>* SimVertices2DStart;
+		const TManagedArray<int32>* SimVertices2DEnd;
+		const TManagedArray<int32>* SimFacesStart;
+		const TManagedArray<int32>* SimFacesEnd;
+		const TManagedArray<int32>* SimPatternFabric;
 
 		//~ Render Patterns Group
-		TManagedArray<int32>* GetRenderVerticesStart() { return RenderVerticesStart; }
-		TManagedArray<int32>* GetRenderVerticesEnd() { return RenderVerticesEnd; }
-		TManagedArray<int32>* GetRenderFacesStart() { return RenderFacesStart; }
-		TManagedArray<int32>* GetRenderFacesEnd() { return RenderFacesEnd; }
-		TManagedArray<int32>* GetRenderDeformerNumInfluences() { return RenderDeformerNumInfluences; }
-		TManagedArray<FString>* GetRenderMaterialPathName() { return RenderMaterialPathName; }
+		const TManagedArray<int32>* RenderVerticesStart;
+		const TManagedArray<int32>* RenderVerticesEnd;
+		const TManagedArray<int32>* RenderFacesStart;
+		const TManagedArray<int32>* RenderFacesEnd;
+		const TManagedArray<int32>* RenderDeformerNumInfluences;  // Number of deformer mapping influences per render vertex, either 0 (no mappings), 1 or 5 (= NUM_INFLUENCES_PER_VERTEX as defined in the ush files)
+		const TManagedArray<FString>* RenderMaterialPathName;
 
 		//~ Sim Faces Group
-		TManagedArray<FIntVector3>* GetSimIndices2D() { return SimIndices2D; }
-		TManagedArray<FIntVector3>* GetSimIndices3D() { return SimIndices3D; }
+		const TManagedArray<FIntVector3>* SimIndices2D;
+		const TManagedArray<FIntVector3>* SimIndices3D;
 
 		//~ Sim Vertices 2D Group
-		TManagedArray<FVector2f>* GetSimPosition2D() { return SimPosition2D; }
-		TManagedArray<int32>* GetSimVertex3DLookup() { return SimVertex3DLookup; }
+		const TManagedArray<FVector2f>* SimPosition2D;
+		const TManagedArray<int32>* SimVertex3DLookup; // Lookup into corresponding 3D vertices
 
 		//~ Sim Vertices 3D Group
-		TManagedArray<FVector3f>* GetSimPosition3D() { return SimPosition3D; }
-		TManagedArray<FVector3f>* GetPreResizedSimPosition3D() { return PreResizedSimPosition3D; }
-		TManagedArray<FVector3f>* GetSimNormal() { return SimNormal; }
-		TManagedArray<TArray<int32>>* GetSimBoneIndices() { return SimBoneIndices; }
-		TManagedArray<TArray<float>>* GetSimBoneWeights() { return SimBoneWeights; }
-		TManagedArray<TArray<int32>>* GetTetherKinematicIndex() { return TetherKinematicIndex; }
-		TManagedArray<TArray<float>>* GetTetherReferenceLength() { return TetherReferenceLength; }
-		TManagedArray<TArray<int32>>* GetSimVertex2DLookup() { return SimVertex2DLookup; }
-		TManagedArray<TArray<int32>>* GetSeamStitchLookup() { return SeamStitchLookup; }
-		TManagedArray<float>* GetSimCustomResizingBlend() { return SimCustomResizingBlend; }
+		const TManagedArray<FVector3f>* SimPosition3D;
+		const TManagedArray<FVector3f>* SimNormal;  // Used for capture, maxdistance, backstop authoring ...etc
+		const TManagedArray<TArray<int32>>* SimBoneIndices;
+		const TManagedArray<TArray<float>>* SimBoneWeights;
+		const TManagedArray<TArray<int32>>* TetherKinematicIndex;
+		const TManagedArray<TArray<float>>* TetherReferenceLength;
+		const TManagedArray<TArray<int32>>* SimVertex2DLookup; // Lookup into corresponding 2D vertices
+		const TManagedArray<TArray<int32>>* SeamStitchLookup; // Lookup into any seam stitches which weld this vertex
+		const TManagedArray<float>* SimCustomResizingBlend; // Resizing weight map. How much to use group resizing result (1) vs rbf interpolation result (0) when doing sim mesh resizing
+		const TManagedArray<FVector3f>* PreResizedSimPosition3D;
 
 		//~ Sim Morph Targets Group
-		TManagedArray<FString>* GetSimMorphTargetName() { return SimMorphTargetName; }
-		TManagedArray<int32>* GetSimMorphTargetVerticesStart() { return SimMorphTargetVerticesStart; }
-		TManagedArray<int32>* GetSimMorphTargetVerticesEnd() { return SimMorphTargetVerticesEnd; }
+		const TManagedArray<FString>* SimMorphTargetName;
+		const TManagedArray<int32>* SimMorphTargetVerticesStart;
+		const TManagedArray<int32>* SimMorphTargetVerticesEnd;
 
-		//~ Sim Morph Targets Vertices Group
-		TManagedArray<FVector3f>* GetSimMorphTargetPositionDelta() { return SimMorphTargetPositionDelta; }
-		TManagedArray<FVector3f>* GetSimMorphTargetTangentZDelta() { return SimMorphTargetTangentZDelta; }
-		TManagedArray<int32>* GetSimMorphTargetSimVertex3DIndex() { return SimMorphTargetSimVertex3DIndex; }
+		//~ Sim Morph Target Vertices Group
+		const TManagedArray<FVector3f>* SimMorphTargetPositionDelta;
+		const TManagedArray<FVector3f>* SimMorphTargetTangentZDelta;
+		const TManagedArray<int32>* SimMorphTargetSimVertex3DIndex;
 
 		//~ Render Faces Group
-		TManagedArray<FIntVector3>* GetRenderIndices() { return RenderIndices; }
+		const TManagedArray<FIntVector3>* RenderIndices;
 
 		//~ Render Vertices Group
-		TManagedArray<FVector3f>* GetRenderPosition() { return RenderPosition; }
-		TManagedArray<FVector3f>* GetRenderNormal() { return RenderNormal; }
-		TManagedArray<FVector3f>* GetRenderTangentU() { return RenderTangentU; }
-		TManagedArray<FVector3f>* GetRenderTangentV() { return RenderTangentV; }
-		TManagedArray<TArray<FVector2f>>* GetRenderUVs() { return RenderUVs; }
-		TManagedArray<FLinearColor>* GetRenderColor() { return RenderColor; }
-		TManagedArray<TArray<int32>>* GetRenderBoneIndices() { return RenderBoneIndices; }
-		TManagedArray<TArray<float>>* GetRenderBoneWeights() { return RenderBoneWeights; }
-		TManagedArray<TArray<FVector4f>>* GetRenderDeformerPositionBaryCoordsAndDist() { return RenderDeformerPositionBaryCoordsAndDist; }
-		TManagedArray<TArray<FVector4f>>* GetRenderDeformerNormalBaryCoordsAndDist() { return RenderDeformerNormalBaryCoordsAndDist; }
-		TManagedArray<TArray<FVector4f>>* GetRenderDeformerTangentBaryCoordsAndDist() { return RenderDeformerTangentBaryCoordsAndDist; }
-		TManagedArray<TArray<FIntVector3>>* GetRenderDeformerSimIndices3D() { return RenderDeformerSimIndices3D; }
-		TManagedArray<TArray<float>>* GetRenderDeformerWeight() { return RenderDeformerWeight; }
-		TManagedArray<float>* GetRenderDeformerSkinningBlend() { return RenderDeformerSkinningBlend; }
-		TManagedArray<float>* GetRenderCustomResizingBlend() { return RenderCustomResizingBlend; }
+		const TManagedArray<FVector3f>* RenderPosition;
+		const TManagedArray<FVector3f>* RenderNormal;
+		const TManagedArray<FVector3f>* RenderTangentU;
+		const TManagedArray<FVector3f>* RenderTangentV;
+		const TManagedArray<TArray<FVector2f>>* RenderUVs;
+		const TManagedArray<FLinearColor>* RenderColor;
+		const TManagedArray<TArray<int32>>* RenderBoneIndices;
+		const TManagedArray<TArray<float>>* RenderBoneWeights;
+		const TManagedArray<TArray<FVector4f>>* RenderDeformerPositionBaryCoordsAndDist;  // Barycentric coords and distance along normal for the position of the final vert
+		const TManagedArray<TArray<FVector4f>>* RenderDeformerNormalBaryCoordsAndDist;  // Barycentric coords and distance along normal for the location of the unit normal endpoint
+		const TManagedArray<TArray<FVector4f>>* RenderDeformerTangentBaryCoordsAndDist;  // Barycentric coords and distance along normal for the location of the unit Tangent endpoint
+		const TManagedArray<TArray<FIntVector3>>* RenderDeformerSimIndices3D;  // The source mesh triangle
+		const TManagedArray<TArray<float>>* RenderDeformerWeight;  // For weighted averaging of multiple triangle influences
+		const TManagedArray<float>* RenderDeformerSkinningBlend;  // Render weight map. How much the vertex actually contributes, value between 0 (fully deformed) and 1 (fully skinned)
+		const TManagedArray<float>* RenderCustomResizingBlend; // Resizing weight map. How much to use group resizing result (1) vs rbf interpolation result (0) when doing render mesh resizing
 
 		//~ Group Resizing Group
-		TManagedArray<FString>* GetCustomResizingRegionSet() { return CustomResizingRegionSet; }
-		TManagedArray<int32>* GetCustomResizingRegionType() { return CustomResizingRegionType; }
+		const TManagedArray<FString>* CustomResizingRegionSet; // Name of the set associated with this binding type
+		const TManagedArray<int32>* CustomResizingRegionType; // Actually an Enum for binding type
+	};
 
-	private:
-		//~ Cloth collection
-		TSharedRef<FManagedArrayCollection> ManagedArrayCollection;
+	/**
+	 * Cloth collection facade data.
+	 */
+	class FClothCollection final : public FConstClothCollection
+	{
+	public:
+		explicit FClothCollection(const TSharedRef<FManagedArrayCollection>& InManagedArrayCollection);
+		~FClothCollection() = default;
 
-		//~ LODs Group
-		TManagedArray<FString>* PhysicsAssetPathName;
-		TManagedArray<FString>* SkeletalMeshPathName;
+		FClothCollection(const FClothCollection&) = delete;
+		FClothCollection& operator=(const FClothCollection&) = delete;
+		FClothCollection(FClothCollection&&) = delete;
+		FClothCollection& operator=(FClothCollection&&) = delete;
 
-		//~ Solvers group
-		TManagedArray<FVector3f>* SolverGravity;
-		TManagedArray<float>* SolverAirDamping;
-		TManagedArray<int32>* SolverSubSteps;
-		TManagedArray<float>* SolverTimeStep;
+		/** Make the underlying collection a cloth collection. */
+		void DefineSchema(EClothCollectionExtendedSchemas OptionalSchemas = EClothCollectionExtendedSchemas::None);
 
-		//~ Fabrics Group
-		TManagedArray<FVector3f>* FabricBendingStiffness;
-		TManagedArray<FVector3f>* FabricBucklingStiffness;
-		TManagedArray<FVector3f>* FabricStretchStiffness;
-		TManagedArray<float>* FabricBucklingRatio;
-		TManagedArray<float>* FabricDensity;
-		TManagedArray<float>* FabricFriction;
-		TManagedArray<float>* FabricDamping;
-		TManagedArray<float>* FabricPressure;
-		TManagedArray<int32>* FabricLayer;
-		TManagedArray<float>* FabricCollisionThickness;
+		/** Set the number of elements of a group. */
+		void SetNumElements(int32 InNumElements, const FName& GroupName);
+
+		/** Set the number of elements to one of the sub groups while maintaining the correct order of the data, and return the first index of the range. */
+		int32 SetNumElements(int32 InNumElements, const FName& GroupName, TManagedArray<int32>* StartArray, TManagedArray<int32>* EndArray, int32 Index);
+
+		/** Remove Elements */
+		void RemoveElements(const FName& Group, const TArray<int32>& SortedDeletionList);
+
+		/** Remove Elements. SortedDeletionList should be global indices. */
+		void RemoveElements(const FName& GroupName, const TArray<int32>& SortedDeletionList, TManagedArray<int32>* StartArray, TManagedArray<int32>* EndArray, int32 Index);
+
+		template<typename T UE_REQUIRES(TIsUserAttributeType<T>::Value)>
+		TManagedArray<T>* FindOrAddUserDefinedAttribute(const FName& Name, const FName& GroupName, const FName& GroupDependency = NAME_None);
+
+		template<typename T UE_REQUIRES(TIsUserAttributeType<T>::Value)>
+		void AddUserDefinedAttribute(const FName& Name, const FName& GroupName, const FName& GroupDependency = NAME_None)
+		{
+			FindOrAddUserDefinedAttribute<T>(Name, GroupName, GroupDependency);
+		}
+
+		void RemoveUserDefinedAttribute(const FName& Name, const FName& GroupName);
+
+		template<typename T UE_REQUIRES(TIsUserAttributeType<T>::Value)>
+		TManagedArray<T>* GetUserDefinedAttribute(const FName& Name, const FName& GroupName);
+
+		//~ LODs Group (There should be only one LOD per ClothCollection)
+		TManagedArray<FString>* GetPhysicsAssetPathName() { return const_cast<TManagedArray<FString>*>(PhysicsAssetPathName); }
+		TManagedArray<FString>* GetSkeletalMeshPathName() { return const_cast<TManagedArray<FString>*>(SkeletalMeshPathName); }
+
+		//~ Solver Group
+		TManagedArray<FVector3f>* GetSolverGravity() { return const_cast<TManagedArray<FVector3f>*>(SolverGravity); }
+		TManagedArray<float>* GetSolverAirDamping() { return const_cast<TManagedArray<float>*>(SolverAirDamping); }
+		TManagedArray<int32>* GetSolverSubSteps() { return const_cast<TManagedArray<int32>*>(SolverSubSteps); }
+		TManagedArray<float>* GetSolverTimeStep() { return const_cast<TManagedArray<float>*>(SolverTimeStep); }
+
+		//~ Fabric Group
+		TManagedArray<FVector3f>* GetFabricBendingStiffness() { return const_cast<TManagedArray<FVector3f>*>(FabricBendingStiffness); }
+		TManagedArray<FVector3f>* GetFabricBucklingStiffness() { return const_cast<TManagedArray<FVector3f>*>(FabricBucklingStiffness); }
+		TManagedArray<FVector3f>* GetFabricStretchStiffness() { return const_cast<TManagedArray<FVector3f>*>(FabricStretchStiffness); }
+		TManagedArray<float>* GetFabricBucklingRatio() { return const_cast<TManagedArray<float>*>(FabricBucklingRatio); }
+		TManagedArray<float>* GetFabricDensity() { return const_cast<TManagedArray<float>*>(FabricDensity); }
+		TManagedArray<float>* GetFabricFriction() { return const_cast<TManagedArray<float>*>(FabricFriction); }
+		TManagedArray<float>* GetFabricDamping() { return const_cast<TManagedArray<float>*>(FabricDamping); }
+		TManagedArray<float>* GetFabricPressure() { return const_cast<TManagedArray<float>*>(FabricPressure); }
+		TManagedArray<int32>* GetFabricLayer() { return const_cast<TManagedArray<int32>*>(FabricLayer); }
+		TManagedArray<float>* GetFabricCollisionThickness() { return const_cast<TManagedArray<float>*>(FabricCollisionThickness); }
 
 		//~ Seam Group
-		TManagedArray<int32>* SeamStitchStart;
-		TManagedArray<int32>* SeamStitchEnd;
+		TManagedArray<int32>* GetSeamStitchStart() { return const_cast<TManagedArray<int32>*>(SeamStitchStart); }
+		TManagedArray<int32>* GetSeamStitchEnd() { return const_cast<TManagedArray<int32>*>(SeamStitchEnd); }
 
 		//~ Seam Stitches Group
-		TManagedArray<FIntVector2>* SeamStitch2DEndIndices;  // Stitched 2D vertex indices pair
-		TManagedArray<int32>* SeamStitch3DIndex;  // Corresponding stitched 3D vertex
+		TManagedArray<FIntVector2>* GetSeamStitch2DEndIndices() { return const_cast<TManagedArray<FIntVector2>*>(SeamStitch2DEndIndices); }
+		TManagedArray<int32>* GetSeamStitch3DIndex() { return const_cast<TManagedArray<int32>*>(SeamStitch3DIndex); }
 
 		//~ Sim Patterns Group
-		TManagedArray<int32>* SimVertices2DStart;
-		TManagedArray<int32>* SimVertices2DEnd;
-		TManagedArray<int32>* SimFacesStart;
-		TManagedArray<int32>* SimFacesEnd;
-		TManagedArray<int32>* SimPatternFabric;
+		TManagedArray<int32>* GetSimVertices2DStart() { return const_cast<TManagedArray<int32>*>(SimVertices2DStart); }
+		TManagedArray<int32>* GetSimVertices2DEnd() { return const_cast<TManagedArray<int32>*>(SimVertices2DEnd); }
+		TManagedArray<int32>* GetSimFacesStart() { return const_cast<TManagedArray<int32>*>(SimFacesStart); }
+		TManagedArray<int32>* GetSimFacesEnd() { return const_cast<TManagedArray<int32>*>(SimFacesEnd); }
+		TManagedArray<int32>* GetSimPatternFabric() { return const_cast<TManagedArray<int32>*>(SimPatternFabric); }
 
 		//~ Render Patterns Group
-		TManagedArray<int32>* RenderVerticesStart;
-		TManagedArray<int32>* RenderVerticesEnd;
-		TManagedArray<int32>* RenderFacesStart;
-		TManagedArray<int32>* RenderFacesEnd;
-		TManagedArray<int32>* RenderDeformerNumInfluences;  // Number of deformer mapping influences per render vertex, either 0 (no mappings), 1 or 5 (= NUM_INFLUENCES_PER_VERTEX as defined in the ush files)
-		TManagedArray<FString>* RenderMaterialPathName;
+		TManagedArray<int32>* GetRenderVerticesStart() { return const_cast<TManagedArray<int32>*>(RenderVerticesStart); }
+		TManagedArray<int32>* GetRenderVerticesEnd() { return const_cast<TManagedArray<int32>*>(RenderVerticesEnd); }
+		TManagedArray<int32>* GetRenderFacesStart() { return const_cast<TManagedArray<int32>*>(RenderFacesStart); }
+		TManagedArray<int32>* GetRenderFacesEnd() { return const_cast<TManagedArray<int32>*>(RenderFacesEnd); }
+		TManagedArray<int32>* GetRenderDeformerNumInfluences() { return const_cast<TManagedArray<int32>*>(RenderDeformerNumInfluences); }
+		TManagedArray<FString>* GetRenderMaterialPathName() { return const_cast<TManagedArray<FString>*>(RenderMaterialPathName); }
 
 		//~ Sim Faces Group
-		TManagedArray<FIntVector3>* SimIndices2D;
-		TManagedArray<FIntVector3>* SimIndices3D;
+		TManagedArray<FIntVector3>* GetSimIndices2D() { return const_cast<TManagedArray<FIntVector3>*>(SimIndices2D); }
+		TManagedArray<FIntVector3>* GetSimIndices3D() { return const_cast<TManagedArray<FIntVector3>*>(SimIndices3D); }
 
 		//~ Sim Vertices 2D Group
-		TManagedArray<FVector2f>* SimPosition2D;
-		TManagedArray<int32>* SimVertex3DLookup; // Lookup into corresponding 3D vertices
+		TManagedArray<FVector2f>* GetSimPosition2D() { return const_cast<TManagedArray<FVector2f>*>(SimPosition2D); }
+		TManagedArray<int32>* GetSimVertex3DLookup() { return const_cast<TManagedArray<int32>*>(SimVertex3DLookup); }
 
 		//~ Sim Vertices 3D Group
-		TManagedArray<FVector3f>* SimPosition3D;
-		TManagedArray<FVector3f>* SimNormal;  // Used for capture, maxdistance, backstop authoring ...etc
-		TManagedArray<TArray<int32>>* SimBoneIndices;
-		TManagedArray<TArray<float>>* SimBoneWeights;
-		TManagedArray<TArray<int32>>* TetherKinematicIndex;
-		TManagedArray<TArray<float>>* TetherReferenceLength;
-		TManagedArray<TArray<int32>>* SimVertex2DLookup; // Lookup into corresponding 2D vertices
-		TManagedArray<TArray<int32>>* SeamStitchLookup; // Lookup into any seam stitches which weld this vertex
-		TManagedArray<float>* SimCustomResizingBlend; // Resizing weight map. How much to use group resizing result (1) vs rbf interpolation result (0) when doing sim mesh resizing
-		TManagedArray<FVector3f>* PreResizedSimPosition3D;
+		TManagedArray<FVector3f>* GetSimPosition3D() { return const_cast<TManagedArray<FVector3f>*>(SimPosition3D); }
+		TManagedArray<FVector3f>* GetPreResizedSimPosition3D() { return const_cast<TManagedArray<FVector3f>*>(PreResizedSimPosition3D); }
+		TManagedArray<FVector3f>* GetSimNormal() { return const_cast<TManagedArray<FVector3f>*>(SimNormal); }
+		TManagedArray<TArray<int32>>* GetSimBoneIndices() { return const_cast<TManagedArray<TArray<int32>>*>(SimBoneIndices); }
+		TManagedArray<TArray<float>>* GetSimBoneWeights() { return const_cast<TManagedArray<TArray<float>>*>(SimBoneWeights); }
+		TManagedArray<TArray<int32>>* GetTetherKinematicIndex() { return const_cast<TManagedArray<TArray<int32>>*>(TetherKinematicIndex); }
+		TManagedArray<TArray<float>>* GetTetherReferenceLength() { return const_cast<TManagedArray<TArray<float>>*>(TetherReferenceLength); }
+		TManagedArray<TArray<int32>>* GetSimVertex2DLookup() { return const_cast<TManagedArray<TArray<int32>>*>(SimVertex2DLookup); }
+		TManagedArray<TArray<int32>>* GetSeamStitchLookup() { return const_cast<TManagedArray<TArray<int32>>*>(SeamStitchLookup); }
+		TManagedArray<float>* GetSimCustomResizingBlend() { return const_cast<TManagedArray<float>*>(SimCustomResizingBlend); }
 
 		//~ Sim Morph Targets Group
-		TManagedArray<FString>* SimMorphTargetName;
-		TManagedArray<int32>* SimMorphTargetVerticesStart;
-		TManagedArray<int32>* SimMorphTargetVerticesEnd;
+		TManagedArray<FString>* GetSimMorphTargetName() { return const_cast<TManagedArray<FString>*>(SimMorphTargetName); }
+		TManagedArray<int32>* GetSimMorphTargetVerticesStart() { return const_cast<TManagedArray<int32>*>(SimMorphTargetVerticesStart); }
+		TManagedArray<int32>* GetSimMorphTargetVerticesEnd() { return const_cast<TManagedArray<int32>*>(SimMorphTargetVerticesEnd); }
 
-		//~ Sim Morph Target Vertices Group
-		TManagedArray<FVector3f>* SimMorphTargetPositionDelta;
-		TManagedArray<FVector3f>* SimMorphTargetTangentZDelta;
-		TManagedArray<int32>* SimMorphTargetSimVertex3DIndex;
+		//~ Sim Morph Targets Vertices Group
+		TManagedArray<FVector3f>* GetSimMorphTargetPositionDelta() { return const_cast<TManagedArray<FVector3f>*>(SimMorphTargetPositionDelta); }
+		TManagedArray<FVector3f>* GetSimMorphTargetTangentZDelta() { return const_cast<TManagedArray<FVector3f>*>(SimMorphTargetTangentZDelta); }
+		TManagedArray<int32>* GetSimMorphTargetSimVertex3DIndex() { return const_cast<TManagedArray<int32>*>(SimMorphTargetSimVertex3DIndex); }
 
 		//~ Render Faces Group
-		TManagedArray<FIntVector3>* RenderIndices;
+		TManagedArray<FIntVector3>* GetRenderIndices() { return const_cast<TManagedArray<FIntVector3>*>(RenderIndices); }
 
 		//~ Render Vertices Group
-		TManagedArray<FVector3f>* RenderPosition;
-		TManagedArray<FVector3f>* RenderNormal;
-		TManagedArray<FVector3f>* RenderTangentU;
-		TManagedArray<FVector3f>* RenderTangentV;
-		TManagedArray<TArray<FVector2f>>* RenderUVs;
-		TManagedArray<FLinearColor>* RenderColor;
-		TManagedArray<TArray<int32>>* RenderBoneIndices;
-		TManagedArray<TArray<float>>* RenderBoneWeights;
-		TManagedArray<TArray<FVector4f>>* RenderDeformerPositionBaryCoordsAndDist;  // Barycentric coords and distance along normal for the position of the final vert
-		TManagedArray<TArray<FVector4f>>* RenderDeformerNormalBaryCoordsAndDist;  // Barycentric coords and distance along normal for the location of the unit normal endpoint
-		TManagedArray<TArray<FVector4f>>* RenderDeformerTangentBaryCoordsAndDist;  // Barycentric coords and distance along normal for the location of the unit Tangent endpoint
-		TManagedArray<TArray<FIntVector3>>* RenderDeformerSimIndices3D;  // The source mesh triangle
-		TManagedArray<TArray<float>>* RenderDeformerWeight;  // For weighted averaging of multiple triangle influences
-		TManagedArray<float>* RenderDeformerSkinningBlend;  // Render weight map. How much the vertex actually contributes, value between 0 (fully deformed) and 1 (fully skinned)
-		TManagedArray<float>* RenderCustomResizingBlend; // Resizing weight map. How much to use group resizing result (1) vs rbf interpolation result (0) when doing render mesh resizing
+		TManagedArray<FVector3f>* GetRenderPosition() { return const_cast<TManagedArray<FVector3f>*>(RenderPosition); }
+		TManagedArray<FVector3f>* GetRenderNormal() { return const_cast<TManagedArray<FVector3f>*>(RenderNormal); }
+		TManagedArray<FVector3f>* GetRenderTangentU() { return const_cast<TManagedArray<FVector3f>*>(RenderTangentU); }
+		TManagedArray<FVector3f>* GetRenderTangentV() { return const_cast<TManagedArray<FVector3f>*>(RenderTangentV); }
+		TManagedArray<TArray<FVector2f>>* GetRenderUVs() { return const_cast<TManagedArray<TArray<FVector2f>>*>(RenderUVs); }
+		TManagedArray<FLinearColor>* GetRenderColor() { return const_cast<TManagedArray<FLinearColor>*>(RenderColor); }
+		TManagedArray<TArray<int32>>* GetRenderBoneIndices() { return const_cast<TManagedArray<TArray<int32>>*>(RenderBoneIndices); }
+		TManagedArray<TArray<float>>* GetRenderBoneWeights() { return const_cast<TManagedArray<TArray<float>>*>(RenderBoneWeights); }
+		TManagedArray<TArray<FVector4f>>* GetRenderDeformerPositionBaryCoordsAndDist() { return const_cast<TManagedArray<TArray<FVector4f>>*>(RenderDeformerPositionBaryCoordsAndDist); }
+		TManagedArray<TArray<FVector4f>>* GetRenderDeformerNormalBaryCoordsAndDist() { return const_cast<TManagedArray<TArray<FVector4f>>*>(RenderDeformerNormalBaryCoordsAndDist); }
+		TManagedArray<TArray<FVector4f>>* GetRenderDeformerTangentBaryCoordsAndDist() { return const_cast<TManagedArray<TArray<FVector4f>>*>(RenderDeformerTangentBaryCoordsAndDist); }
+		TManagedArray<TArray<FIntVector3>>* GetRenderDeformerSimIndices3D() { return const_cast<TManagedArray<TArray<FIntVector3>>*>(RenderDeformerSimIndices3D); }
+		TManagedArray<TArray<float>>* GetRenderDeformerWeight() { return const_cast<TManagedArray<TArray<float>>*>(RenderDeformerWeight); }
+		TManagedArray<float>* GetRenderDeformerSkinningBlend() { return const_cast<TManagedArray<float>*>(RenderDeformerSkinningBlend); }
+		TManagedArray<float>* GetRenderCustomResizingBlend() { return const_cast<TManagedArray<float>*>(RenderCustomResizingBlend); }
 
 		//~ Group Resizing Group
-		TManagedArray<FString>* CustomResizingRegionSet; // Name of the set associated with this binding type
-		TManagedArray<int32>* CustomResizingRegionType; // Actually an Enum for binding type
+		TManagedArray<FString>* GetCustomResizingRegionSet() { return const_cast<TManagedArray<FString>*>(CustomResizingRegionSet); }
+		TManagedArray<int32>* GetCustomResizingRegionType() { return const_cast<TManagedArray<int32>*>(CustomResizingRegionType); }
+
+	private:
+		TSharedRef<FManagedArrayCollection> GetManagedArrayCollection()
+		{
+			return ConstCastSharedRef<FManagedArrayCollection>(ManagedArrayCollection);
+		}
 	};
 
 	template<typename T>
-	inline TConstArrayView<T> FClothCollection::GetElements(const TManagedArray<T>* ElementArray, const TManagedArray<int32>* StartArray, const TManagedArray<int32>* EndArray, int32 ArrayIndex)
+	inline TConstArrayView<T> FConstClothCollection::GetElements(const TManagedArray<T>* ElementArray, const TManagedArray<int32>* StartArray, const TManagedArray<int32>* EndArray, int32 ArrayIndex)
 	{
 		if (StartArray && EndArray && ElementArray)
 		{
@@ -477,7 +503,7 @@ namespace UE::Chaos::ClothAsset
 	}
 
 	template<typename T>
-	inline TArrayView<T> FClothCollection::GetElements(TManagedArray<T>* ElementArray, const TManagedArray<int32>* StartArray, const TManagedArray<int32>* EndArray, int32 ArrayIndex)
+	inline TArrayView<T> FConstClothCollection::GetElements(TManagedArray<T>* ElementArray, const TManagedArray<int32>* StartArray, const TManagedArray<int32>* EndArray, int32 ArrayIndex)
 	{
 		if (StartArray && EndArray && ElementArray)
 		{
@@ -493,7 +519,7 @@ namespace UE::Chaos::ClothAsset
 	}
 
 	template<typename T>
-	inline TConstArrayView<T> FClothCollection::GetElements(const TManagedArray<T>* ElementArray)
+	inline TConstArrayView<T> FConstClothCollection::GetElements(const TManagedArray<T>* ElementArray)
 	{
 		if (ElementArray)
 		{
@@ -503,7 +529,7 @@ namespace UE::Chaos::ClothAsset
 	}
 
 	template<typename T>
-	inline TArrayView<T> FClothCollection::GetElements(TManagedArray<T>* ElementArray)
+	inline TArrayView<T> FConstClothCollection::GetElements(TManagedArray<T>* ElementArray)
 	{
 		if (ElementArray)
 		{
@@ -513,7 +539,7 @@ namespace UE::Chaos::ClothAsset
 	}
 
 	template<typename T>
-	inline TConstArrayView<T> FClothCollection::GetSubElements(
+	inline TConstArrayView<T> FConstClothCollection::GetSubElements(
 		const TManagedArray<T>* SubElementArray,
 		const TManagedArray<int32>* StartArray,
 		const TManagedArray<int32>* EndArray,
@@ -536,7 +562,7 @@ namespace UE::Chaos::ClothAsset
 	}
 
 	template<typename T>
-	inline TArrayView<T> FClothCollection::GetSubElements(
+	inline TArrayView<T> FConstClothCollection::GetSubElements(
 		TManagedArray<T>* SubElementArray,
 		const TManagedArray<int32>* StartArray,
 		const TManagedArray<int32>* EndArray,
@@ -559,7 +585,7 @@ namespace UE::Chaos::ClothAsset
 	}
 
 	template<typename T>
-	inline void FClothCollection::CopyArrayViewData(const TArrayView<T>& To, const TConstArrayView<T>& From)
+	inline void FConstClothCollection::CopyArrayViewData(const TArrayView<T>& To, const TConstArrayView<T>& From)
 	{
 		check(To.Num() == From.Num());
 		for (int32 Index = 0; Index < To.Num(); ++Index)
@@ -569,7 +595,7 @@ namespace UE::Chaos::ClothAsset
 	}
 
 	template<typename T UE_REQUIRES_DEFINITION(!std::is_array_v<T>)>
-	inline void FClothCollection::CopyArrayViewDataAndApplyOffset(const TArrayView<T>& To, const TConstArrayView<T>& From, const T Offset)
+	inline void FConstClothCollection::CopyArrayViewDataAndApplyOffset(const TArrayView<T>& To, const TConstArrayView<T>& From, const T Offset)
 	{
 		check(To.Num() == From.Num());
 		for (int32 Index = 0; Index < To.Num(); ++Index)
@@ -579,7 +605,7 @@ namespace UE::Chaos::ClothAsset
 	}
 
 	template<typename T>
-	inline void FClothCollection::CopyArrayViewDataAndApplyOffset(const TArrayView<TArray<T>>& To, const TConstArrayView<TArray<T>>& From, const T Offset)
+	inline void FConstClothCollection::CopyArrayViewDataAndApplyOffset(const TArrayView<TArray<T>>& To, const TConstArrayView<TArray<T>>& From, const T Offset)
 	{
 		check(To.Num() == From.Num());
 		for (int32 Index = 0; Index < To.Num(); ++Index)
@@ -593,7 +619,7 @@ namespace UE::Chaos::ClothAsset
 	}
 
 	template<typename T>
-	inline uint32 FClothCollection::GetElementsTypeHash(const TManagedArray<T>* ElementArray)
+	inline uint32 FConstClothCollection::GetElementsTypeHash(const TManagedArray<T>* ElementArray)
 	{
 		if (ElementArray)
 		{
diff --git a/Engine/Plugins/ChaosClothAsset/Source/ChaosClothAsset/Private/ChaosClothAsset/CollectionClothFabricFacade.cpp b/Engine/Plugins/ChaosClothAsset/Source/ChaosClothAsset/Private/ChaosClothAsset/CollectionClothFabricFacade.cpp
--- a/Engine/Plugins/ChaosClothAsset/Source/ChaosClothAsset/Private/ChaosClothAsset/CollectionClothFabricFacade.cpp
+++ b/Engine/Plugins/ChaosClothAsset/Source/ChaosClothAsset/Private/ChaosClothAsset/CollectionClothFabricFacade.cpp
@@ -78,7 +78,7 @@ namespace UE::Chaos::ClothAsset
 		return ClothCollection->GetFabricCollisionThickness() ? ClothCollection->GetElements(ClothCollection->GetFabricCollisionThickness())[GetElementIndex()] : FDefaultFabric::CollisionThickness;
 	}
 	
-	FCollectionClothFabricConstFacade::FCollectionClothFabricConstFacade(const TSharedRef<const FClothCollection>& ClothCollection, int32 InFabricIndex)
+	FCollectionClothFabricConstFacade::FCollectionClothFabricConstFacade(const TSharedRef<const FConstClothCollection>& ClothCollection, int32 InFabricIndex)
 		: ClothCollection(ClothCollection)
 		, FabricIndex(InFabricIndex)
 	{
@@ -142,4 +142,9 @@ namespace UE::Chaos::ClothAsset
 			FDefaultFabric::StretchStiffness, FDefaultFabric::Density, FDefaultFabric::Friction, FDefaultFabric::Damping,
 			FDefaultFabric::Pressure, FDefaultFabric::Layer, FDefaultFabric::CollisionThickness);
 	}
+
+	TSharedRef<class FClothCollection> FCollectionClothFabricFacade::GetClothCollection()
+	{
+		return StaticCastSharedRef<class FClothCollection>(ConstCastSharedRef<class FConstClothCollection>(ClothCollection));
+	}
 }  // End namespace UE::Chaos::ClothAsset
diff --git a/Engine/Plugins/ChaosClothAsset/Source/ChaosClothAsset/Private/ChaosClothAsset/CollectionClothFacade.cpp b/Engine/Plugins/ChaosClothAsset/Source/ChaosClothAsset/Private/ChaosClothAsset/CollectionClothFacade.cpp
--- a/Engine/Plugins/ChaosClothAsset/Source/ChaosClothAsset/Private/ChaosClothAsset/CollectionClothFacade.cpp
+++ b/Engine/Plugins/ChaosClothAsset/Source/ChaosClothAsset/Private/ChaosClothAsset/CollectionClothFacade.cpp
@@ -10,11 +10,11 @@
 namespace UE::Chaos::ClothAsset
 {
 	FCollectionClothConstFacade::FCollectionClothConstFacade(const TSharedRef<const FManagedArrayCollection>& ManagedArrayCollection)
-		: ClothCollection(MakeShared<const FClothCollection>(ConstCastSharedRef<FManagedArrayCollection>(ManagedArrayCollection)))
+		: ClothCollection(MakeShared<FConstClothCollection>(ManagedArrayCollection))
 	{
 	}
 
-	FCollectionClothConstFacade::FCollectionClothConstFacade(const TSharedRef<const FClothCollection>& ClothCollection)
+	FCollectionClothConstFacade::FCollectionClothConstFacade(const TSharedRef<const FConstClothCollection>& ClothCollection)
 		:ClothCollection(ClothCollection)
 	{
 	}
@@ -635,7 +635,7 @@ namespace UE::Chaos::ClothAsset
 	}
 
 	FCollectionClothFacade::FCollectionClothFacade(const TSharedRef<FManagedArrayCollection>& ManagedArrayCollection)
-		: FCollectionClothConstFacade(ManagedArrayCollection)
+		: FCollectionClothConstFacade(MakeShared<FClothCollection>(ManagedArrayCollection))
 	{
 	}
 
@@ -1443,4 +1443,9 @@ namespace UE::Chaos::ClothAsset
 		}
 	}
 
+	TSharedRef<class FClothCollection> FCollectionClothFacade::GetClothCollection()
+	{
+		return StaticCastSharedRef<class FClothCollection>(ConstCastSharedRef<class FConstClothCollection>(ClothCollection));
+	}
+
 } // End namespace UE::Chaos::ClothAsset
diff --git a/Engine/Plugins/ChaosClothAsset/Source/ChaosClothAsset/Private/ChaosClothAsset/CollectionClothRenderPatternFacade.cpp b/Engine/Plugins/ChaosClothAsset/Source/ChaosClothAsset/Private/ChaosClothAsset/CollectionClothRenderPatternFacade.cpp
--- a/Engine/Plugins/ChaosClothAsset/Source/ChaosClothAsset/Private/ChaosClothAsset/CollectionClothRenderPatternFacade.cpp
+++ b/Engine/Plugins/ChaosClothAsset/Source/ChaosClothAsset/Private/ChaosClothAsset/CollectionClothRenderPatternFacade.cpp
@@ -201,7 +201,7 @@ namespace UE::Chaos::ClothAsset
 		return GetNumRenderVertices() == 0 && GetNumRenderFaces() == 0;
 	}
 
-	FCollectionClothRenderPatternConstFacade::FCollectionClothRenderPatternConstFacade(const TSharedRef<const FClothCollection>& ClothCollection, int32 InPatternIndex)
+	FCollectionClothRenderPatternConstFacade::FCollectionClothRenderPatternConstFacade(const TSharedRef<const FConstClothCollection>& ClothCollection, int32 InPatternIndex)
 		: ClothCollection(ClothCollection)
 		, PatternIndex(InPatternIndex)
 	{
@@ -489,4 +489,9 @@ namespace UE::Chaos::ClothAsset
 		(*GetClothCollection()->GetRenderFacesStart())[ElementIndex] = INDEX_NONE;
 		(*GetClothCollection()->GetRenderFacesEnd())[ElementIndex] = INDEX_NONE;
 	}
+
+	TSharedRef<class FClothCollection> FCollectionClothRenderPatternFacade::GetClothCollection()
+	{
+		return StaticCastSharedRef<class FClothCollection>(ConstCastSharedRef<class FConstClothCollection>(ClothCollection));
+	}
 }  // End namespace UE::Chaos::ClothAsset
diff --git a/Engine/Plugins/ChaosClothAsset/Source/ChaosClothAsset/Private/ChaosClothAsset/CollectionClothSeamFacade.cpp b/Engine/Plugins/ChaosClothAsset/Source/ChaosClothAsset/Private/ChaosClothAsset/CollectionClothSeamFacade.cpp
--- a/Engine/Plugins/ChaosClothAsset/Source/ChaosClothAsset/Private/ChaosClothAsset/CollectionClothSeamFacade.cpp
+++ b/Engine/Plugins/ChaosClothAsset/Source/ChaosClothAsset/Private/ChaosClothAsset/CollectionClothSeamFacade.cpp
@@ -643,7 +643,7 @@ static FAutoConsoleVariableRef CVarEnableSeamChecks(TEXT("p.ChaosCloth.EnableSea
 #endif
 	}
 
-	FCollectionClothSeamConstFacade::FCollectionClothSeamConstFacade(const TSharedRef<const FClothCollection>& ClothCollection, int32 SeamIndex)
+	FCollectionClothSeamConstFacade::FCollectionClothSeamConstFacade(const TSharedRef<const FConstClothCollection>& ClothCollection, int32 SeamIndex)
 		: ClothCollection(ClothCollection)
 		, SeamIndex(SeamIndex)
 	{
@@ -1075,4 +1075,9 @@ static FAutoConsoleVariableRef CVarEnableSeamChecks(TEXT("p.ChaosCloth.EnableSea
 		(*GetClothCollection()->GetSeamStitchStart())[ElementIndex] = INDEX_NONE;
 		(*GetClothCollection()->GetSeamStitchEnd())[ElementIndex] = INDEX_NONE;
 	}
+
+	TSharedRef<class FClothCollection> FCollectionClothSeamFacade::GetClothCollection()
+	{
+		return StaticCastSharedRef<class FClothCollection>(ConstCastSharedRef<class FConstClothCollection>(ClothCollection));
+	}
 }  // End namespace UE::Chaos::ClothAsset
diff --git a/Engine/Plugins/ChaosClothAsset/Source/ChaosClothAsset/Private/ChaosClothAsset/CollectionClothSimMorphTargetFacade.cpp b/Engine/Plugins/ChaosClothAsset/Source/ChaosClothAsset/Private/ChaosClothAsset/CollectionClothSimMorphTargetFacade.cpp
--- a/Engine/Plugins/ChaosClothAsset/Source/ChaosClothAsset/Private/ChaosClothAsset/CollectionClothSimMorphTargetFacade.cpp
+++ b/Engine/Plugins/ChaosClothAsset/Source/ChaosClothAsset/Private/ChaosClothAsset/CollectionClothSimMorphTargetFacade.cpp
@@ -61,7 +61,7 @@ namespace UE::Chaos::ClothAsset
 		return GetNumSimMorphTargetVertices() == 0;
 	}
 
-	FCollectionClothSimMorphTargetConstFacade::FCollectionClothSimMorphTargetConstFacade(const TSharedRef<const FClothCollection>& ClothCollection, int32 InMorphTargetIndex)
+	FCollectionClothSimMorphTargetConstFacade::FCollectionClothSimMorphTargetConstFacade(const TSharedRef<const FConstClothCollection>& ClothCollection, int32 InMorphTargetIndex)
 		: ClothCollection(ClothCollection)
 		, MorphTargetIndex(InMorphTargetIndex)
 	{
@@ -191,4 +191,9 @@ namespace UE::Chaos::ClothAsset
 		(*GetClothCollection()->GetSimMorphTargetVerticesStart())[ElementIndex] = INDEX_NONE;
 		(*GetClothCollection()->GetSimMorphTargetVerticesEnd())[ElementIndex] = INDEX_NONE;
 	}
+
+	TSharedRef<class FClothCollection> FCollectionClothSimMorphTargetFacade::GetClothCollection()
+	{
+		return StaticCastSharedRef<class FClothCollection>(ConstCastSharedRef<class FConstClothCollection>(ClothCollection));
+	}
 }  // End namespace UE::Chaos::ClothAsset
diff --git a/Engine/Plugins/ChaosClothAsset/Source/ChaosClothAsset/Private/ChaosClothAsset/CollectionClothSimPatternFacade.cpp b/Engine/Plugins/ChaosClothAsset/Source/ChaosClothAsset/Private/ChaosClothAsset/CollectionClothSimPatternFacade.cpp
--- a/Engine/Plugins/ChaosClothAsset/Source/ChaosClothAsset/Private/ChaosClothAsset/CollectionClothSimPatternFacade.cpp
+++ b/Engine/Plugins/ChaosClothAsset/Source/ChaosClothAsset/Private/ChaosClothAsset/CollectionClothSimPatternFacade.cpp
@@ -97,7 +97,7 @@ namespace UE::Chaos::ClothAsset
 		return GetNumSimFaces() == 0 && GetNumSimVertices2D() == 0;
 	}
 
-	FCollectionClothSimPatternConstFacade::FCollectionClothSimPatternConstFacade(const TSharedRef<const FClothCollection>& ClothCollection, int32 InPatternIndex)
+	FCollectionClothSimPatternConstFacade::FCollectionClothSimPatternConstFacade(const TSharedRef<const FConstClothCollection>& ClothCollection, int32 InPatternIndex)
 		: ClothCollection(ClothCollection)
 		, PatternIndex(InPatternIndex)
 	{
@@ -356,4 +356,9 @@ namespace UE::Chaos::ClothAsset
 			GetClothCollection()->GetSimVertices2DEnd(),
 			GetElementIndex());
 	}
+
+	TSharedRef<class FClothCollection> FCollectionClothSimPatternFacade::GetClothCollection()
+	{
+		return StaticCastSharedRef<class FClothCollection>(ConstCastSharedRef<class FConstClothCollection>(ClothCollection));
+	}
 }  // End namespace UE::Chaos::ClothAsset
diff --git a/Engine/Plugins/ChaosClothAsset/Source/ChaosClothAsset/Public/ChaosClothAsset/CollectionClothFabricFacade.h b/Engine/Plugins/ChaosClothAsset/Source/ChaosClothAsset/Public/ChaosClothAsset/CollectionClothFabricFacade.h
--- a/Engine/Plugins/ChaosClothAsset/Source/ChaosClothAsset/Public/ChaosClothAsset/CollectionClothFabricFacade.h
+++ b/Engine/Plugins/ChaosClothAsset/Source/ChaosClothAsset/Public/ChaosClothAsset/CollectionClothFabricFacade.h
@@ -94,12 +94,12 @@ namespace UE::Chaos::ClothAsset
 	protected:
 		friend class FCollectionClothFabricFacade;  // For other instances access
 		friend class FCollectionClothConstFacade;
-		UE_API FCollectionClothFabricConstFacade(const TSharedRef<const class FClothCollection>& InClothCollection, int32 InFabricIndex);
+		FCollectionClothFabricConstFacade(const TSharedRef<const class FConstClothCollection>& InClothCollection, int32 InFabricIndex);
 
 		static constexpr int32 GetBaseElementIndex() { return 0; }
 
 		/** Cloth collection modified by  the fabric facade */
-		TSharedRef<const class FClothCollection> ClothCollection;
+		TSharedRef<const class FConstClothCollection> ClothCollection;
 
 		/** Fabric index that will be referred in the sim patterns */
 		int32 FabricIndex;
@@ -137,16 +137,16 @@ namespace UE::Chaos::ClothAsset
 
 	private:
 		friend class FCollectionClothFacade;
-		UE_API FCollectionClothFabricFacade(const TSharedRef<class FClothCollection>& InClothCollection, int32 InFabricIndex);
+		FCollectionClothFabricFacade(const TSharedRef<class FClothCollection>& InClothCollection, int32 InFabricIndex);
 
 		/** Set default values to the fabric properties */
-		UE_API void SetDefaults();
+		void SetDefaults();
 
 		/** Reset the fabric values properties */
-		UE_API void Reset();
+		void Reset();
 
 		/** Get the non const cloth collection */
-		TSharedRef<class FClothCollection> GetClothCollection() { return ConstCastSharedRef<class FClothCollection>(ClothCollection); }
+		TSharedRef<class FClothCollection> GetClothCollection();
 	};
 }  // End namespace UE::Chaos::ClothAsset
 
diff --git a/Engine/Plugins/ChaosClothAsset/Source/ChaosClothAsset/Public/ChaosClothAsset/CollectionClothFacade.h b/Engine/Plugins/ChaosClothAsset/Source/ChaosClothAsset/Public/ChaosClothAsset/CollectionClothFacade.h
--- a/Engine/Plugins/ChaosClothAsset/Source/ChaosClothAsset/Public/ChaosClothAsset/CollectionClothFacade.h
+++ b/Engine/Plugins/ChaosClothAsset/Source/ChaosClothAsset/Public/ChaosClothAsset/CollectionClothFacade.h
@@ -217,16 +217,16 @@ namespace UE::Chaos::ClothAsset
 			TArray<uint32>& PatternToWeldedIndices, TArray<TArray<int32>>* OptionalWeldedToPatternIndices = nullptr) const;
 
 	protected:
+		explicit FCollectionClothConstFacade(const TSharedRef<const class FConstClothCollection>& ClothCollection);
+		
 		friend class FCollectionClothFacade;
-		TSharedRef<const class FClothCollection> ClothCollection;
+		TSharedRef<const class FConstClothCollection> ClothCollection;
 
 		friend class FCollectionClothFacade;  // To enable access from a different instance
 		friend class FCollectionClothSeamFacade;
 		friend class FCollectionClothSeamConstFacade;
 		friend class FCollectionClothFabricFacade;
 		friend class FCollectionClothFabricConstFacade;
-		
-		UE_API explicit FCollectionClothConstFacade(const TSharedRef<const class FClothCollection>& ClothCollection);
 	};
 
 	/**
@@ -426,23 +426,23 @@ namespace UE::Chaos::ClothAsset
 
 	private:
 		
-		UE_API void SetDefaults();
+		void SetDefaults();
 
-		TSharedRef<class FClothCollection> GetClothCollection() { return ConstCastSharedRef<class FClothCollection>(ClothCollection); }
+		TSharedRef<class FClothCollection> GetClothCollection();
 
 		friend class FCollectionClothSeamFacade;
 		friend class FCollectionClothSimPatternFacade;
 		friend class FCollectionClothFabricFacade;
 
-		UE_API explicit FCollectionClothFacade(const TSharedRef<class FClothCollection>& InClothCollection);
+		explicit FCollectionClothFacade(const TSharedRef<class FClothCollection>& InClothCollection);
 
 		// These methods are private because they're managed by the FCollectionClothSeamFacade.
 		//~ Sim Vertices 2D Group
-		UE_API TArrayView<int32> GetSimVertex3DLookupPrivate();
+		TArrayView<int32> GetSimVertex3DLookupPrivate();
 
 		//~ Sim Vertices 3D Group
-		UE_API TArrayView<TArray<int32>> GetSeamStitchLookupPrivate();
-		UE_API TArrayView<TArray<int32>> GetSimVertex2DLookupPrivate();
+		TArrayView<TArray<int32>> GetSeamStitchLookupPrivate();
+		TArrayView<TArray<int32>> GetSimVertex2DLookupPrivate();
 	};
 }  // End namespace UE::Chaos::ClothAsset
 
diff --git a/Engine/Plugins/ChaosClothAsset/Source/ChaosClothAsset/Public/ChaosClothAsset/CollectionClothRenderPatternFacade.h b/Engine/Plugins/ChaosClothAsset/Source/ChaosClothAsset/Public/ChaosClothAsset/CollectionClothRenderPatternFacade.h
--- a/Engine/Plugins/ChaosClothAsset/Source/ChaosClothAsset/Public/ChaosClothAsset/CollectionClothRenderPatternFacade.h
+++ b/Engine/Plugins/ChaosClothAsset/Source/ChaosClothAsset/Public/ChaosClothAsset/CollectionClothRenderPatternFacade.h
@@ -70,12 +70,12 @@ namespace UE::Chaos::ClothAsset
 	protected:
 		friend class FCollectionClothRenderPatternFacade;  // For other instances access
 		friend class FCollectionClothConstFacade;
-		UE_API FCollectionClothRenderPatternConstFacade(const TSharedRef<const class FClothCollection>& InClothCollection, int32 InPatternIndex);
+		FCollectionClothRenderPatternConstFacade(const TSharedRef<const class FConstClothCollection>& InClothCollection, int32 InPatternIndex);
 
 		static constexpr int32 GetBaseElementIndex() { return 0; }
 		int32 GetElementIndex() const { return GetBaseElementIndex() + PatternIndex; }
 
-		TSharedRef<const class FClothCollection> ClothCollection;
+		TSharedRef<const class FConstClothCollection> ClothCollection;
 		int32 PatternIndex;
 	};
 
@@ -139,13 +139,13 @@ namespace UE::Chaos::ClothAsset
 		UE_API void RemoveRenderFaces(const TArray<int32>& SortedDeletionList);
 		UE_API TArrayView<FIntVector3> GetRenderIndices();
 
-	protected:
+	private:
 		friend class FCollectionClothFacade;
-		UE_API FCollectionClothRenderPatternFacade(const TSharedRef<class FClothCollection>& InClothCollection, int32 InPatternIndex);
+		FCollectionClothRenderPatternFacade(const TSharedRef<class FClothCollection>& InClothCollection, int32 InPatternIndex);
 
-		UE_API void SetDefaults();
+		void SetDefaults();
 
-		TSharedRef<class FClothCollection> GetClothCollection() { return ConstCastSharedRef<class FClothCollection>(ClothCollection); }
+		TSharedRef<class FClothCollection> GetClothCollection();
 	};
 }  // End namespace UE::Chaos::ClothAsset
 
diff --git a/Engine/Plugins/ChaosClothAsset/Source/ChaosClothAsset/Public/ChaosClothAsset/CollectionClothSeamFacade.h b/Engine/Plugins/ChaosClothAsset/Source/ChaosClothAsset/Public/ChaosClothAsset/CollectionClothSeamFacade.h
--- a/Engine/Plugins/ChaosClothAsset/Source/ChaosClothAsset/Public/ChaosClothAsset/CollectionClothSeamFacade.h
+++ b/Engine/Plugins/ChaosClothAsset/Source/ChaosClothAsset/Public/ChaosClothAsset/CollectionClothSeamFacade.h
@@ -45,12 +45,12 @@ namespace UE::Chaos::ClothAsset
 	protected:
 		friend class FCollectionClothSeamFacade;  // For other instances access
 		friend class FCollectionClothConstFacade;
-		UE_API FCollectionClothSeamConstFacade(const TSharedRef<const class FClothCollection>& InClothCollection, int32 SeamIndex);
+		FCollectionClothSeamConstFacade(const TSharedRef<const class FConstClothCollection>& InClothCollection, int32 SeamIndex);
 
 		static constexpr int32 GetBaseElementIndex() { return 0; }
 		int32 GetElementIndex() const { return GetBaseElementIndex() + SeamIndex; }
 
-		TSharedRef<const class FClothCollection> ClothCollection;
+		TSharedRef<const class FConstClothCollection> ClothCollection;
 		int32 SeamIndex;
 	};
 
@@ -84,19 +84,19 @@ namespace UE::Chaos::ClothAsset
 		/** Clean up references to invalid indices, including updating stitches to maintain topology. */
 		UE_API void CleanupAndCompact();
 
-private:
+	private:
 		friend class FCollectionClothFacade;
-		UE_API FCollectionClothSeamFacade(const TSharedRef<class FClothCollection>& InClothCollection, int32 SeamIndex);
+		FCollectionClothSeamFacade(const TSharedRef<class FClothCollection>& InClothCollection, int32 SeamIndex);
 
-		UE_API void SetDefaults();
+		void SetDefaults();
 
 		//~ Seam Stitches Group
 		// Indices correspond with the FCollectionClothConstFacade indices (e.g., not pattern indices)
-		UE_API void SetNumSeamStitches(const int32 NumStitches);
-		UE_API TArrayView<FIntVector2> GetSeamStitch2DEndIndices();
-		UE_API TArrayView<int32> GetSeamStitch3DIndex();
+		void SetNumSeamStitches(const int32 NumStitches);
+		TArrayView<FIntVector2> GetSeamStitch2DEndIndices();
+		TArrayView<int32> GetSeamStitch3DIndex();
 
-		TSharedRef<class FClothCollection> GetClothCollection() { return ConstCastSharedRef<class FClothCollection>(ClothCollection); }
+		TSharedRef<class FClothCollection> GetClothCollection();
 	};
 }  // End namespace UE::Chaos::ClothAsset
 
diff --git a/Engine/Plugins/ChaosClothAsset/Source/ChaosClothAsset/Public/ChaosClothAsset/CollectionClothSimMorphTargetFacade.h b/Engine/Plugins/ChaosClothAsset/Source/ChaosClothAsset/Public/ChaosClothAsset/CollectionClothSimMorphTargetFacade.h
--- a/Engine/Plugins/ChaosClothAsset/Source/ChaosClothAsset/Public/ChaosClothAsset/CollectionClothSimMorphTargetFacade.h
+++ b/Engine/Plugins/ChaosClothAsset/Source/ChaosClothAsset/Public/ChaosClothAsset/CollectionClothSimMorphTargetFacade.h
@@ -48,12 +48,12 @@ namespace UE::Chaos::ClothAsset
 	protected:
 		friend class FCollectionClothSimMorphTargetFacade;  // For other instances access
 		friend class FCollectionClothConstFacade;
-		UE_API FCollectionClothSimMorphTargetConstFacade(const TSharedRef<const class FClothCollection>& InClothCollection, int32 InPatternIndex);
+		FCollectionClothSimMorphTargetConstFacade(const TSharedRef<const class FConstClothCollection>& InClothCollection, int32 InPatternIndex);
 
 		static constexpr int32 GetBaseElementIndex() { return 0; }
 		int32 GetElementIndex() const { return GetBaseElementIndex() + MorphTargetIndex; }
 
-		TSharedRef<const class FClothCollection> ClothCollection;
+		TSharedRef<const class FConstClothCollection> ClothCollection;
 		int32 MorphTargetIndex;
 	};
 
@@ -98,13 +98,13 @@ namespace UE::Chaos::ClothAsset
 		/** Remove all morph target vertices with invalid indices. */
 		UE_API void Compact();
 
-	protected:
+	private:
 		friend class FCollectionClothFacade;
-		UE_API FCollectionClothSimMorphTargetFacade(const TSharedRef<class FClothCollection>& InClothCollection, int32 InMorphTargetIndex);
+		FCollectionClothSimMorphTargetFacade(const TSharedRef<class FClothCollection>& InClothCollection, int32 InMorphTargetIndex);
 
-		UE_API void SetDefaults();
+		void SetDefaults();
 
-		TSharedRef<class FClothCollection> GetClothCollection() { return ConstCastSharedRef<class FClothCollection>(ClothCollection); }
+		TSharedRef<class FClothCollection> GetClothCollection();
 	};
 }  // End namespace UE::Chaos::ClothAsset
 
diff --git a/Engine/Plugins/ChaosClothAsset/Source/ChaosClothAsset/Public/ChaosClothAsset/CollectionClothSimPatternFacade.h b/Engine/Plugins/ChaosClothAsset/Source/ChaosClothAsset/Public/ChaosClothAsset/CollectionClothSimPatternFacade.h
--- a/Engine/Plugins/ChaosClothAsset/Source/ChaosClothAsset/Public/ChaosClothAsset/CollectionClothSimPatternFacade.h
+++ b/Engine/Plugins/ChaosClothAsset/Source/ChaosClothAsset/Public/ChaosClothAsset/CollectionClothSimPatternFacade.h
@@ -61,12 +61,12 @@ namespace UE::Chaos::ClothAsset
 	protected:
 		friend class FCollectionClothSimPatternFacade;  // For other instances access
 		friend class FCollectionClothConstFacade;
-		UE_API FCollectionClothSimPatternConstFacade(const TSharedRef<const class FClothCollection>& InClothCollection, int32 InPatternIndex);
+		FCollectionClothSimPatternConstFacade(const TSharedRef<const class FConstClothCollection>& InClothCollection, int32 InPatternIndex);
 
 		static constexpr int32 GetBaseElementIndex() { return 0; }
 		int32 GetElementIndex() const { return GetBaseElementIndex() + PatternIndex; }
 
-		TSharedRef<const class FClothCollection> ClothCollection;
+		TSharedRef<const class FConstClothCollection> ClothCollection;
 		int32 PatternIndex;
 	};
 
@@ -120,17 +120,17 @@ namespace UE::Chaos::ClothAsset
 
 	private:
 		friend class FCollectionClothFacade;
-		UE_API FCollectionClothSimPatternFacade(const TSharedRef<class FClothCollection>& InClothCollection, int32 InPatternIndex);
+		FCollectionClothSimPatternFacade(const TSharedRef<class FClothCollection>& InClothCollection, int32 InPatternIndex);
 
-		UE_API void SetDefaults();
+		void SetDefaults();
 
 		/** Append this many 2D and 3D SimVertices which correspond to each other. Returns {2D start index (local pattern index), 3D start index}. */
-		UE_API FIntVector2 AppendSimVertices(int32 NumSimVertices);
+		FIntVector2 AppendSimVertices(int32 NumSimVertices);
 
-		UE_API void SetNumSimVertices2D(int32 NumSimVertices);
-		UE_API TArrayView<int32> GetSimVertex3DLookup();
+		void SetNumSimVertices2D(int32 NumSimVertices);
+		TArrayView<int32> GetSimVertex3DLookup();
 
-		TSharedRef<class FClothCollection> GetClothCollection() { return ConstCastSharedRef<class FClothCollection>(ClothCollection); }
+		TSharedRef<class FClothCollection> GetClothCollection();
 	};
 }  // End namespace UE::Chaos::ClothAsset
 
