Chaos Cloth Asset - Added GetNumTriangles() method to the ClothSimulationModel to simplify the deformer mapping descriptor implementations.

#rnx

diff --git a/Engine/Plugins/ChaosClothAsset/Source/ChaosClothAssetEngine/Public/ChaosClothAsset/ClothSimulationModel.h b/Engine/Plugins/ChaosClothAsset/Source/ChaosClothAssetEngine/Public/ChaosClothAsset/ClothSimulationModel.h
--- a/Engine/Plugins/ChaosClothAsset/Source/ChaosClothAssetEngine/Public/ChaosClothAsset/ClothSimulationModel.h
+++ b/Engine/Plugins/ChaosClothAsset/Source/ChaosClothAssetEngine/Public/ChaosClothAsset/ClothSimulationModel.h
@@ -123,6 +123,7 @@ struct FChaosClothSimulationModel
 	bool IsValidLodIndex(int32 LodIndex) const { return ClothSimulationLodModels.IsValidIndex(LodIndex); }
 
 	int32 GetNumVertices(int32 LodIndex) const { return IsValidLodIndex(LodIndex) ? ClothSimulationLodModels[LodIndex].Positions.Num() : 0; }
+	int32 GetNumTriangles(int32 LodIndex) const { return IsValidLodIndex(LodIndex) ? ClothSimulationLodModels[LodIndex].Indices.Num() / 3 : 0; }
 
 	TConstArrayView<FVector3f> GetPositions(int32 LodIndex) const { return IsValidLodIndex(LodIndex) ? ClothSimulationLodModels[LodIndex].Positions : TConstArrayView<FVector3f>(); }
 	TConstArrayView<FVector2f> GetPatternPositions(int32 LodIndex) const { return IsValidLodIndex(LodIndex) ? ClothSimulationLodModels[LodIndex].PatternPositions : TConstArrayView<FVector2f>(); }
