Chaos Cloth Asset - Added first pass factories to use the Cloth Asset inside the SkeletalMesh.

#rnx

diff --git a/Engine/Plugins/ChaosClothAsset/Source/ChaosClothAssetEngine/ChaosClothAssetEngine.Build.cs b/Engine/Plugins/ChaosClothAsset/Source/ChaosClothAssetEngine/ChaosClothAssetEngine.Build.cs
--- a/Engine/Plugins/ChaosClothAsset/Source/ChaosClothAssetEngine/ChaosClothAssetEngine.Build.cs
+++ b/Engine/Plugins/ChaosClothAsset/Source/ChaosClothAssetEngine/ChaosClothAssetEngine.Build.cs
@@ -29,7 +29,8 @@ public class ChaosClothAssetEngine : ModuleRules
 				"ChaosClothAsset",
 				"ChaosCloth",
 				"ChaosCaching",
-				"ClothingSystemRuntimeInterface"
+				"ClothingSystemRuntimeInterface",
+				"ClothingSystemEditorInterface"  // For UChaosClothAssetSKMClothingAssetFactory
 			}
 		);
 		SetupModulePhysicsSupport(Target);
diff --git a/Engine/Plugins/ChaosClothAsset/Source/ChaosClothAssetEngine/Private/ChaosClothAsset/ChaosClothAssetEngineModule.cpp b/Engine/Plugins/ChaosClothAsset/Source/ChaosClothAssetEngine/Private/ChaosClothAsset/ChaosClothAssetEngineModule.cpp
--- a/Engine/Plugins/ChaosClothAsset/Source/ChaosClothAssetEngine/Private/ChaosClothAsset/ChaosClothAssetEngineModule.cpp
+++ b/Engine/Plugins/ChaosClothAsset/Source/ChaosClothAssetEngine/Private/ChaosClothAsset/ChaosClothAssetEngineModule.cpp
@@ -1,11 +1,15 @@
 // Copyright Epic Games, Inc. All Rights Reserved.
 
 #include "ChaosClothAsset/ClothAssetPrivate.h"
-#include "ChaosClothAsset/ClothComponent.h"
+#include "ChaosClothAsset/ClothAssetSKMSimulationFactory.h"
 #include "ChaosClothAsset/ClothComponentCacheAdapter.h"
+#include "Features/IModularFeatures.h"
 #include "Modules/ModuleInterface.h"
 #include "Modules/ModuleManager.h"
 #if WITH_EDITOR
+#include "ChaosClothAsset/ClothAssetSKMClothingAssetFactory.h"
+#include "ChaosClothAsset/ClothComponent.h"
+#include "ClothingAssetFactoryInterface.h"
 #include "PropertyEditorModule.h"
 #endif
 
@@ -15,9 +19,15 @@ DEFINE_LOG_CATEGORY(LogChaosClothAsset);
 
 namespace UE::Chaos::ClothAsset
 {
-	class FClothAssetEngineModule : public IModuleInterface
+	class FClothAssetEngineModule
+		: public IModuleInterface
+		, public IClothingSimulationFactoryClassProvider
+#if WITH_EDITOR
+		, public IClothingAssetFactoryProvider
+#endif
 	{
 	public:
+		//~ Begin IModuleInterface
 		virtual void StartupModule() override
 		{
 			ClothComponentCacheAdapter = MakeUnique<FClothComponentCacheAdapter>();
@@ -29,12 +39,18 @@ namespace UE::Chaos::ClothAsset
 			const TSharedPtr<FPropertySection> Section = PropertyModule.FindOrCreateSection(UChaosClothComponent::StaticClass()->GetFName(), TEXT("ClothSim"), LOCTEXT("ClothSim", "Cloth Sim"));
 			Section->AddCategory("Cloth Component");
 			PropertyModule.NotifyCustomizationModuleChanged();
+
+			IModularFeatures::Get().RegisterModularFeature(IClothingAssetFactoryProvider::FeatureName, static_cast<IClothingAssetFactoryProvider*>(this));
 #endif
+			IModularFeatures::Get().RegisterModularFeature(IClothingSimulationFactoryClassProvider::FeatureName, static_cast<IClothingSimulationFactoryClassProvider*>(this));
 		}
 
 		virtual void ShutdownModule() override
 		{
+			IModularFeatures::Get().UnregisterModularFeature(IClothingSimulationFactoryClassProvider::FeatureName, static_cast<IClothingSimulationFactoryClassProvider*>(this));
 #if WITH_EDITOR
+			IModularFeatures::Get().UnregisterModularFeature(IClothingAssetFactoryProvider::FeatureName, static_cast<IClothingAssetFactoryProvider*>(this));
+
 			if (UObjectInitialized())
 			{
 				// Remove sections
@@ -43,6 +59,24 @@ namespace UE::Chaos::ClothAsset
 			}
 #endif
 		}
+		//~ End IModuleInterface
+
+		//~ Begin IClothingSimulationFactoryClassProvider
+		virtual TSubclassOf<class UClothingSimulationFactory> GetClothingSimulationFactoryClass() const override
+		{
+			return UChaosClothAssetSKMSimulationFactory::StaticClass();
+		}
+		//~ End IClothingSimulationFactoryClassProvider
+
+#if WITH_EDITOR
+		//~ Begin IClothingAssetFactoryProvider
+		virtual UClothingAssetFactoryBase* GetFactory() override
+		{
+			return UChaosClothAssetSKMClothingAssetFactory::StaticClass()->GetDefaultObject<UClothingAssetFactoryBase>();
+		}
+		//~ End IClothingSimulationFactoryClassProvider
+#endif
+
 	private:
 		TUniquePtr<FClothComponentCacheAdapter> ClothComponentCacheAdapter;
 	};
diff --git a/Engine/Plugins/ChaosClothAsset/Source/ChaosClothAssetEngine/Private/ChaosClothAsset/ClothAssetBuilder.cpp b/Engine/Plugins/ChaosClothAsset/Source/ChaosClothAssetEngine/Private/ChaosClothAsset/ClothAssetBuilder.cpp
--- a/Engine/Plugins/ChaosClothAsset/Source/ChaosClothAssetEngine/Private/ChaosClothAsset/ClothAssetBuilder.cpp
+++ b/Engine/Plugins/ChaosClothAsset/Source/ChaosClothAssetEngine/Private/ChaosClothAsset/ClothAssetBuilder.cpp
@@ -8,6 +8,7 @@
 #include "ChaosClothAsset/ClothCollectionGroup.h"
 #include "ChaosClothAsset/CollectionClothFacade.h"
 #include "ChaosClothAsset/ClothAssetPrivate.h"
+#include "ChaosClothAsset/ClothEngineTools.h"
 #include "Chaos/CollectionPropertyFacade.h"
 #include "Engine/RendererSettings.h"
 #include "Engine/SkeletalMesh.h"
@@ -151,18 +152,7 @@ void UChaosClothAsset::FBuilder::BuildLod(FSkeletalMeshLODModel& LODModel, const
 
 	// Retrieve MaxDistance information (weight map and Low/High values)
 	FCollectionPropertyConstFacade Properties(ClothCollection);
-
-	int32 MaxDistancePropertyKeyIndex;
-	FString MaxDistanceString = TEXT("MaxDistance");
-	MaxDistanceString = Properties.GetStringValue(MaxDistanceString, MaxDistanceString, &MaxDistancePropertyKeyIndex);
-	const bool bHasMaxDistanceProperty = (MaxDistancePropertyKeyIndex != INDEX_NONE);
-	const float MaxDistanceOffset = bHasMaxDistanceProperty ? Properties.GetLowValue<float>(MaxDistancePropertyKeyIndex) : TNumericLimits<float>::Max();  // Uses infinite distance when no MaxDistance properties are set
-	const float MaxDistanceScale = bHasMaxDistanceProperty ? Properties.GetHighValue<float>(MaxDistancePropertyKeyIndex) - MaxDistanceOffset : 0.f;
-	const TConstArrayView<float> MaxDistanceWeightMap = ClothFacade.GetWeightMap(FName(MaxDistanceString));
-
-	const FPointWeightMap MaxDistances = (MaxDistanceWeightMap.Num() == NumLodSimVertices) ?
-		FPointWeightMap(MaxDistanceWeightMap, MaxDistanceOffset, MaxDistanceScale) :
-		FPointWeightMap(NumLodSimVertices, MaxDistanceOffset);
+	const FPointWeightMap MaxDistances = FClothEngineTools::GetMaxDistanceWeightMap(ClothFacade, Properties, NumLodSimVertices);
 
 	const int32 NumRenderVertices = ClothFacade.GetNumRenderVertices();
 	LODModel.MeshToImportVertexMap.Reserve(NumRenderVertices);
diff --git a/Engine/Plugins/ChaosClothAsset/Source/ChaosClothAssetEngine/Private/ChaosClothAsset/ClothAssetSKMClothingAsset.cpp b/Engine/Plugins/ChaosClothAsset/Source/ChaosClothAssetEngine/Private/ChaosClothAsset/ClothAssetSKMClothingAsset.cpp
new file mode 100644
--- /dev/null
+++ b/Engine/Plugins/ChaosClothAsset/Source/ChaosClothAssetEngine/Private/ChaosClothAsset/ClothAssetSKMClothingAsset.cpp
@@ -0,0 +1,182 @@
+// Copyright Epic Games, Inc. All Rights Reserved.
+
+#include "ChaosClothAsset/ClothAssetSKMClothingAsset.h"
+#include "ChaosClothAsset/ClothAssetBase.h"
+#include "ChaosClothAsset/ClothAssetPrivate.h"
+
+#if WITH_EDITOR
+#include "Chaos/CollectionPropertyFacade.h"
+#include "ChaosClothAsset/ClothEngineTools.h"
+#include "ChaosClothAsset/ClothSimulationModel.h"
+#include "ChaosClothAsset/CollectionClothFacade.h"
+#include "Engine/SkeletalMesh.h"
+#include "Utils/ClothingMeshUtils.h"
+#include "Rendering/SkeletalMeshModel.h"
+#endif  // #if WITH_EDITOR
+
+#include UE_INLINE_GENERATED_CPP_BY_NAME(ClothAssetSKMClothingAsset)
+
+#define LOCTEXT_NAMESPACE "ChaosClothAssetSKMClothingAsset"
+
+UChaosClothAssetSKMClothingAsset::UChaosClothAssetSKMClothingAsset(const FObjectInitializer& ObjectInitializer)
+	: Super(ObjectInitializer)
+{
+}
+
+void UChaosClothAssetSKMClothingAsset::SetAsset(const UChaosClothAssetBase* InAsset)
+{
+	if (Asset != InAsset)
+	{
+		Asset = InAsset;
+		OnAssetChanged();
+	}
+}
+
+#if WITH_EDITOR
+void UChaosClothAssetSKMClothingAsset::PostEditChangeProperty(FPropertyChangedEvent& PropertyChangedEvent)
+{
+	Super::PostEditChangeProperty(PropertyChangedEvent);
+
+	if (PropertyChangedEvent.GetPropertyName() == GET_MEMBER_NAME_CHECKED(UChaosClothAssetSKMClothingAsset, Asset))
+	{
+		OnAssetChanged();
+	}
+}
+
+bool UChaosClothAssetSKMClothingAsset::BindToSkeletalMesh(USkeletalMesh* SkeletalMesh, const int32 MeshLodIndex, const int32 SectionIndex, const int32 AssetLodIndex)
+{
+	using namespace UE::Chaos::ClothAsset;
+
+	if (!ensure(SkeletalMesh) ||
+		!ensure(SkeletalMesh == GetOuter()) ||
+		!ensure(SkeletalMesh->GetImportedModel()) ||
+		!ensure(SkeletalMesh->GetImportedModel()->LODModels.IsValidIndex(MeshLodIndex)) ||
+		!ensure(SkeletalMesh->GetImportedModel()->LODModels[MeshLodIndex].Sections.IsValidIndex(SectionIndex)))
+	{
+		return false;
+	}
+	if (!Asset)
+	{
+		const FText Text = FText::Format(
+			LOCTEXT("NoValidClothSimulationModel", "Failed to bind Clothing Data [{0}], because no Asset has been set."),
+			FText::FromString(GetName()));
+		WarningNotification(Text);
+		return false;
+	}
+	const int32 ClothSimulationModelIndex = GetClothSimulationModelIndex();
+	if (ClothSimulationModelIndex == INDEX_NONE)
+	{
+		const FText Text = FText::Format(
+			LOCTEXT("NoValidClothSimulationModel", "Failed to bind Clothing Data [{0}], because [{1}] is not an existing Cloth Simulation Model of Asset [{2}]."),
+			FText::FromString(GetName()),
+			FText::FromName(ClothSimulationModelName),
+			FText::FromName(Asset->GetFName()));
+		WarningNotification(Text);
+		return false;
+	}
+	const TSharedPtr<const FChaosClothSimulationModel> ClothSimulationModel = Asset->GetClothSimulationModel(ClothSimulationModelIndex);
+	if (!ClothSimulationModel ||
+		!ClothSimulationModel->GetNumVertices(AssetLodIndex) ||
+		!ClothSimulationModel->GetNumTriangles(AssetLodIndex))
+	{
+		const FText Text = FText::Format(
+			LOCTEXT("NoValidClothSimulationModel", "Failed to bind Clothing Data [{0}], because [{1}] is an empty Cloth Simulation Model of Asset [{2}]."),
+			FText::FromString(GetName()),
+			FText::FromName(ClothSimulationModelName),
+			FText::FromName(Asset->GetFName()));
+		WarningNotification(Text);
+		return false;
+	}
+	if (!ClothSimulationModel->IsValidLodIndex(AssetLodIndex))
+	{
+		const FText Text = FText::Format(
+			LOCTEXT("NoValidClothSimulationModel", "Failed to bind Clothing Data [{0}], because Cloth Simulation Model [{1}] has no LOD{2} in Asset [{3}]."),
+			FText::FromString(GetName()),
+			FText::FromName(ClothSimulationModelName),
+			AssetLodIndex,
+			FText::FromName(Asset->GetFName()));
+		WarningNotification(Text);
+		return false;
+	}
+
+	FSkeletalMeshLODModel& SkeletalMeshLODModel = SkeletalMesh->GetImportedModel()->LODModels[MeshLodIndex];
+	FSkelMeshSection& SkelMeshSection = SkeletalMeshLODModel.Sections[SectionIndex];
+
+	// Set the Clothing Data
+	SkelMeshSection.CorrespondClothAssetIndex = ClothSimulationModelIndex;
+	SkelMeshSection.ClothingData.AssetGuid = AssetGuid;
+	SkelMeshSection.ClothingData.AssetLodIndex = AssetLodIndex;
+
+	// Calculate the deformer mappings
+	const int32 NumIndices = SkelMeshSection.NumTriangles * 3;
+	const int32 NumVertices = SkelMeshSection.SoftVertices.Num();
+	const uint32 BaseIndex = SkelMeshSection.BaseIndex;
+	const uint32 BaseVertexIndex = SkelMeshSection.BaseVertexIndex;
+
+	TArray<FVector3f> RenderPositions;
+	TArray<FVector3f> RenderNormals;
+	TArray<FVector3f> RenderTangents;
+	RenderPositions.Reserve(NumVertices);
+	RenderNormals.Reserve(NumVertices);
+	RenderTangents.Reserve(NumVertices);
+	for (const FSoftSkinVertex& SoftSkinVertex : SkelMeshSection.SoftVertices)
+	{
+		RenderPositions.Add(SoftSkinVertex.Position);
+		RenderNormals.Add(SoftSkinVertex.TangentZ);
+		RenderTangents.Add(SoftSkinVertex.TangentX);
+	}
+
+	const TConstArrayView<uint32> SectionRenderIndices(SkeletalMeshLODModel.IndexBuffer.GetData() + BaseIndex, NumIndices);
+	TArray<uint32> RenderIndices;
+	RenderIndices.Reserve(SectionRenderIndices.Num());
+	for (const uint32 SectionRenderIndex : SectionRenderIndices)
+	{
+		RenderIndices.Add(SectionRenderIndex - BaseVertexIndex);
+	}
+
+	const ClothingMeshUtils::ClothMeshDesc TargetMesh(RenderPositions, RenderNormals, RenderTangents, RenderIndices);
+	const ClothingMeshUtils::ClothMeshDesc SourceMesh(
+		ClothSimulationModel->GetPositions(AssetLodIndex),
+		ClothSimulationModel->GetIndices(AssetLodIndex));  // Let it calculate the averaged normals as to match the simulation data output
+
+	const TSharedRef<const FManagedArrayCollection>& Collection = Asset->GetCollections(ClothSimulationModelIndex)[AssetLodIndex];
+	FCollectionClothConstFacade ClothFacade(Collection);
+	::Chaos::Softs::FCollectionPropertyConstFacade PropertyFacade(Collection);
+	const FPointWeightMap MaxDistances = FClothEngineTools::GetMaxDistanceWeightMap(ClothFacade, PropertyFacade, ClothSimulationModel->GetNumVertices(AssetLodIndex));
+
+	constexpr bool bSmoothTransition = true;
+	constexpr bool bUseMultipleInfluences = false;
+	constexpr float SkinningKernelRadius = 100.f;
+	TArray<FMeshToMeshVertData> MeshToMeshData;
+	ClothingMeshUtils::GenerateMeshToMeshVertData(
+		MeshToMeshData, TargetMesh, SourceMesh, &MaxDistances,
+		bSmoothTransition, bUseMultipleInfluences, SkinningKernelRadius);
+
+	SkelMeshSection.ClothMappingDataLODs.SetNum(1);
+	SkelMeshSection.ClothMappingDataLODs[0] = MoveTemp(MeshToMeshData);
+
+	return true;
+}
+#endif  // #if WITH_EDITOR
+
+TArray<FName> UChaosClothAssetSKMClothingAsset::GetClothSimulationModelNames() const
+{
+	TArray<FName> SimulationModels;
+	if (Asset)
+	{
+		const int32 NumSimulationModels = Asset->GetNumClothSimulationModels();
+		SimulationModels.Reserve(NumSimulationModels);
+		for (int32 Index = 0; Index < NumSimulationModels; ++Index)
+		{
+			SimulationModels.Add(Asset->GetClothSimulationModelName(Index));
+		}
+	}
+	return SimulationModels;
+}
+
+void UChaosClothAssetSKMClothingAsset::OnAssetChanged()
+{
+	ClothSimulationModelName = Asset && Asset->GetNumClothSimulationModels() ? Asset->GetClothSimulationModelName(0) : NAME_None;
+}
+
+#undef LOCTEXT_NAMESPACE
diff --git a/Engine/Plugins/ChaosClothAsset/Source/ChaosClothAssetEngine/Private/ChaosClothAsset/ClothAssetSKMClothingAsset.h b/Engine/Plugins/ChaosClothAsset/Source/ChaosClothAssetEngine/Private/ChaosClothAsset/ClothAssetSKMClothingAsset.h
new file mode 100644
--- /dev/null
+++ b/Engine/Plugins/ChaosClothAsset/Source/ChaosClothAssetEngine/Private/ChaosClothAsset/ClothAssetSKMClothingAsset.h
@@ -0,0 +1,75 @@
+// Copyright Epic Games, Inc. All Rights Reserved.
+
+#pragma once
+
+#include "ClothingAssetBase.h"
+#include "ClothAssetSKMClothingAsset.generated.h"
+
+struct FManagedArrayCollection;
+class UChaosClothAssetBase;
+
+/**
+ * Clothing Data implementation for the Cloth/Outfit Asset.
+ */
+UCLASS(BlueprintType, MinimalAPI)
+class UChaosClothAssetSKMClothingAsset final : public UClothingAssetBase
+{
+	GENERATED_BODY()
+public:
+	UChaosClothAssetSKMClothingAsset(const FObjectInitializer& ObjectInitializer);
+
+	void SetAsset(const UChaosClothAssetBase* InAsset);
+
+private:
+#if WITH_EDITOR
+	//~ Begin UObject Interface
+	virtual void PostEditChangeProperty(FPropertyChangedEvent& PropertyChangedEvent) override;
+	//~ EnUObject Interface
+
+	//~ Begin UClothingAssetBase Interface
+	virtual bool BindToSkeletalMesh(USkeletalMesh* SkeletalMesh, const int32 MeshLodIndex, const int32 SectionIndex, const int32 AssetLodIndex) override;
+
+	virtual void UnbindFromSkeletalMesh(USkeletalMesh* SkeletalMesh) override
+	{
+		// TODO
+	}
+
+	virtual void UnbindFromSkeletalMesh(USkeletalMesh* SkeletalMesh, const int32 MeshLodIndex) override
+	{
+		// TODO
+	}
+
+	virtual void UpdateAllLODBiasMappings(USkeletalMesh* SkeletalMesh) override
+	{
+		// TODO
+	}
+#endif  // #if WITH_EDITOR
+
+	virtual void RefreshBoneMapping(USkeletalMesh* SkeletalMesh) override
+	{
+	}
+
+	virtual void PostUpdateAllAssets() override
+	{
+		// TODO
+	}
+	//~ End UClothingAssetBase Interface
+
+	UFUNCTION(CallInEditor)
+	TArray<FName> GetClothSimulationModelNames() const;
+
+	int32 GetClothSimulationModelIndex() const
+	{
+		return GetClothSimulationModelNames().Find(ClothSimulationModelName);
+	}
+
+	void OnAssetChanged();
+
+	/** Cloth or Outfit Asset to use with this Skeletal Mesh. */
+	UPROPERTY(EditAnywhere, Setter = "SetAsset", Category = "Chaos Cloth Asset")
+	TObjectPtr<const UChaosClothAssetBase> Asset;
+
+	/** The Simulation Model to use for this Clothing Data. Useful for Outfit Assets that may contain multiple models. */
+	UPROPERTY(EditAnywhere, Category = "Chaos Cloth Asset", Meta = (DisplayName = "Cloth Simulation Model", GetOptions = "GetClothSimulationModelNames"))
+	FName ClothSimulationModelName;
+};
diff --git a/Engine/Plugins/ChaosClothAsset/Source/ChaosClothAssetEngine/Private/ChaosClothAsset/ClothAssetSKMClothingAssetFactory.cpp b/Engine/Plugins/ChaosClothAsset/Source/ChaosClothAssetEngine/Private/ChaosClothAsset/ClothAssetSKMClothingAssetFactory.cpp
new file mode 100644
--- /dev/null
+++ b/Engine/Plugins/ChaosClothAsset/Source/ChaosClothAssetEngine/Private/ChaosClothAsset/ClothAssetSKMClothingAssetFactory.cpp
@@ -0,0 +1,31 @@
+// Copyright Epic Games, Inc. All Rights Reserved.
+
+#include "ChaosClothAsset/ClothAssetSKMClothingAssetFactory.h"
+#include "ChaosClothAsset/ClothAssetBase.h"
+#include "Engine/SkeletalMesh.h"
+#include "ChaosClothAsset/ClothAssetSKMClothingAsset.h"
+
+#include UE_INLINE_GENERATED_CPP_BY_NAME(ClothAssetSKMClothingAssetFactory)
+
+UChaosClothAssetSKMClothingAssetFactory::UChaosClothAssetSKMClothingAssetFactory(const FObjectInitializer& ObjectInitializer)
+	: Super(ObjectInitializer)
+{
+}
+
+const UClass* UChaosClothAssetSKMClothingAssetFactory::GetSupportedSourceAssetType() const
+{
+	return UChaosClothAssetBase::StaticClass();
+}
+
+TArray<UClothingAssetBase*> UChaosClothAssetSKMClothingAssetFactory::CreateFromSourceAsset(USkeletalMesh* TargetMesh, const UObject* SourceAsset) const
+{
+	TArray<UClothingAssetBase*> ClothingAssets;
+	if (const UChaosClothAssetBase* const ClothAsset = Cast<UChaosClothAssetBase>(SourceAsset))
+	{
+		const FName AssetName = MakeUniqueObjectName(TargetMesh, UChaosClothAssetSKMClothingAsset::StaticClass(), ClothAsset->GetFName());
+		UChaosClothAssetSKMClothingAsset* const ClothAssetSKMClothingAsset = NewObject<UChaosClothAssetSKMClothingAsset>(TargetMesh, AssetName);
+		ClothAssetSKMClothingAsset->SetAsset(Cast<UChaosClothAssetBase>(SourceAsset));
+		ClothingAssets.Emplace(ClothAssetSKMClothingAsset);
+	}
+	return ClothingAssets;
+}
diff --git a/Engine/Plugins/ChaosClothAsset/Source/ChaosClothAssetEngine/Private/ChaosClothAsset/ClothAssetSKMClothingAssetFactory.h b/Engine/Plugins/ChaosClothAsset/Source/ChaosClothAssetEngine/Private/ChaosClothAsset/ClothAssetSKMClothingAssetFactory.h
new file mode 100644
--- /dev/null
+++ b/Engine/Plugins/ChaosClothAsset/Source/ChaosClothAssetEngine/Private/ChaosClothAsset/ClothAssetSKMClothingAssetFactory.h
@@ -0,0 +1,35 @@
+// Copyright Epic Games, Inc. All Rights Reserved.
+
+#pragma once
+
+#include "ClothingAssetFactoryInterface.h"
+#include "ClothAssetSKMClothingAssetFactory.generated.h"
+
+UCLASS(MinimalAPI)
+class UChaosClothAssetSKMClothingAssetFactory : public UClothingAssetFactoryBase
+{
+	GENERATED_BODY()
+public:
+	UChaosClothAssetSKMClothingAssetFactory(const FObjectInitializer& ObjectInitializer);
+
+	//~ Begin UClothingAssetFactoryBase Interface
+	virtual const UClass* GetSupportedSourceAssetType() const override;
+
+	virtual TArray<UClothingAssetBase*> CreateFromSourceAsset(USkeletalMesh* TargetMesh, const UObject* SourceAsset) const override;
+
+	virtual UClothingAssetBase* CreateFromSkeletalMesh(USkeletalMesh* TargetMesh, FSkeletalMeshClothBuildParams& Params) override
+	{
+		return nullptr;  // TODO: Implement create from skeletal mesh
+	}
+
+	virtual UClothingAssetBase* CreateFromExistingCloth(USkeletalMesh* /*TargetMesh*/, USkeletalMesh* /*SourceMesh*/, UClothingAssetBase* /*SourceAsset*/) override
+	{
+		return nullptr;  // TODO: Implement create from another clothing data
+	}
+
+	virtual UClothingAssetBase* ImportLodToClothing(USkeletalMesh* TargetMesh, FSkeletalMeshClothBuildParams& Params) override
+	{
+		return nullptr;
+	}
+	//~ End UClothingAssetFactoryBase Interface
+};
diff --git a/Engine/Plugins/ChaosClothAsset/Source/ChaosClothAssetEngine/Private/ChaosClothAsset/ClothAssetSKMSimulationFactory.cpp b/Engine/Plugins/ChaosClothAsset/Source/ChaosClothAssetEngine/Private/ChaosClothAsset/ClothAssetSKMSimulationFactory.cpp
new file mode 100644
--- /dev/null
+++ b/Engine/Plugins/ChaosClothAsset/Source/ChaosClothAssetEngine/Private/ChaosClothAsset/ClothAssetSKMSimulationFactory.cpp
@@ -0,0 +1,37 @@
+// Copyright Epic Games, Inc. All Rights Reserved.
+
+#include "ChaosClothAsset/ClothAssetSKMSimulationFactory.h"
+#include "ChaosCloth/ChaosClothingSimulation.h"
+#include "ChaosClothAsset/ClothAssetSKMClothingAsset.h"
+
+#include UE_INLINE_GENERATED_CPP_BY_NAME(ClothAssetSKMSimulationFactory)
+
+#define LOCTEXT_NAMESPACE "ChaosClothAssetSKMSimulationFactory"
+
+IClothingSimulation* UChaosClothAssetSKMSimulationFactory::CreateSimulation()
+{
+	IClothingSimulation* Simulation = nullptr;  // new Chaos::FClothingSimulation();
+	return Simulation;
+}
+
+void UChaosClothAssetSKMSimulationFactory::DestroySimulation(IClothingSimulation* InSimulation)
+{
+    delete InSimulation;
+}
+
+bool UChaosClothAssetSKMSimulationFactory::SupportsAsset(UClothingAssetBase* InAsset)
+{
+    return Cast<UChaosClothAssetSKMClothingAsset>(InAsset) != nullptr;
+}
+
+bool UChaosClothAssetSKMSimulationFactory::SupportsRuntimeInteraction()
+{
+    return false;  // TODO
+}
+
+UClothingSimulationInteractor* UChaosClothAssetSKMSimulationFactory::CreateInteractor()
+{
+	return nullptr;// NewObject<UChaosClothAssetSimulationInteractor>(GetTransientPackage());  // TODO
+}
+
+#undef LOCTEXT_NAMESPACE
diff --git a/Engine/Plugins/ChaosClothAsset/Source/ChaosClothAssetEngine/Private/ChaosClothAsset/ClothAssetSKMSimulationFactory.h b/Engine/Plugins/ChaosClothAsset/Source/ChaosClothAssetEngine/Private/ChaosClothAsset/ClothAssetSKMSimulationFactory.h
new file mode 100644
--- /dev/null
+++ b/Engine/Plugins/ChaosClothAsset/Source/ChaosClothAssetEngine/Private/ChaosClothAsset/ClothAssetSKMSimulationFactory.h
@@ -0,0 +1,33 @@
+// Copyright Epic Games, Inc. All Rights Reserved.
+
+#pragma once
+
+#include "ClothingSimulationFactory.h"
+#include "ClothAssetSKMSimulationFactory.generated.h"
+
+class UClothConfigBase;
+
+UCLASS(MinimalAPI)
+class UChaosClothAssetSKMSimulationFactory final : public UClothingSimulationFactory
+{
+	GENERATED_BODY()
+public:
+	//~ Begin UClothingSimulationFactory Interface
+	virtual IClothingSimulation* CreateSimulation() override;
+	virtual void DestroySimulation(IClothingSimulation* InSimulation) override;
+	virtual bool SupportsAsset(UClothingAssetBase* InAsset) override;
+
+	virtual bool SupportsRuntimeInteraction() override;
+	virtual UClothingSimulationInteractor* CreateInteractor() override;
+
+	virtual TArrayView<const TSubclassOf<UClothConfigBase>> GetClothConfigClasses() const override
+	{
+		return TArrayView<const TSubclassOf<UClothConfigBase>>();
+	}
+
+	const UEnum* GetWeightMapTargetEnum() const override
+	{
+		return nullptr;
+	}
+	//~ End UClothingSimulationFactory Interface
+};
diff --git a/Engine/Plugins/ChaosClothAsset/Source/ChaosClothAssetEngine/Private/ChaosClothAsset/ClothEngineTools.cpp b/Engine/Plugins/ChaosClothAsset/Source/ChaosClothAssetEngine/Private/ChaosClothAsset/ClothEngineTools.cpp
--- a/Engine/Plugins/ChaosClothAsset/Source/ChaosClothAssetEngine/Private/ChaosClothAsset/ClothEngineTools.cpp
+++ b/Engine/Plugins/ChaosClothAsset/Source/ChaosClothAssetEngine/Private/ChaosClothAsset/ClothEngineTools.cpp
@@ -1,6 +1,7 @@
 // Copyright Epic Games, Inc. All Rights Reserved.
 
 #include "ChaosClothAsset/ClothEngineTools.h"
+#include "Chaos/CollectionPropertyFacade.h"
 #include "ChaosClothAsset/CollectionClothFacade.h"
 #include "ChaosClothAsset/CollectionClothSelectionFacade.h"
 #include "ChaosClothAsset/ClothCollectionGroup.h"
@@ -122,4 +123,26 @@ void FClothEngineTools::GenerateTethersFromCustomSelectionSets(const TSharedRef<
 	}
 }
 
+FPointWeightMap FClothEngineTools::GetMaxDistanceWeightMap(const TSharedRef<const FManagedArrayCollection>& ClothCollection, const int32 NumLodSimVertices)
+{
+	FCollectionClothConstFacade ClothFacade(ClothCollection);
+	::Chaos::Softs::FCollectionPropertyConstFacade PropertyFacade(ClothCollection);
+	return GetMaxDistanceWeightMap(ClothFacade, PropertyFacade, NumLodSimVertices);
+}
+
+FPointWeightMap FClothEngineTools::GetMaxDistanceWeightMap(const FCollectionClothConstFacade& ClothFacade, const ::Chaos::Softs::FCollectionPropertyConstFacade& PropertyFacade, const int32 NumLodSimVertices)
+{
+	int32 MaxDistancePropertyKeyIndex;
+	FString MaxDistanceString = TEXT("MaxDistance");
+	MaxDistanceString = PropertyFacade.GetStringValue(MaxDistanceString, MaxDistanceString, &MaxDistancePropertyKeyIndex);
+	const bool bHasMaxDistanceProperty = (MaxDistancePropertyKeyIndex != INDEX_NONE);
+	const float MaxDistanceOffset = bHasMaxDistanceProperty ? PropertyFacade.GetLowValue<float>(MaxDistancePropertyKeyIndex) : TNumericLimits<float>::Max();  // Uses infinite distance when no MaxDistance properties are set
+	const float MaxDistanceScale = bHasMaxDistanceProperty ? PropertyFacade.GetHighValue<float>(MaxDistancePropertyKeyIndex) - MaxDistanceOffset : 0.f;
+	const TConstArrayView<float> MaxDistanceWeightMap = ClothFacade.GetWeightMap(FName(MaxDistanceString));
+
+	return (MaxDistanceWeightMap.Num() == NumLodSimVertices) ?
+		FPointWeightMap(MaxDistanceWeightMap, MaxDistanceOffset, MaxDistanceScale) :
+		FPointWeightMap(NumLodSimVertices, MaxDistanceOffset);
+}
+
 }  // End namespace UE::Chaos::ClothAsset
diff --git a/Engine/Plugins/ChaosClothAsset/Source/ChaosClothAssetEngine/Public/ChaosClothAsset/ClothEngineTools.h b/Engine/Plugins/ChaosClothAsset/Source/ChaosClothAssetEngine/Public/ChaosClothAsset/ClothEngineTools.h
--- a/Engine/Plugins/ChaosClothAsset/Source/ChaosClothAssetEngine/Public/ChaosClothAsset/ClothEngineTools.h
+++ b/Engine/Plugins/ChaosClothAsset/Source/ChaosClothAssetEngine/Public/ChaosClothAsset/ClothEngineTools.h
@@ -2,17 +2,27 @@
 
 #pragma once
 
-#include "Templates/SharedPointer.h"
 #include "Math/Vector2D.h"
+#include "Templates/SharedPointer.h"
+#include "PointWeightMap.h"
 
 #define UE_API CHAOSCLOTHASSETENGINE_API
 
 struct FManagedArrayCollection;
 class FName;
 class FSkeletalMeshLODModel;
+namespace Chaos
+{
+	namespace Softs
+	{
+		class FCollectionPropertyConstFacade;
+	}
+}
 
 namespace UE::Chaos::ClothAsset
 {
+	class FCollectionClothConstFacade;
+
 	/**
 	 *  Tools operating on cloth collections with Engine dependency
 	 */
@@ -23,6 +33,11 @@ namespace UE::Chaos::ClothAsset
 		static UE_API void GenerateTethersFromSelectionSet(const TSharedRef<FManagedArrayCollection>& ClothCollection, const FName& FixedEndSet, const bool bGeodesicTethers);
 		/** @param CustomTetherEndSets: First element of each pair is DynamicSet, second is FixedSet */
 		static UE_API void GenerateTethersFromCustomSelectionSets(const TSharedRef<FManagedArrayCollection>& ClothCollection, const FName& FixedEndSet, const TArray<TPair<FName, FName>>& CustomTetherEndSets, const bool bGeodesicTethers);
+	
+		/** Retrieve the MaxDistance weight map from the cloth facades. Return an empty map if the collection doesn't match the requested NumLodSimVertices.*/
+		static FPointWeightMap GetMaxDistanceWeightMap(const FCollectionClothConstFacade& ClothFacade, const ::Chaos::Softs::FCollectionPropertyConstFacade& PropertyFacade, const int32 NumLodSimVertices);
+		/** Retrieve the MaxDistance weight map from the cloth collection. Return an empty map if the collection doesn't match the requested NumLodSimVertices. */
+		static FPointWeightMap GetMaxDistanceWeightMap(const TSharedRef<const FManagedArrayCollection>& ClothCollection, const int32 NumLodSimVertices);
 	};
 }  // End namespace UE::Chaos::ClothAsset
 
diff --git a/Engine/Source/Runtime/ClothingSystemRuntimeCommon/Private/ClothingAsset.cpp b/Engine/Source/Runtime/ClothingSystemRuntimeCommon/Private/ClothingAsset.cpp
--- a/Engine/Source/Runtime/ClothingSystemRuntimeCommon/Private/ClothingAsset.cpp
+++ b/Engine/Source/Runtime/ClothingSystemRuntimeCommon/Private/ClothingAsset.cpp
@@ -33,7 +33,6 @@
 
 #include UE_INLINE_GENERATED_CPP_BY_NAME(ClothingAsset)
 
-DEFINE_LOG_CATEGORY(LogClothingAsset)
 #define LOCTEXT_NAMESPACE "ClothingAsset"
 
 //==============================================================================
@@ -126,16 +125,6 @@ UClothingAssetCommon::UClothingAssetCommon(const FObjectInitializer& ObjectIniti
 }
 
 #if WITH_EDITOR
-
-void Warn(const FText& Error)
-{
-	FNotificationInfo* const NotificationInfo = new FNotificationInfo(Error);
-	NotificationInfo->ExpireDuration = 5.0f;
-	FSlateNotificationManager::Get().QueueNotification(NotificationInfo);
-
-	UE_LOG(LogClothingAsset, Warning, TEXT("%s"), *Error.ToString());
-}
-
 bool UClothingAssetCommon::BindToSkeletalMesh(
 	USkeletalMesh* InSkelMesh, 
 	const int32 InMeshLodIndex, 
@@ -158,43 +147,43 @@ bool UClothingAssetCommon::BindToSkeletalMesh(
 	// If we've been added to the wrong mesh
 	if(InSkelMesh != GetOuter())
 	{
-		FText Error = FText::Format(
+		const FText Error = FText::Format(
 			LOCTEXT("Error_WrongMesh", "Failed to bind clothing asset {0} as the provided mesh is not the owner of this asset."), 
 			FText::FromString(GetName()));
-		Warn(Error);
+		WarningNotification(Error);
 		return false;
 	}
 
 	// If we don't have clothing data
 	if(!LodData.IsValidIndex(InAssetLodIndex))
 	{
-		FText Error = FText::Format(
+		const FText Error = FText::Format(
 			LOCTEXT("Error_NoClothingLod", "Failed to bind clothing asset {0} LOD{1} as LOD{2} does not exist."), 
 			FText::FromString(GetName()), 
 			InAssetLodIndex,
 			InAssetLodIndex);
-		Warn(Error);
+		WarningNotification(Error);
 		return false;
 	}
 
 	// If we don't have a mesh
 	if(!InSkelMesh)
 	{
-		FText Error = FText::Format(
+		const FText Error = FText::Format(
 			LOCTEXT("Error_NoMesh", "Failed to bind clothing asset {0} as provided skel mesh does not exist."), 
 			FText::FromString(GetName()));
-		Warn(Error);
+		WarningNotification(Error);
 		return false;
 	}
 
 	// If the mesh LOD index is invalid
 	if(!InSkelMesh->GetImportedModel()->LODModels.IsValidIndex(InMeshLodIndex))
 	{
-		FText Error = FText::Format(
+		const FText Error = FText::Format(
 			LOCTEXT("Error_InvalidMeshLOD", "Failed to bind clothing asset {0} as mesh LOD{1} does not exist."), 
 			FText::FromString(GetName()), 
 			InMeshLodIndex);
-		Warn(Error);
+		WarningNotification(Error);
 		return false;
 	}
 
@@ -204,13 +193,13 @@ bool UClothingAssetCommon::BindToSkeletalMesh(
 		const int32& MappedLod = LodMap[MapIndex];
 		if(MappedLod == InAssetLodIndex)
 		{
-			FText Error = FText::Format(
+			const FText Error = FText::Format(
 				LOCTEXT("Error_LodMapped", "Failed to bind clothing asset {0} LOD{1} as LOD{2} is already mapped to mesh LOD{3}."), 
 				FText::FromString(GetName()), 
 				InAssetLodIndex, 
 				InAssetLodIndex, 
 				MapIndex);
-			Warn(Error);
+			WarningNotification(Error);
 			return false;
 		}
 	}
@@ -284,13 +273,13 @@ bool UClothingAssetCommon::BindToSkeletalMesh(
 	if(TempBoneMap.Num() > FGPUBaseSkinVertexFactory::GetMaxGPUSkinBones())
 	{
 		// Failed to apply as we've exceeded the number of bones we can skin
-		FText Error = FText::Format(
+		const FText Error = FText::Format(
 			LOCTEXT("Error_TooManyBones", "Failed to bind clothing asset {0} LOD{1} as this causes the section to require {2} bones. The maximum per section is currently {3}."), 
 			FText::FromString(GetName()), 
 			InAssetLodIndex, 
 			TempBoneMap.Num(), 
 			FGPUBaseSkinVertexFactory::GetMaxGPUSkinBones());
-		Warn(Error);
+		WarningNotification(Error);
 		return false;
 	}
 
@@ -537,7 +526,7 @@ void UClothingAssetCommon::UpdateLODBiasMappings(const USkeletalMesh* SkeletalMe
 					SectionIndex,
 					UpdatedLODIndex, 
 					LODIndex);
-				Warn(Error);
+				WarningNotification(Error);
 				continue;
 			}
 
@@ -625,11 +614,11 @@ void UClothingAssetCommon::UnbindFromSkeletalMesh(
 	{
 		if(!Mesh->LODModels.IsValidIndex(InMeshLodIndex))
 		{
-			FText Error = FText::Format(
+			const FText Error = FText::Format(
 				LOCTEXT("Error_UnbindNoMeshLod", "Failed to remove clothing asset {0} from mesh LOD{1} as that LOD doesn't exist."), 
 				FText::FromString(GetName()), 
 				InMeshLodIndex);
-			Warn(Error);
+			WarningNotification(Error);
 
 			return;
 		}
diff --git a/Engine/Source/Runtime/ClothingSystemRuntimeCommon/Public/ClothingAsset.h b/Engine/Source/Runtime/ClothingSystemRuntimeCommon/Public/ClothingAsset.h
--- a/Engine/Source/Runtime/ClothingSystemRuntimeCommon/Public/ClothingAsset.h
+++ b/Engine/Source/Runtime/ClothingSystemRuntimeCommon/Public/ClothingAsset.h
@@ -10,8 +10,6 @@
 #include "Algo/AnyOf.h"
 #include "ClothingAsset.generated.h"
 
-DECLARE_LOG_CATEGORY_EXTERN(LogClothingAsset, Log, All);
-
 class UPhysicsAsset;
 struct FSkelMeshSection;
 
diff --git a/Engine/Source/Runtime/ClothingSystemRuntimeInterface/ClothingSystemRuntimeInterface.Build.cs b/Engine/Source/Runtime/ClothingSystemRuntimeInterface/ClothingSystemRuntimeInterface.Build.cs
--- a/Engine/Source/Runtime/ClothingSystemRuntimeInterface/ClothingSystemRuntimeInterface.Build.cs
+++ b/Engine/Source/Runtime/ClothingSystemRuntimeInterface/ClothingSystemRuntimeInterface.Build.cs
@@ -16,5 +16,10 @@ public class ClothingSystemRuntimeInterface : ModuleRules
                 "CoreUObject"
             }
         );
-    }
+
+		if (Target.bBuildEditor || Target.bCompileAgainstEditor)
+		{
+			PrivateDependencyModuleNames.Add("Slate");  // For editor notifications
+		}
+	}
 }
diff --git a/Engine/Source/Runtime/ClothingSystemRuntimeInterface/Private/ClothingAssetBase.cpp b/Engine/Source/Runtime/ClothingSystemRuntimeInterface/Private/ClothingAssetBase.cpp
new file mode 100644
--- /dev/null
+++ b/Engine/Source/Runtime/ClothingSystemRuntimeInterface/Private/ClothingAssetBase.cpp
@@ -0,0 +1,23 @@
+// Copyright Epic Games, Inc. All Rights Reserved.
+
+#include "ClothingAssetBase.h"
+
+#if WITH_EDITOR
+#include "Framework/Notifications/NotificationManager.h"
+#include "Widgets/Notifications/SNotificationList.h"
+#endif
+
+#include UE_INLINE_GENERATED_CPP_BY_NAME(ClothingAssetBase)
+
+DEFINE_LOG_CATEGORY(LogClothingAsset)
+
+#if WITH_EDITOR
+void UClothingAssetBase::WarningNotification(const FText& Text)
+{
+	FNotificationInfo* const NotificationInfo = new FNotificationInfo(Text);
+	NotificationInfo->ExpireDuration = 5.0f;
+	FSlateNotificationManager::Get().QueueNotification(NotificationInfo);
+
+	UE_LOG(LogClothingAsset, Warning, TEXT("%s"), *Text.ToString());
+}
+#endif
diff --git a/Engine/Source/Runtime/ClothingSystemRuntimeInterface/Public/ClothingAssetBase.h b/Engine/Source/Runtime/ClothingSystemRuntimeInterface/Public/ClothingAssetBase.h
--- a/Engine/Source/Runtime/ClothingSystemRuntimeInterface/Public/ClothingAssetBase.h
+++ b/Engine/Source/Runtime/ClothingSystemRuntimeInterface/Public/ClothingAssetBase.h
@@ -8,6 +8,10 @@
 
 class USkeletalMesh;
 
+#define UE_API CLOTHINGSYSTEMRUNTIMEINTERFACE_API
+
+UE_API DECLARE_LOG_CATEGORY_EXTERN(LogClothingAsset, Log, All);
+
 /**
  * An interface object for any clothing asset the engine can use.
  * Any clothing asset concrete object should derive from this.
@@ -54,7 +58,9 @@ public:
 	/** Add a new LOD class instance. */
 	UE_DEPRECATED(5.7, "Use UClothingAssetCommon::AddNewLod() instead.")
 	virtual int32 AddNewLod()
-	PURE_VIRTUAL(UClothingAssetBase::AddNewLod(), return INDEX_NONE;);
+	{
+		return INDEX_NONE;
+	}
 
 	/**
 	*	Builds the LOD transition data
@@ -62,8 +68,7 @@ public:
 	*	in exactly the same way the render mesh is skinned to create a smooth swap
 	*/
 	UE_DEPRECATED(5.7, "Use UClothingAssetCommon::BuildLodTransitionData() instead.")
-	virtual void BuildLodTransitionData()
-	PURE_VIRTUAL(UClothingAssetBase::BuildLodTransitionData(), );
+	virtual void BuildLodTransitionData() {}
 #endif
 
 #if WITH_EDITORONLY_DATA
@@ -72,8 +77,7 @@ public:
 	 * has changed, so any intermediate generated data can be regenerated.
 	 */
 	UE_DEPRECATED(5.7, "Use UClothingAssetCommon::InvalidateAllCachedData() instead.")
-	virtual void InvalidateAllCachedData()
-	PURE_VIRTUAL(UClothingAssetBase::InvalidateAllCachedData(),);
+	virtual void InvalidateAllCachedData() {}
 #endif // WITH_EDITORONLY_DATA
 
 	/** 
@@ -88,12 +92,16 @@ public:
 	/** Check the validity of a LOD index */
 	UE_DEPRECATED(5.7, "Use UClothingAssetCommon::IsValidLod() instead.")
 	virtual bool IsValidLod(int32 InLodIndex) const
-	PURE_VIRTUAL(UClothingAssetBase::IsValidLod(), return false;);
+	{
+		return false;
+	}
 
 	/** Get the number of LODs defined in the clothing asset */
 	UE_DEPRECATED(5.7, "Use UClothingAssetCommon::GetNumLods() instead.")
 	virtual int32 GetNumLods() const
-	PURE_VIRTUAL(UClothingAssetBase::GetNumLods(), return 0;);
+	{
+		return 0;
+	}
 
 	/** Called after all cloth assets sharing the same simulation are added or loaded */
 	virtual void PostUpdateAllAssets()
@@ -116,7 +124,14 @@ protected:
 	/** The asset factory should have access, as it will assign the asset guid when building assets */
 	friend class UClothingAssetFactory;
 
+#if WITH_EDITOR
+	/** Warn using slate's notification manager. Use this method to notify binding issues back to the user. */
+	static UE_API void WarningNotification(const FText& Text);
+#endif
+
 	/** Guid to identify this asset. Will be embedded into chunks that are created using this asset */
 	UPROPERTY()
 	FGuid AssetGuid;
 };
+
+#undef UE_API
