CIS - Fixed Text conflict from LOCTEXT macro for namespace 'ChaosClothAssetSKMClothingAsset' and key 'NoValidClothSimulationModel'

#rnx

diff --git a/Engine/Plugins/ChaosClothAsset/Source/ChaosClothAssetEngine/Private/ChaosClothAsset/ClothAssetSKMClothingAsset.cpp b/Engine/Plugins/ChaosClothAsset/Source/ChaosClothAssetEngine/Private/ChaosClothAsset/ClothAssetSKMClothingAsset.cpp
--- a/Engine/Plugins/ChaosClothAsset/Source/ChaosClothAssetEngine/Private/ChaosClothAsset/ClothAssetSKMClothingAsset.cpp
+++ b/Engine/Plugins/ChaosClothAsset/Source/ChaosClothAssetEngine/Private/ChaosClothAsset/ClothAssetSKMClothingAsset.cpp
@@ -58,7 +58,7 @@ bool UChaosClothAssetSKMClothingAsset::BindToSkeletalMesh(USkeletalMesh* Skeleta
 	if (!Asset)
 	{
 		const FText Text = FText::Format(
-			LOCTEXT("NoValidClothSimulationModel", "Failed to bind Clothing Data [{0}], because no Asset has been set."),
+			LOCTEXT("NoAssetHasBeenSet", "Failed to bind Clothing Data [{0}], because no Asset has been set."),
 			FText::FromString(GetName()));
 		WarningNotification(Text);
 		return false;
@@ -67,7 +67,7 @@ bool UChaosClothAssetSKMClothingAsset::BindToSkeletalMesh(USkeletalMesh* Skeleta
 	if (ClothSimulationModelIndex == INDEX_NONE)
 	{
 		const FText Text = FText::Format(
-			LOCTEXT("NoValidClothSimulationModel", "Failed to bind Clothing Data [{0}], because [{1}] is not an existing Cloth Simulation Model of Asset [{2}]."),
+			LOCTEXT("NotAClothSimulationModelOfAsset", "Failed to bind Clothing Data [{0}], because [{1}] is not an existing Cloth Simulation Model of Asset [{2}]."),
 			FText::FromString(GetName()),
 			FText::FromName(ClothSimulationModelName),
 			FText::FromName(Asset->GetFName()));
@@ -80,7 +80,7 @@ bool UChaosClothAssetSKMClothingAsset::BindToSkeletalMesh(USkeletalMesh* Skeleta
 		!ClothSimulationModel->GetNumTriangles(AssetLodIndex))
 	{
 		const FText Text = FText::Format(
-			LOCTEXT("NoValidClothSimulationModel", "Failed to bind Clothing Data [{0}], because [{1}] is an empty Cloth Simulation Model of Asset [{2}]."),
+			LOCTEXT("EmptyClothSimulationModel", "Failed to bind Clothing Data [{0}], because [{1}] is an empty Cloth Simulation Model of Asset [{2}]."),
 			FText::FromString(GetName()),
 			FText::FromName(ClothSimulationModelName),
 			FText::FromName(Asset->GetFName()));
@@ -90,7 +90,7 @@ bool UChaosClothAssetSKMClothingAsset::BindToSkeletalMesh(USkeletalMesh* Skeleta
 	if (!ClothSimulationModel->IsValidLodIndex(AssetLodIndex))
 	{
 		const FText Text = FText::Format(
-			LOCTEXT("NoValidClothSimulationModel", "Failed to bind Clothing Data [{0}], because Cloth Simulation Model [{1}] has no LOD{2} in Asset [{3}]."),
+			LOCTEXT("NoLODInClothSimulationModel", "Failed to bind Clothing Data [{0}], because Cloth Simulation Model [{1}] has no LOD{2} in Asset [{3}]."),
 			FText::FromString(GetName()),
 			FText::FromName(ClothSimulationModelName),
 			AssetLodIndex,
