Fixed nonunity compile errors

diff --git a/Engine/Plugins/ChaosClothAsset/Source/ChaosClothAsset/Private/ChaosClothAsset/CollectionClothFacade.cpp b/Engine/Plugins/ChaosClothAsset/Source/ChaosClothAsset/Private/ChaosClothAsset/CollectionClothFacade.cpp
--- a/Engine/Plugins/ChaosClothAsset/Source/ChaosClothAsset/Private/ChaosClothAsset/CollectionClothFacade.cpp
+++ b/Engine/Plugins/ChaosClothAsset/Source/ChaosClothAsset/Private/ChaosClothAsset/CollectionClothFacade.cpp
@@ -5,6 +5,7 @@
 #include "ChaosClothAsset/ClothCollectionGroup.h"
 #include "ChaosClothAsset/ClothGeometryTools.h"
 #include "Chaos/ChaosArchive.h"
+#include "UObject/SoftObjectPath.h"
 #include "UObject/UE5MainStreamObjectVersion.h"
 
 namespace UE::Chaos::ClothAsset
diff --git a/Engine/Plugins/ChaosClothAsset/Source/ChaosClothAsset/Private/ChaosClothAsset/CollectionClothRenderPatternFacade.cpp b/Engine/Plugins/ChaosClothAsset/Source/ChaosClothAsset/Private/ChaosClothAsset/CollectionClothRenderPatternFacade.cpp
--- a/Engine/Plugins/ChaosClothAsset/Source/ChaosClothAsset/Private/ChaosClothAsset/CollectionClothRenderPatternFacade.cpp
+++ b/Engine/Plugins/ChaosClothAsset/Source/ChaosClothAsset/Private/ChaosClothAsset/CollectionClothRenderPatternFacade.cpp
@@ -4,6 +4,7 @@
 #include "ChaosClothAsset/ClothCollection.h"
 #include "ChaosClothAsset/ClothCollectionGroup.h"
 #include "ChaosClothAsset/ClothGeometryTools.h"
+#include "UObject/SoftObjectPath.h"
 
 namespace UE::Chaos::ClothAsset
 {
diff --git a/Engine/Plugins/Runtime/MassGameplay/Source/MassRepresentation/Public/MassStationaryISMRepresentationFragmentDestructor.h b/Engine/Plugins/Runtime/MassGameplay/Source/MassRepresentation/Public/MassStationaryISMRepresentationFragmentDestructor.h
--- a/Engine/Plugins/Runtime/MassGameplay/Source/MassRepresentation/Public/MassStationaryISMRepresentationFragmentDestructor.h
+++ b/Engine/Plugins/Runtime/MassGameplay/Source/MassRepresentation/Public/MassStationaryISMRepresentationFragmentDestructor.h
@@ -2,6 +2,7 @@
 
 #pragma once 
 
+#include "MassEntityQuery.h"
 #include "MassObserverProcessor.h"
 #include "MassStationaryISMRepresentationFragmentDestructor.generated.h"
 
diff --git a/Engine/Source/Runtime/MassEntity/Public/MassDebuggerBreakpoints.h b/Engine/Source/Runtime/MassEntity/Public/MassDebuggerBreakpoints.h
--- a/Engine/Source/Runtime/MassEntity/Public/MassDebuggerBreakpoints.h
+++ b/Engine/Source/Runtime/MassEntity/Public/MassDebuggerBreakpoints.h
@@ -2,6 +2,7 @@
 
 #pragma once
 
+#include "HAL/Platform.h"
 #include "MassEntityMacros.h"
 
 #if WITH_MASSENTITY_DEBUG
