Chaos Cloth Asset - Added external collision support to the collision sources.

#rnx

diff --git a/Engine/Plugins/ChaosClothAsset/Source/ChaosClothAssetEngine/Private/ChaosClothAsset/CollisionSources.cpp b/Engine/Plugins/ChaosClothAsset/Source/ChaosClothAssetEngine/Private/ChaosClothAsset/CollisionSources.cpp
--- a/Engine/Plugins/ChaosClothAsset/Source/ChaosClothAssetEngine/Private/ChaosClothAsset/CollisionSources.cpp
+++ b/Engine/Plugins/ChaosClothAsset/Source/ChaosClothAssetEngine/Private/ChaosClothAsset/CollisionSources.cpp
@@ -84,13 +84,28 @@ namespace UE::Chaos::ClothAsset
 		}
 	}
 
+	void FCollisionSources::AddExternalCollisions(const FClothCollisionData& InExternalCollisions)
+	{
+		ExternalCollisions = InExternalCollisions;
+	}
+
+	void FCollisionSources::ClearExternalCollisions()
+	{
+		ExternalCollisions.Reset();
+	}
+
 	void FCollisionSources::Reset()
 	{
 		CollisionSources.Reset();
 		++Version;
 	}
 
-	void FCollisionSources::ExtractNonEnvironmentalCollisionData(FClothCollisionData& CollisionData)
+	void FCollisionSources::SetOwnerComponent(USkinnedMeshComponent* InOwnerComponent)
+	{
+		OwnerComponent = MakeWeakObjectPtr(InOwnerComponent);
+	}
+
+	void FCollisionSources::ExtractSkinnedMeshCollisionData(FClothCollisionData& CollisionData)
 	{
 		CollisionData.Reset();
 		if (OwnerComponent.IsValid())
@@ -157,6 +172,11 @@ namespace UE::Chaos::ClothAsset
 		}
 	}
 
+	void FCollisionSources::ExtractExternalCollisionData(FClothCollisionData& CollisionData)
+	{
+		CollisionData = ExternalCollisions;
+	}
+
 	void FCollisionSourcesProxy::ExtractCollisionData()
 	{
 		bool bCollisionDataChanged = false;
@@ -165,7 +185,7 @@ namespace UE::Chaos::ClothAsset
 
 		if (!CollisionSources.IsCollisionDataUpToDate(Version))
 		{
-			CollisionSources.ExtractNonEnvironmentalCollisionData(NonEnvironmentalCollisionData);
+			CollisionSources.ExtractSkinnedMeshCollisionData(SkinnedMeshCollisionData);
 			CollisionSources.MarkCollisionDataUpToDate(Version);
 			bCollisionDataChanged = true;
 		}
@@ -176,10 +196,17 @@ namespace UE::Chaos::ClothAsset
 		CollisionSources.ExtractEnvironmentalCollisionData(EnvironmentalCollisionData);
 		bCollisionDataChanged |= !EnvironmentalCollisionData.IsEmpty();
 
+		// ExternalCollisionData changes every frame, so if we had some before, it's now stale.
+		bCollisionDataChanged |= !ExternalCollisionData.IsEmpty();
+
+		CollisionSources.ExtractExternalCollisionData(ExternalCollisionData);
+		bCollisionDataChanged |= !ExternalCollisionData.IsEmpty();
+
 		if (bCollisionDataChanged)
 		{
-			CollisionData = NonEnvironmentalCollisionData;
+			CollisionData = SkinnedMeshCollisionData;
 			CollisionData.Append(EnvironmentalCollisionData);
+			CollisionData.Append(ExternalCollisionData);
 		}
 	}
 
diff --git a/Engine/Plugins/ChaosClothAsset/Source/ChaosClothAssetEngine/Private/ChaosClothAsset/CollisionSources.h b/Engine/Plugins/ChaosClothAsset/Source/ChaosClothAssetEngine/Private/ChaosClothAsset/CollisionSources.h
--- a/Engine/Plugins/ChaosClothAsset/Source/ChaosClothAssetEngine/Private/ChaosClothAsset/CollisionSources.h
+++ b/Engine/Plugins/ChaosClothAsset/Source/ChaosClothAssetEngine/Private/ChaosClothAsset/CollisionSources.h
@@ -17,7 +17,7 @@ namespace UE::Chaos::ClothAsset
 	class FCollisionSources final
 	{
 	public:
-		explicit FCollisionSources(USkinnedMeshComponent* InOwnerComponent, bool bInCollideWithEnvironment = false);
+		explicit FCollisionSources(USkinnedMeshComponent* InOwnerComponent = nullptr, bool bInCollideWithEnvironment = false);
 		~FCollisionSources();
 
 		void Add(USkinnedMeshComponent* SourceComponent, const UPhysicsAsset* SourcePhysicsAsset, bool bUseSphylsOnly = false);
@@ -26,8 +26,19 @@ namespace UE::Chaos::ClothAsset
 
 		void Remove(const USkinnedMeshComponent* SourceComponent, const UPhysicsAsset* SourcePhysicsAsset);
 
+		void AddExternalCollisions(const FClothCollisionData& InExternalCollisions);
+
+		void ClearExternalCollisions();
+
 		void Reset();
 
+		void SetOwnerComponent(USkinnedMeshComponent* InOwnerComponent);
+
+		USkinnedMeshComponent* GetOwnerComponent() const
+		{
+			return OwnerComponent.Get();
+		}
+
 		void SetCollideWithEnvironment(bool bCollide)
 		{
 			bCollideWithEnvironment = bCollide;
@@ -56,14 +67,16 @@ namespace UE::Chaos::ClothAsset
 			void ExtractCollisionData(const USkinnedMeshComponent& InOwnerComponent, FClothCollisionData& CollisionData);
 		};
 
-		void ExtractNonEnvironmentalCollisionData(FClothCollisionData& CollisionData);
+		void ExtractSkinnedMeshCollisionData(FClothCollisionData& CollisionData);
 		void ExtractEnvironmentalCollisionData(FClothCollisionData& CollisionData);
+		void ExtractExternalCollisionData(FClothCollisionData& CollisionData);
 		bool IsCollisionDataUpToDate(int32 InVersion) const;
 		void MarkCollisionDataUpToDate(int32 InVersion);
 		void ResolveWeakComponentPtrs();
 
-		const TWeakObjectPtr<USkinnedMeshComponent> OwnerComponent;
+		TWeakObjectPtr<USkinnedMeshComponent> OwnerComponent;
 		TArray<FCollisionSource> CollisionSources;
+		FClothCollisionData ExternalCollisions;
 		int32 Version = INDEX_NONE;
 		bool bCollideWithEnvironment = false;
 	};
@@ -84,8 +97,9 @@ namespace UE::Chaos::ClothAsset
 
 	protected:
 		FCollisionSources& CollisionSources;
-		FClothCollisionData NonEnvironmentalCollisionData;
+		FClothCollisionData SkinnedMeshCollisionData;
 		FClothCollisionData EnvironmentalCollisionData;
+		FClothCollisionData ExternalCollisionData;
 
 		FClothCollisionData CollisionData;
 		int32 Version = INDEX_NONE;
