Fix check inside tether regeneration after remeshing. This appears to be caused by topologically degenerate triangles. Opting to just exclude the degenerate triangles when packaging up data to the tether generation code. We have code to CleanupAndCompactMesh which will remove the triangles from the collection, but this will affect indexing and might have downstream effects.

diff --git a/Engine/Plugins/ChaosClothAsset/Source/ChaosClothAssetEngine/Private/ChaosClothAsset/ClothEngineTools.cpp b/Engine/Plugins/ChaosClothAsset/Source/ChaosClothAssetEngine/Private/ChaosClothAsset/ClothEngineTools.cpp
--- a/Engine/Plugins/ChaosClothAsset/Source/ChaosClothAssetEngine/Private/ChaosClothAsset/ClothEngineTools.cpp
+++ b/Engine/Plugins/ChaosClothAsset/Source/ChaosClothAssetEngine/Private/ChaosClothAsset/ClothEngineTools.cpp
@@ -46,9 +46,19 @@ void FClothEngineTools::GenerateTethers(const TSharedRef<FManagedArrayCollection
 		SimIndices.Reserve(ClothFacade.GetNumSimFaces() * 3);
 		for (const FIntVector3& Face : ClothFacade.GetSimIndices3D())
 		{
-			SimIndices.Add(Face[0]);
-			SimIndices.Add(Face[1]);
-			SimIndices.Add(Face[2]);
+			// Exclude degenerate faces
+			if (Face[0] != INDEX_NONE &&
+				Face[1] != INDEX_NONE &&
+				Face[2] != INDEX_NONE &&
+
+				Face[0] != Face[1] &&
+				Face[0] != Face[2] &&
+				Face[1] != Face[2])
+			{
+				SimIndices.Add(Face[0]);
+				SimIndices.Add(Face[1]);
+				SimIndices.Add(Face[2]);
+			}
 		}
 
 		if (MaxDistanceValue.Equals(FVector2f(0.f, 1.f)))
@@ -57,7 +67,6 @@ void FClothEngineTools::GenerateTethers(const TSharedRef<FManagedArrayCollection
 		}
 		else
 		{
-
 			const TSet<int32> KinematicVertices = FClothGeometryTools::GenerateKinematicVertices3D(ClothCollection, WeightMapName, MaxDistanceValue, NAME_None);
 			TetherData.GenerateTethers(ClothFacade.GetSimPosition3D(), TConstArrayView<uint32>(SimIndices), KinematicVertices, bGenerateGeodesicTethers);
 		}
@@ -79,9 +88,19 @@ void FClothEngineTools::GenerateTethersFromSelectionSet(const TSharedRef<FManage
 		SimIndices.Reserve(ClothFacade.GetNumSimFaces() * 3);
 		for (const FIntVector3& Face : ClothFacade.GetSimIndices3D())
 		{
-			SimIndices.Add(Face[0]);
-			SimIndices.Add(Face[1]);
-			SimIndices.Add(Face[2]);
+			// Exclude degenerate faces
+			if (Face[0] != INDEX_NONE &&
+				Face[1] != INDEX_NONE &&
+				Face[2] != INDEX_NONE &&
+
+				Face[0] != Face[1] &&
+				Face[0] != Face[2] &&
+				Face[1] != Face[2])
+			{
+				SimIndices.Add(Face[0]);
+				SimIndices.Add(Face[1]);
+				SimIndices.Add(Face[2]);
+			}
 		}
 
 		TetherData.GenerateTethers(ClothFacade.GetSimPosition3D(), TConstArrayView<uint32>(SimIndices), SelectionFacade.GetSelectionSet(FixedEndSet), bGeodesicTethers);
@@ -101,9 +120,19 @@ void FClothEngineTools::GenerateTethersFromCustomSelectionSets(const TSharedRef<
 		SimIndices.Reserve(ClothFacade.GetNumSimFaces() * 3);
 		for (const FIntVector3& Face : ClothFacade.GetSimIndices3D())
 		{
-			SimIndices.Add(Face[0]);
-			SimIndices.Add(Face[1]);
-			SimIndices.Add(Face[2]);
+			// Exclude degenerate faces
+			if (Face[0] != INDEX_NONE &&
+				Face[1] != INDEX_NONE &&
+				Face[2] != INDEX_NONE &&
+
+				Face[0] != Face[1] &&
+				Face[0] != Face[2] &&
+				Face[1] != Face[2])
+			{
+				SimIndices.Add(Face[0]);
+				SimIndices.Add(Face[1]);
+				SimIndices.Add(Face[2]);
+			}
 		}
 
 		const TSet<int32>& FixedEndSet = SelectionFacade.GetSelectionSet(InFixedEndSet);
