Fix weight map regression when upgrading cloth properties to use FNames

diff --git a/Engine/Plugins/ChaosClothAsset/Source/ChaosClothAssetEngine/Private/ChaosClothAsset/ClothEngineTools.cpp b/Engine/Plugins/ChaosClothAsset/Source/ChaosClothAssetEngine/Private/ChaosClothAsset/ClothEngineTools.cpp
--- a/Engine/Plugins/ChaosClothAsset/Source/ChaosClothAssetEngine/Private/ChaosClothAsset/ClothEngineTools.cpp
+++ b/Engine/Plugins/ChaosClothAsset/Source/ChaosClothAssetEngine/Private/ChaosClothAsset/ClothEngineTools.cpp
@@ -168,7 +168,7 @@ FPointWeightMap FClothEngineTools::GetMaxDistanceWeightMap(const FCollectionClot
 	const bool bHasMaxDistanceProperty = (MaxDistancePropertyKeyIndex != INDEX_NONE);
 	const float MaxDistanceOffset = bHasMaxDistanceProperty ? PropertyFacade.GetLowValue<float>(MaxDistancePropertyKeyIndex) : TNumericLimits<float>::Max();  // Uses infinite distance when no MaxDistance properties are set
 	const float MaxDistanceScale = bHasMaxDistanceProperty ? PropertyFacade.GetHighValue<float>(MaxDistancePropertyKeyIndex) - MaxDistanceOffset : 0.f;
-	const TConstArrayView<float> MaxDistanceWeightMap = ClothFacade.GetWeightMap(MaxDistanceName);
+	const TConstArrayView<float> MaxDistanceWeightMap = ClothFacade.GetWeightMap(FName(MaxDistanceString));
 
 	return (MaxDistanceWeightMap.Num() == NumLodSimVertices) ?
 		FPointWeightMap(MaxDistanceWeightMap, MaxDistanceOffset, MaxDistanceScale) :
