Cloth - Added support for the Cloth/Outfit Asset to the SkeletalMesh clothing data, and multi-simulation support to the SkeletalMeshComponent.

[RN] Cloth - Added support for the Cloth/Outfit Asset to the SkeletalMesh clothing data, and multi-simulation support to the SkeletalMeshComponent.

diff --git a/Engine/Plugins/ChaosCloth/Source/ChaosCloth/Private/ChaosCloth/ChaosClothingSimulation.cpp b/Engine/Plugins/ChaosCloth/Source/ChaosCloth/Private/ChaosCloth/ChaosClothingSimulation.cpp
--- a/Engine/Plugins/ChaosCloth/Source/ChaosCloth/Private/ChaosCloth/ChaosClothingSimulation.cpp
+++ b/Engine/Plugins/ChaosCloth/Source/ChaosCloth/Private/ChaosCloth/ChaosClothingSimulation.cpp
@@ -3,6 +3,7 @@
 #include "ChaosCloth/ChaosClothPrivate.h"
 
 #include "ClothingSimulation.h"
+#include "ClothingSimulationInstance.h"
 #include "ClothingAsset.h"
 #include "ChaosCloth/ChaosClothConfig.h"
 
@@ -331,7 +332,7 @@ IClothingSimulationContext* FClothingSimulation::CreateContext()
 	return new FClothingSimulationContext();
 }
 
-void FClothingSimulation::CreateActor(USkeletalMeshComponent* InOwnerComponent, UClothingAssetBase* InAsset, int32 InSimDataIndex)
+void FClothingSimulation::CreateActor(USkeletalMeshComponent* InOwnerComponent, const UClothingAssetBase* InAsset, int32 InSimDataIndex)
 {
 	TRACE_CPUPROFILER_EVENT_SCOPE(FClothingSimulation_CreateActor);
 	SCOPE_CYCLE_COUNTER(STAT_ChaosClothCreateActor);
@@ -345,7 +346,7 @@ void FClothingSimulation::CreateActor(USkeletalMeshComponent* InOwnerComponent,
 	}
 
 	// ClothSharedSimConfig should either be a nullptr, or point to an object common to the whole skeletal mesh
-	UClothingAssetCommon* const Asset = Cast<UClothingAssetCommon>(InAsset);
+	const UClothingAssetCommon* const Asset = Cast<UClothingAssetCommon>(InAsset);
 	if (!ClothSharedSimConfig)
 	{
 		ClothSharedSimConfig = Asset->GetClothConfig<UChaosClothSharedSimConfig>();
@@ -353,7 +354,7 @@ void FClothingSimulation::CreateActor(USkeletalMeshComponent* InOwnerComponent,
 		UpdateSimulationFromSharedSimConfig();
 
 		// Must set the local space location prior to adding any mesh/cloth, as otherwise the start poses would be in the wrong local space
-		const FClothingSimulationContext* const Context = static_cast<const FClothingSimulationContext*>(InOwnerComponent->GetClothingSimulationContext());
+		const FClothingSimulationContext* const Context = GetClothingSimulationContext(InOwnerComponent);
 		check(Context);
 		static const bool bReset = true;
 		const FReal WorldToSolverScale = bUseLocalSpaceSimulation ? Context->SolverGeometryScale : 1.;
@@ -494,7 +495,7 @@ bool FClothingSimulation::ShouldSimulate() const
 	return false;
 }
 
-void FClothingSimulation::Simulate(IClothingSimulationContext* InContext)
+void FClothingSimulation::Simulate(const IClothingSimulationContext* InContext)
 {
 	TRACE_CPUPROFILER_EVENT_SCOPE(FClothingSimulation_Simulate);
 #if !(UE_BUILD_SHIPPING || UE_BUILD_TEST)
@@ -502,7 +503,7 @@ void FClothingSimulation::Simulate(IClothingSimulationContext* InContext)
 #endif
 	{
 		SCOPE_CYCLE_COUNTER(STAT_ChaosClothSimulate);
-		const FClothingSimulationContext* const Context = static_cast<FClothingSimulationContext*>(InContext);
+		const FClothingSimulationContext* const Context = static_cast<const FClothingSimulationContext*>(InContext);
 		if (Context->DeltaSeconds == 0.f)
 		{
 			return;
@@ -612,8 +613,8 @@ void FClothingSimulation::Simulate(IClothingSimulationContext* InContext)
 
 void FClothingSimulation::GetSimulationData(
 	TMap<int32, FClothSimulData>& OutData,
-	USkeletalMeshComponent* InOwnerComponent,
-	USkinnedMeshComponent* InOverrideComponent) const
+	const USkeletalMeshComponent* InOwnerComponent,
+	const USkinnedMeshComponent* InOverrideComponent) const
 {
 	TRACE_CPUPROFILER_EVENT_SCOPE(FClothingSimulation_GetSimulationData);
 	SCOPE_CYCLE_COUNTER(STAT_ChaosClothGetSimulationData);
@@ -635,7 +636,7 @@ void FClothingSimulation::GetSimulationData(
 	const FReal LocalSpaceScale = Solver->GetLocalSpaceScale();
 
 	// Retrieve the component transforms
-	const FClothingSimulationContext* const Context = static_cast<const FClothingSimulationContext*>(InOwnerComponent->GetClothingSimulationContext());
+	const FClothingSimulationContext* const Context = GetClothingSimulationContext(InOwnerComponent);
 	check(Context);  // A simulation can't be created without context
 	const FTransform& OwnerTransform = Context->ComponentToWorld;
 
@@ -752,7 +753,7 @@ FBoxSphereBounds FClothingSimulation::GetBounds(const USkeletalMeshComponent* In
 	}
 	else if (InOwnerComponent)
 	{
-		const FClothingSimulationContext* const Context = static_cast<const FClothingSimulationContext*>(InOwnerComponent->GetClothingSimulationContext());
+		const FClothingSimulationContext* const Context = GetClothingSimulationContext(InOwnerComponent);
 		check(Context);  // A simulation can't be created without context
 		const FTransform& OwnerTransform = Context->ComponentToWorld;
 		// Return local bounds
@@ -878,4 +879,16 @@ FClothingSimulationCloth* Chaos::FClothingSimulation::GetCloth(int32 ClothId)
 	return Cloth ? Cloth->Get(): nullptr;
 }
 
+const FClothingSimulationContext* FClothingSimulation::GetClothingSimulationContext(const USkeletalMeshComponent* SkeletalMeshComponent) const
+{
+	for (const FClothingSimulationInstance& ClothingSimulationInstance : SkeletalMeshComponent->GetClothingSimulationInstances())
+	{
+		if (ClothingSimulationInstance.GetClothingSimulation() == this)
+		{
+			return static_cast<const FClothingSimulationContext*>(ClothingSimulationInstance.GetClothingSimulationContext());
+		}
+	}
+	return nullptr;
+}
+
 }  // End namespace Chaos
diff --git a/Engine/Plugins/ChaosCloth/Source/ChaosCloth/Private/ChaosCloth/ChaosClothingSimulationFactory.cpp b/Engine/Plugins/ChaosCloth/Source/ChaosCloth/Private/ChaosCloth/ChaosClothingSimulationFactory.cpp
--- a/Engine/Plugins/ChaosCloth/Source/ChaosCloth/Private/ChaosCloth/ChaosClothingSimulationFactory.cpp
+++ b/Engine/Plugins/ChaosCloth/Source/ChaosCloth/Private/ChaosCloth/ChaosClothingSimulationFactory.cpp
@@ -17,17 +17,17 @@ IClothingSimulation* UChaosClothingSimulationFactory::CreateSimulation()
 
 void UChaosClothingSimulationFactory::DestroySimulation(IClothingSimulation* InSimulation)
 {
-    delete InSimulation;
+	delete InSimulation;
 }
 
-bool UChaosClothingSimulationFactory::SupportsAsset(UClothingAssetBase* InAsset)
+bool UChaosClothingSimulationFactory::SupportsAsset(const UClothingAssetBase* InAsset) const
 {
-    return true;
+	return Cast<UClothingAssetCommon>(InAsset) != nullptr;
 }
 
-bool UChaosClothingSimulationFactory::SupportsRuntimeInteraction()
+bool UChaosClothingSimulationFactory::SupportsRuntimeInteraction() const
 {
-    return true;
+	return true;
 }
 
 UClothingSimulationInteractor* UChaosClothingSimulationFactory::CreateInteractor()
diff --git a/Engine/Plugins/ChaosCloth/Source/ChaosCloth/Private/ChaosCloth/ChaosClothingSimulationSkeletalMesh.cpp b/Engine/Plugins/ChaosCloth/Source/ChaosCloth/Private/ChaosCloth/ChaosClothingSimulationSkeletalMesh.cpp
--- a/Engine/Plugins/ChaosCloth/Source/ChaosCloth/Private/ChaosCloth/ChaosClothingSimulationSkeletalMesh.cpp
+++ b/Engine/Plugins/ChaosCloth/Source/ChaosCloth/Private/ChaosCloth/ChaosClothingSimulationSkeletalMesh.cpp
@@ -5,6 +5,8 @@
 #include "ChaosWeightMapTarget.h"
 #include "ClothingAsset.h"
 #include "ClothingSimulation.h"
+#include "ClothingSimulationInstance.h"
+#include "ClothingSimulationFactory.h"
 #include "Engine/SkeletalMesh.h"
 
 namespace Chaos
@@ -275,8 +277,19 @@ namespace Chaos
 
 	const FClothingSimulationContextCommon* FClothingSimulationSkeletalMesh::GetContext() const
 	{
-		return GetSkeletalMeshComponent() ?
-			static_cast<const FClothingSimulationContextCommon*>(GetSkeletalMeshComponent()->GetClothingSimulationContext()) :
-			nullptr;
+		if (SkeletalMeshComponent)
+		{
+			for (const FClothingSimulationInstance& ClothingSimulationInstance : SkeletalMeshComponent->GetClothingSimulationInstances())
+			{
+				if (UClothingSimulationFactory* const ClothingSimulationFactory = ClothingSimulationInstance.GetClothingSimulationFactory())
+				{
+					if (ClothingSimulationFactory->SupportsAsset(Asset))
+					{
+						return static_cast<const FClothingSimulationContextCommon*>(ClothingSimulationInstance.GetClothingSimulationContext());
+					}
+				}
+			}
+		}
+		return nullptr;
 	}
 }  // End namespace Chaos
diff --git a/Engine/Plugins/ChaosCloth/Source/ChaosCloth/Private/ChaosCloth/SkeletalMeshComponentCacheAdapter.cpp b/Engine/Plugins/ChaosCloth/Source/ChaosCloth/Private/ChaosCloth/SkeletalMeshComponentCacheAdapter.cpp
--- a/Engine/Plugins/ChaosCloth/Source/ChaosCloth/Private/ChaosCloth/SkeletalMeshComponentCacheAdapter.cpp
+++ b/Engine/Plugins/ChaosCloth/Source/ChaosCloth/Private/ChaosCloth/SkeletalMeshComponentCacheAdapter.cpp
@@ -2,6 +2,7 @@
 
 #include "ChaosCloth/SkeletalMeshComponentCacheAdapter.h"
 #include "ChaosCloth/ChaosClothingCacheSchema.h"
+#include "ChaosCloth/ChaosClothPrivate.h"
 #include "Chaos/ParticleHandle.h"
 #include "Chaos/ChaosCache.h"
 #include "Chaos/CacheManagerActor.h"
@@ -73,10 +74,19 @@ namespace Chaos
 		{
 			if( USkeletalMeshComponent* MeshComp = CastChecked<USkeletalMeshComponent>(InComponent) )
 			{
+				if (MeshComp->GetClothingSimulationInstances().Num() > 1)
+				{
+					static bool bErroredOnce = false;
+					UE_CLOG(!bErroredOnce, LogChaosCloth, Error,
+						TEXT("SkeletalMeshComponent [%s] has more than one clothing simulation running, but cloth caching can only deal with a single one."), *MeshComp->GetName());
+					bErroredOnce = true;
+				}
+PRAGMA_DISABLE_DEPRECATION_WARNINGS
 				if(MeshComp->GetClothingSimulation())
 				{
 					return StaticCast<FClothingSimulation*>(MeshComp->GetClothingSimulation())->Solver.Get();
 				}
+PRAGMA_ENABLE_DEPRECATION_WARNINGS
 			}
 		}
 		return nullptr;
diff --git a/Engine/Plugins/ChaosCloth/Source/ChaosCloth/Public/ChaosCloth/ChaosClothingSimulation.h b/Engine/Plugins/ChaosCloth/Source/ChaosCloth/Public/ChaosCloth/ChaosClothingSimulation.h
--- a/Engine/Plugins/ChaosCloth/Source/ChaosCloth/Public/ChaosCloth/ChaosClothingSimulation.h
+++ b/Engine/Plugins/ChaosCloth/Source/ChaosCloth/Public/ChaosCloth/ChaosClothingSimulation.h
@@ -40,12 +40,12 @@ namespace Chaos
 		virtual IClothingSimulationContext* CreateContext() override;
 		virtual void DestroyContext(IClothingSimulationContext* InContext) override { delete InContext; }
 
-		virtual void CreateActor(USkeletalMeshComponent* InOwnerComponent, UClothingAssetBase* InAsset, int32 SimDataIndex) override;
+		virtual void CreateActor(USkeletalMeshComponent* InOwnerComponent, const UClothingAssetBase* InAsset, int32 SimDataIndex) override;
 		virtual void DestroyActors() override;
 
 		virtual bool ShouldSimulate() const override;
-		virtual void Simulate(IClothingSimulationContext* InContext) override;
-		virtual void GetSimulationData(TMap<int32, FClothSimulData>& OutData, USkeletalMeshComponent* InOwnerComponent, USkinnedMeshComponent* InOverrideComponent) const override;
+		virtual void Simulate(const IClothingSimulationContext* InContext) override;
+		virtual void GetSimulationData(TMap<int32, FClothSimulData>& OutData, const USkeletalMeshComponent* InOwnerComponent, const USkinnedMeshComponent* InOverrideComponent) const override;
 
 		// Return bounds in local space (or in world space if InOwnerComponent is null).
 		virtual FBoxSphereBounds GetBounds(const USkeletalMeshComponent* InOwnerComponent) const override;
@@ -124,6 +124,8 @@ namespace Chaos
 
 		void UpdateSimulationFromSharedSimConfig();
 
+		const FClothingSimulationContext* GetClothingSimulationContext(const USkeletalMeshComponent* SkeletalMeshComponent) const;
+
 	private:
 		// Visualization object
 PRAGMA_DISABLE_DEPRECATION_WARNINGS
@@ -141,7 +143,7 @@ PRAGMA_ENABLE_DEPRECATION_WARNINGS
 		FClothCollisionData ExternalCollisionData;
 
 		// Shared cloth config
-		UChaosClothSharedSimConfig* ClothSharedSimConfig;
+		const UChaosClothSharedSimConfig* ClothSharedSimConfig;
 
 		// Properties that must be readable from all threads
 		TAtomic<int32> NumCloths;
diff --git a/Engine/Plugins/ChaosCloth/Source/ChaosCloth/Public/ChaosCloth/ChaosClothingSimulationFactory.h b/Engine/Plugins/ChaosCloth/Source/ChaosCloth/Public/ChaosCloth/ChaosClothingSimulationFactory.h
--- a/Engine/Plugins/ChaosCloth/Source/ChaosCloth/Public/ChaosCloth/ChaosClothingSimulationFactory.h
+++ b/Engine/Plugins/ChaosCloth/Source/ChaosCloth/Public/ChaosCloth/ChaosClothingSimulationFactory.h
@@ -13,9 +13,9 @@ class UChaosClothingSimulationFactory final : public UClothingSimulationFactory
 public:
 	CHAOSCLOTH_API virtual IClothingSimulation* CreateSimulation() override;
 	CHAOSCLOTH_API virtual void DestroySimulation(IClothingSimulation* InSimulation) override;
-	CHAOSCLOTH_API virtual bool SupportsAsset(UClothingAssetBase* InAsset) override;
+	CHAOSCLOTH_API virtual bool SupportsAsset(const UClothingAssetBase* InAsset) const override;
 
-	CHAOSCLOTH_API virtual bool SupportsRuntimeInteraction() override;
+	CHAOSCLOTH_API virtual bool SupportsRuntimeInteraction() const override;
 	CHAOSCLOTH_API virtual UClothingSimulationInteractor* CreateInteractor() override;
 
 	CHAOSCLOTH_API virtual TArrayView<const TSubclassOf<UClothConfigBase>> GetClothConfigClasses() const override;
diff --git a/Engine/Plugins/ChaosClothAsset/Source/ChaosClothAssetEngine/Private/ChaosClothAsset/ChaosClothAssetEngineModule.cpp b/Engine/Plugins/ChaosClothAsset/Source/ChaosClothAssetEngine/Private/ChaosClothAsset/ChaosClothAssetEngineModule.cpp
--- a/Engine/Plugins/ChaosClothAsset/Source/ChaosClothAssetEngine/Private/ChaosClothAsset/ChaosClothAssetEngineModule.cpp
+++ b/Engine/Plugins/ChaosClothAsset/Source/ChaosClothAssetEngine/Private/ChaosClothAsset/ChaosClothAssetEngineModule.cpp
@@ -1,7 +1,7 @@
 // Copyright Epic Games, Inc. All Rights Reserved.
 
 #include "ChaosClothAsset/ClothAssetPrivate.h"
-#include "ChaosClothAsset/ClothAssetSKMSimulationFactory.h"
+#include "ChaosClothAsset/ClothAssetSKMClothingSimulation.h"
 #include "ChaosClothAsset/ClothComponentCacheAdapter.h"
 #include "Features/IModularFeatures.h"
 #include "Modules/ModuleInterface.h"
@@ -55,7 +55,7 @@ namespace UE::Chaos::ClothAsset
 		//~ Begin IClothingSimulationFactoryClassProvider
 		virtual TSubclassOf<class UClothingSimulationFactory> GetClothingSimulationFactoryClass() const override
 		{
-			return UChaosClothAssetSKMSimulationFactory::StaticClass();
+			return UChaosClothAssetSKMClothingSimulationFactory::StaticClass();
 		}
 		//~ End IClothingSimulationFactoryClassProvider
 
diff --git a/Engine/Plugins/ChaosClothAsset/Source/ChaosClothAssetEngine/Private/ChaosClothAsset/ClothAssetBase.cpp b/Engine/Plugins/ChaosClothAsset/Source/ChaosClothAssetEngine/Private/ChaosClothAsset/ClothAssetBase.cpp
--- a/Engine/Plugins/ChaosClothAsset/Source/ChaosClothAssetEngine/Private/ChaosClothAsset/ClothAssetBase.cpp
+++ b/Engine/Plugins/ChaosClothAsset/Source/ChaosClothAssetEngine/Private/ChaosClothAsset/ClothAssetBase.cpp
@@ -4,10 +4,12 @@
 #if WITH_EDITORONLY_DATA
 #include "Animation/AnimationAsset.h"
 #endif
+#include "ChaosClothAsset/ClothAssetSKMClothingAsset.h"
 #include "ChaosClothAsset/ClothAssetPrivate.h"
 #include "ChaosClothAsset/ClothComponent.h"
 #include "ChaosClothAsset/ClothSimulationModel.h"
 #include "ChaosClothAsset/SkeletalMeshConverterClassProvider.h"
+#include "Components/SkeletalMeshComponent.h"
 #include "Engine/Engine.h"
 #include "Engine/SkeletalMesh.h"
 #include "Engine/SkinnedAssetAsyncCompileUtils.h"
@@ -599,7 +601,7 @@ bool UChaosClothAssetBase::ExportToSkeletalMesh(USkeletalMesh& SkeletalMesh) con
 
 void UChaosClothAssetBase::ReregisterComponents()
 {
-	// Recreate the simulation proxies with the updated physics asset
+	// Reregister cloth components
 	for (TObjectIterator<UChaosClothComponent> ObjectIterator; ObjectIterator; ++ObjectIterator)
 	{
 		if (UChaosClothComponent* const Component = *ObjectIterator)
@@ -610,4 +612,32 @@ void UChaosClothAssetBase::ReregisterComponents()
 			}
 		}
 	}
+	// Reregister skeletal mesh components
+	for (TObjectIterator<UChaosClothAssetSKMClothingAsset> ObjectIterator; ObjectIterator; ++ObjectIterator)
+	{
+		if (UChaosClothAssetSKMClothingAsset* const ClothingAsset = *ObjectIterator)
+		{
+			if (ClothingAsset->GetAsset() == this)
+			{
+#if WITH_EDITOR
+				ClothingAsset->OnAssetChanged();  // ReregisterComponents is called, therefore the asset must have changed, this will update the Clothing Data and redo the bindings
+#endif
+
+				if (USkeletalMesh* const OwnerMesh = Cast<USkeletalMesh>(ClothingAsset->GetOuter()))
+				{
+					for (TObjectIterator<USkeletalMeshComponent> It; It; ++It)
+					{
+						if (USkeletalMeshComponent* const Component = *It)
+						{
+							if (Component->GetSkeletalMeshAsset() == OwnerMesh)
+							{
+								FComponentReregisterContext Context(Component);
+								// Context goes out of scope, causing Component to be re-registered
+							}
+						}
+					}
+				}
+			}
+		}
+	}
 }
diff --git a/Engine/Plugins/ChaosClothAsset/Source/ChaosClothAssetEngine/Private/ChaosClothAsset/ClothAssetSKMClothingAsset.cpp b/Engine/Plugins/ChaosClothAsset/Source/ChaosClothAssetEngine/Private/ChaosClothAsset/ClothAssetSKMClothingAsset.cpp
--- a/Engine/Plugins/ChaosClothAsset/Source/ChaosClothAssetEngine/Private/ChaosClothAsset/ClothAssetSKMClothingAsset.cpp
+++ b/Engine/Plugins/ChaosClothAsset/Source/ChaosClothAssetEngine/Private/ChaosClothAsset/ClothAssetSKMClothingAsset.cpp
@@ -9,20 +9,121 @@
 #include "ChaosClothAsset/ClothEngineTools.h"
 #include "ChaosClothAsset/ClothSimulationModel.h"
 #include "ChaosClothAsset/CollectionClothFacade.h"
+#include "Components/SkeletalMeshComponent.h"
 #include "Engine/SkeletalMesh.h"
-#include "Utils/ClothingMeshUtils.h"
 #include "Rendering/SkeletalMeshModel.h"
+#include "Utils/ClothingMeshUtils.h"
+#include "ClothingAsset.h"
 #endif  // #if WITH_EDITOR
 
 #include UE_INLINE_GENERATED_CPP_BY_NAME(ClothAssetSKMClothingAsset)
 
 #define LOCTEXT_NAMESPACE "ChaosClothAssetSKMClothingAsset"
 
+namespace UE::Chaos::ClothAsset::Private
+{
+#if WITH_EDITOR
+	static FString MakeClothSimulationModelIdString(const FName Name, const FGuid Guid)
+	{
+		return FString::Format(TEXT("{0}-{1}"), { Name.ToString(), Guid.ToString() });
+	}
+
+	// Modified version of FLODUtilities::UnbindClothingAndBackup to only unbind specified assets
+	static void UnbindClothingAndBackup(USkeletalMesh& SkeletalMesh, const UChaosClothAssetSKMClothingAsset& ClothingAsset, TArray<ClothingAssetUtils::FClothingAssetMeshBinding> ClothingBindings)
+	{
+		ClothingBindings.Reset();
+
+		if (SkeletalMesh.GetImportedModel())
+		{
+			auto UnbindClothingLODAndBackup = [&SkeletalMesh, &ClothingAsset](const int32 LODIndex, TArray<ClothingAssetUtils::FClothingAssetMeshBinding> ClothingLODBindings)
+				{
+					ClothingLODBindings.Reset();
+
+					TIndirectArray<FSkeletalMeshLODModel>& LODModels = SkeletalMesh.GetImportedModel()->LODModels;
+					if (LODModels.IsValidIndex(LODIndex))
+					{
+						FSkeletalMeshLODModel& LODModel = LODModels[LODIndex];
+
+						// Store this LOD's bindings
+						ClothingAssetUtils::GetAllLodMeshClothingAssetBindings(&SkeletalMesh, ClothingLODBindings, LODIndex);  // TODO: We only need the binding for the specified ClothingAsset
+
+						// Unbind this Cloth's LOD
+						for (ClothingAssetUtils::FClothingAssetMeshBinding& Binding : ClothingLODBindings)
+						{
+							if (Binding.LODIndex == LODIndex && Binding.Asset == &ClothingAsset)
+							{
+								const int32 OriginalDataSectionIndex = LODModel.Sections[Binding.SectionIndex].OriginalDataSectionIndex;
+								Binding.Asset->UnbindFromSkeletalMesh(&SkeletalMesh, Binding.LODIndex);
+								Binding.SectionIndex = OriginalDataSectionIndex;
+
+								FSkelMeshSourceSectionUserData& SectionUserData = LODModel.UserSectionsData.FindChecked(OriginalDataSectionIndex);
+								SectionUserData.ClothingData.AssetGuid = FGuid();
+								SectionUserData.ClothingData.AssetLodIndex = INDEX_NONE;
+								SectionUserData.CorrespondClothAssetIndex = INDEX_NONE;
+							}
+						}
+					}
+				};
+
+			for (int32 LODIndex = 0; LODIndex < SkeletalMesh.GetImportedModel()->LODModels.Num(); ++LODIndex)
+			{
+				TArray<ClothingAssetUtils::FClothingAssetMeshBinding> ClothingLODBindings;
+				UnbindClothingLODAndBackup(LODIndex, ClothingLODBindings);
+				ClothingBindings.Append(ClothingLODBindings);
+			}
+		}
+	}
+
+	// Modified version of FLODUtilities::RestoreClothingFromBackup to only unbind specified assets
+	static void RestoreClothingFromBackup(USkeletalMesh& SkeletalMesh, const UChaosClothAssetSKMClothingAsset& ClothingAsset, TArray<ClothingAssetUtils::FClothingAssetMeshBinding>& ClothingBindings)
+	{
+		if (SkeletalMesh.GetImportedModel())
+		{
+			auto RestoreClothingLODFromBackup = [&SkeletalMesh, &ClothingAsset](const int32 LODIndex, TArray<ClothingAssetUtils::FClothingAssetMeshBinding>& ClothingBindings)
+				{
+					TIndirectArray<FSkeletalMeshLODModel>& LODModels = SkeletalMesh.GetImportedModel()->LODModels;
+					if (LODModels.IsValidIndex(LODIndex))
+					{
+						FSkeletalMeshLODModel& LODModel = SkeletalMesh.GetImportedModel()->LODModels[LODIndex];
+						for (ClothingAssetUtils::FClothingAssetMeshBinding& Binding : ClothingBindings)
+						{
+							for (int32 SectionIndex = 0; SectionIndex < LODModel.Sections.Num(); ++SectionIndex)
+							{
+								if (LODModel.Sections[SectionIndex].OriginalDataSectionIndex == Binding.SectionIndex)
+								{
+									if (Binding.LODIndex == LODIndex && Binding.Asset == &ClothingAsset)
+									{
+										if (Binding.Asset->BindToSkeletalMesh(&SkeletalMesh, Binding.LODIndex, SectionIndex, Binding.AssetInternalLodIndex))
+										{
+											// If successfull set back the section user data
+											FSkelMeshSourceSectionUserData& SectionUserData = LODModel.UserSectionsData.FindChecked(Binding.SectionIndex);
+											SectionUserData.CorrespondClothAssetIndex = LODModel.Sections[SectionIndex].CorrespondClothAssetIndex;
+											SectionUserData.ClothingData = LODModel.Sections[SectionIndex].ClothingData;
+										}
+									}
+									break;  // Found the section, next binding
+								}
+							}
+						}
+					}
+				};
+
+			for (int32 LODIndex = 0; LODIndex < SkeletalMesh.GetImportedModel()->LODModels.Num(); ++LODIndex)
+			{
+				RestoreClothingLODFromBackup(LODIndex, ClothingBindings);
+			}
+		}
+	}
+#endif
+}
+
 UChaosClothAssetSKMClothingAsset::UChaosClothAssetSKMClothingAsset(const FObjectInitializer& ObjectInitializer)
 	: Super(ObjectInitializer)
 {
+	UClothingAssetBase::AssetGuid = FGuid::NewGuid();
 }
 
+#if WITH_EDITOR
 void UChaosClothAssetSKMClothingAsset::SetAsset(const UChaosClothAssetBase* InAsset)
 {
 	if (Asset != InAsset)
@@ -32,7 +133,72 @@ void UChaosClothAssetSKMClothingAsset::SetAsset(const UChaosClothAssetBase* InAs
 	}
 }
 
-#if WITH_EDITOR
+void UChaosClothAssetSKMClothingAsset::OnAssetChanged()
+{
+	// If the asset has changed, check whether the current specified model still exists in this new asset, otherwise update the model guid
+	int32 ModelIndex = INDEX_NONE;
+
+	if (Asset)
+	{
+		if (const int32 NumClothSimulationModels = Asset->GetNumClothSimulationModels())
+		{
+			// Find a matching model GUID
+			const FGuid ClothSimulationModelGuid = GetClothSimulationModelGuid();
+			for (ModelIndex = 0; ModelIndex < NumClothSimulationModels; ++ModelIndex)
+			{
+				if (ClothSimulationModelGuid == Asset->GetAssetGuid(ModelIndex))
+				{
+					break;
+				}
+			}
+
+			if (ModelIndex == NumClothSimulationModels)
+			{
+				// Try locate a model with a similar name to update to a valid GUID if possible
+				const FName ClothSimulationModelName = GetClothSimulationModelName();
+
+				for (ModelIndex = 0; ModelIndex < Asset->GetNumClothSimulationModels(); ++ModelIndex)
+				{
+					if (Asset->GetClothSimulationModelName(ModelIndex) == ClothSimulationModelName)
+					{
+						break;
+					}
+				}
+
+				if (ModelIndex == NumClothSimulationModels)
+				{
+					// Or reset to the first model if there isn't one
+					ModelIndex = 0;
+				}
+			}
+		}
+	}
+
+	// Notify the model has changed
+	if (ClothSimulationModelIndex != ModelIndex)
+	{
+		ClothSimulationModelIndex = ModelIndex;
+
+		OnModelChanged();
+	}
+
+	// Reregister components
+	if (const USkeletalMesh* const OwnerMesh = Cast<USkeletalMesh>(GetOuter()))
+	{
+		for (TObjectIterator<USkeletalMeshComponent> It; It; ++It)
+		{
+			if (USkeletalMeshComponent* const Component = *It)
+			{
+				if (Component->GetSkeletalMeshAsset() == OwnerMesh)
+				{
+					FComponentReregisterContext Context(Component);
+					// Context goes out of scope, causing Component to be re-registered
+				}
+			}
+		}
+	}  // TODO: Move UClothingAssetCommon::ReregisterComponentsUsingClothing() to a common place (in Utils/ClothingMeshUtils?)
+}
+
 void UChaosClothAssetSKMClothingAsset::PostEditChangeProperty(FPropertyChangedEvent& PropertyChangedEvent)
 {
 	Super::PostEditChangeProperty(PropertyChangedEvent);
@@ -41,6 +207,10 @@ void UChaosClothAssetSKMClothingAsset::PostEditChangeProperty(FPropertyChangedEv
 	{
 		OnAssetChanged();
 	}
+	else if (PropertyChangedEvent.GetPropertyName() == GET_MEMBER_NAME_CHECKED(UChaosClothAssetSKMClothingAsset, ClothSimulationModelId))
+	{
+		OnModelChanged();
+	}
 }
 
 bool UChaosClothAssetSKMClothingAsset::BindToSkeletalMesh(USkeletalMesh* SkeletalMesh, const int32 MeshLodIndex, const int32 SectionIndex, const int32 AssetLodIndex)
@@ -63,13 +233,12 @@ bool UChaosClothAssetSKMClothingAsset::BindToSkeletalMesh(USkeletalMesh* Skeleta
 		WarningNotification(Text);
 		return false;
 	}
-	const int32 ClothSimulationModelIndex = GetClothSimulationModelIndex();
 	if (ClothSimulationModelIndex == INDEX_NONE)
 	{
 		const FText Text = FText::Format(
 			LOCTEXT("NotAClothSimulationModelOfAsset", "Failed to bind Clothing Data [{0}], because [{1}] is not an existing Cloth Simulation Model of Asset [{2}]."),
 			FText::FromString(GetName()),
-			FText::FromName(ClothSimulationModelName),
+			FText::FromName(GetClothSimulationModelName()),
 			FText::FromName(Asset->GetFName()));
 		WarningNotification(Text);
 		return false;
@@ -82,7 +251,7 @@ bool UChaosClothAssetSKMClothingAsset::BindToSkeletalMesh(USkeletalMesh* Skeleta
 		const FText Text = FText::Format(
 			LOCTEXT("EmptyClothSimulationModel", "Failed to bind Clothing Data [{0}], because [{1}] is an empty Cloth Simulation Model of Asset [{2}]."),
 			FText::FromString(GetName()),
-			FText::FromName(ClothSimulationModelName),
+			FText::FromName(GetClothSimulationModelName()),
 			FText::FromName(Asset->GetFName()));
 		WarningNotification(Text);
 		return false;
@@ -92,17 +261,64 @@ bool UChaosClothAssetSKMClothingAsset::BindToSkeletalMesh(USkeletalMesh* Skeleta
 		const FText Text = FText::Format(
 			LOCTEXT("NoLODInClothSimulationModel", "Failed to bind Clothing Data [{0}], because Cloth Simulation Model [{1}] has no LOD{2} in Asset [{3}]."),
 			FText::FromString(GetName()),
-			FText::FromName(ClothSimulationModelName),
+			FText::FromName(GetClothSimulationModelName()),
 			AssetLodIndex,
 			FText::FromName(Asset->GetFName()));
 		WarningNotification(Text);
 		return false;
 	}
 
+	// Get the original render section
 	FSkeletalMeshLODModel& SkeletalMeshLODModel = SkeletalMesh->GetImportedModel()->LODModels[MeshLodIndex];
 	FSkelMeshSection& SkelMeshSection = SkeletalMeshLODModel.Sections[SectionIndex];
 
-	// Set the Clothing Data
+	// Set the new section bone map but first verify that it doesn't exceed the maximum when adding the cloth asset bones
+	TArray<FBoneIndexType> BoneMap = SkelMeshSection.BoneMap;
+	for (const FName BoneName : ClothSimulationModel->UsedBoneNames)
+	{
+		const int32 BoneIndex = SkeletalMesh->GetRefSkeleton().FindBoneIndex(BoneName);
+		if (BoneIndex != INDEX_NONE)
+		{
+			BoneMap.AddUnique(BoneIndex);
+		}
+	}
+
+	if (BoneMap.Num() > FGPUBaseSkinVertexFactory::GetMaxGPUSkinBones())
+	{
+		// Failed to apply as we've exceeded the number of bones we can skin
+		const FText Text = FText::Format(
+			LOCTEXT("TooManyBonesInClothSimulationModel", "Failed to bind Clothing Data [{0}] using asset [{1}], as this causes the section to require {2} bones. The maximum per section is currently {3}."),
+			FText::FromString(GetName()),
+			FText::FromName(Asset->GetFName()),
+			BoneMap.Num(),
+			FGPUBaseSkinVertexFactory::GetMaxGPUSkinBones());
+		WarningNotification(Text);
+		return false;
+	}
+
+	FScopedSkeletalMeshPostEditChange SkeletalMeshPostEditChange(SkeletalMesh);  // When FScopedSkeletalMeshPostEditChange goes out of scope it causes the postedit change and components to be re-registered
+	SkelMeshSection.BoneMap = MoveTemp(BoneMap);
+
+	bool bRequireBoneChange = false;
+	for (FBoneIndexType& BoneIndex : SkelMeshSection.BoneMap)
+	{
+		if (!SkeletalMeshLODModel.RequiredBones.Contains(BoneIndex))
+		{
+			bRequireBoneChange = true;
+			if (SkeletalMesh->GetRefSkeleton().IsValidIndex(BoneIndex))
+			{
+				SkeletalMeshLODModel.RequiredBones.Add(BoneIndex);
+				SkeletalMeshLODModel.ActiveBoneIndices.AddUnique(BoneIndex);
+			}
+		}
+	}
+	if (bRequireBoneChange)
+	{
+		SkeletalMeshLODModel.RequiredBones.Sort();
+		SkeletalMesh->GetRefSkeleton().EnsureParentsExistAndSort(SkeletalMeshLODModel.ActiveBoneIndices);
+	}
+
+	// Set the asset guid
 	SkelMeshSection.CorrespondClothAssetIndex = ClothSimulationModelIndex;
 	SkelMeshSection.ClothingData.AssetGuid = AssetGuid;
 	SkelMeshSection.ClothingData.AssetLodIndex = AssetLodIndex;
@@ -155,28 +371,93 @@ bool UChaosClothAssetSKMClothingAsset::BindToSkeletalMesh(USkeletalMesh* Skeleta
 	SkelMeshSection.ClothMappingDataLODs.SetNum(1);
 	SkelMeshSection.ClothMappingDataLODs[0] = MoveTemp(MeshToMeshData);
 
+	// Set the asset index, used during rendering to pick the correct sim mesh buffer
+	int32 AssetIndex = INDEX_NONE;
+	verify(SkeletalMesh->GetMeshClothingAssets().Find(this, AssetIndex));  // A cloth asset can be added multiple times, but it will need a different clothing data each time
+	SkelMeshSection.CorrespondClothAssetIndex = AssetIndex;
+
+	// Update the extra cloth deformer mapping LOD bias using this cloth entry
+	// TODO: Implement a generic form of LOD Bias mapping in ClothingAssetUtils (need to rework the clothing asset interface for this)
+	//UpdateLODBiasMappings(SkeletalMesh, MeshLodIndex, SectionIndex);
+
 	return true;
 }
+
+void UChaosClothAssetSKMClothingAsset::UnbindFromSkeletalMesh(USkeletalMesh* SkeletalMesh)
+{
+	if (SkeletalMesh)
+	{
+		ClothingAssetUtils::UnbindFromSkeletalMesh(*this, *SkeletalMesh);
+	}
+}
+
+void UChaosClothAssetSKMClothingAsset::UnbindFromSkeletalMesh(USkeletalMesh* SkeletalMesh, const int32 MeshLodIndex)
+{
+	if (SkeletalMesh)
+	{
+		ClothingAssetUtils::UnbindFromSkeletalMesh(*this, *SkeletalMesh, MeshLodIndex);
+	}
+}
+
+void UChaosClothAssetSKMClothingAsset::UpdateAllLODBiasMappings(USkeletalMesh* SkeletalMesh)
+{
+	// TODO: Implement a generic form of LOD Bias mapping in ClothingAssetUtils (need to rework the clothing asset interface for this)
+}
+
+void UChaosClothAssetSKMClothingAsset::OnModelChanged()
+{
+	// Update the ClothSimulationModelId dropdown property
+	using namespace UE::Chaos::ClothAsset::Private;
+	ClothSimulationModelId = (ClothSimulationModelIndex != INDEX_NONE) ?
+		MakeClothSimulationModelIdString(Asset->GetClothSimulationModelName(ClothSimulationModelIndex), Asset->GetAssetGuid(ClothSimulationModelIndex)) :
+		FString();
+
+	// Unbind/rebind clothing with the new GUID
+	if (USkeletalMesh* const SkeletalMesh = Cast<USkeletalMesh>(GetOuter()))
+	{
+		TArray<ClothingAssetUtils::FClothingAssetMeshBinding> ClothingBindings;
+		UnbindClothingAndBackup(*SkeletalMesh, *this, ClothingBindings);
+		RestoreClothingFromBackup(*SkeletalMesh, *this, ClothingBindings);
+	}
+}
+
+FGuid UChaosClothAssetSKMClothingAsset::GetClothSimulationModelGuid() const
+{
+	int32 DashIndex;
+	return ClothSimulationModelId.FindLastChar(TEXT('-'), DashIndex) ?
+		FGuid(ClothSimulationModelId.RightChop(DashIndex + 1)) : FGuid();
+}
+
+FName UChaosClothAssetSKMClothingAsset::GetClothSimulationModelName() const
+{
+	int32 DashIndex;
+	return ClothSimulationModelId.FindLastChar(TEXT('-'), DashIndex) ?
+		FName(ClothSimulationModelId.LeftChop(ClothSimulationModelId.Len() - DashIndex)) : FName();
+}
 #endif  // #if WITH_EDITOR
 
-TArray<FName> UChaosClothAssetSKMClothingAsset::GetClothSimulationModelNames() const
+#if WITH_EDITORONLY_DATA
+TArray<FString> UChaosClothAssetSKMClothingAsset::GetClothSimulationModelIds() const
 {
-	TArray<FName> SimulationModels;
+	using namespace UE::Chaos::ClothAsset::Private;
+
+	TArray<FString> ClothSimulationModelIds;
 	if (Asset)
 	{
 		const int32 NumSimulationModels = Asset->GetNumClothSimulationModels();
-		SimulationModels.Reserve(NumSimulationModels);
-		for (int32 Index = 0; Index < NumSimulationModels; ++Index)
+		ClothSimulationModelIds.Reserve(NumSimulationModels);
+		for (int32 ModelIndex = 0; ModelIndex < NumSimulationModels; ++ModelIndex)
 		{
-			SimulationModels.Add(Asset->GetClothSimulationModelName(Index));
+			if (Asset->GetClothSimulationModel(ModelIndex).IsValid() && Asset->GetClothSimulationModel(ModelIndex)->GetNumLods())
+			{
+				ClothSimulationModelIds.Emplace(MakeClothSimulationModelIdString(
+					Asset->GetClothSimulationModelName(ModelIndex),
+					Asset->GetAssetGuid(ModelIndex)));
+			}
 		}
 	}
-	return SimulationModels;
-}
-
-void UChaosClothAssetSKMClothingAsset::OnAssetChanged()
-{
-	ClothSimulationModelName = Asset && Asset->GetNumClothSimulationModels() ? Asset->GetClothSimulationModelName(0) : NAME_None;
+	return ClothSimulationModelIds;
 }
+#endif  // #if WITH_EDITORONLY_DATA
 
 #undef LOCTEXT_NAMESPACE
diff --git a/Engine/Plugins/ChaosClothAsset/Source/ChaosClothAssetEngine/Private/ChaosClothAsset/ClothAssetSKMClothingSimulation.cpp b/Engine/Plugins/ChaosClothAsset/Source/ChaosClothAssetEngine/Private/ChaosClothAsset/ClothAssetSKMClothingSimulation.cpp
new file mode 100644
--- /dev/null
+++ b/Engine/Plugins/ChaosClothAsset/Source/ChaosClothAssetEngine/Private/ChaosClothAsset/ClothAssetSKMClothingSimulation.cpp
@@ -0,0 +1,295 @@
+// Copyright Epic Games, Inc. All Rights Reserved.
+
+#include "ChaosClothAsset/ClothAssetSKMClothingSimulation.h"
+#include "ChaosClothAsset/ClothAssetBase.h"
+#include "ChaosClothAsset/ClothAssetSKMClothingAsset.h"
+#include "ChaosClothAsset/ClothComponent.h"  // For FChaosClothSimulationProperties  TODO: Move FChaosClothSimulationProperties to its own header file?
+#include "ChaosClothAsset/ClothSimulationProxy.h"
+#include "ChaosClothAsset/CollisionSources.h"
+#include "Components/SkeletalMeshComponent.h"
+#include "HAL/IConsoleManager.h"
+#include "ClothCollisionData.h"
+
+#include UE_INLINE_GENERATED_CPP_BY_NAME(ClothAssetSKMClothingSimulation)
+
+namespace UE::Chaos::ClothAsset
+{
+	FSKMClothingSimulation::FSKMClothingSimulation() = default;
+	FSKMClothingSimulation::~FSKMClothingSimulation() = default;
+
+	void FSKMClothingSimulation::Initialize()
+	{
+		Shutdown();
+		CollisionSources = MakeUnique<FCollisionSources>();
+		ClothSimulationProxy = MakeShared<FClothSimulationProxy>(*this);
+	}
+
+	void FSKMClothingSimulation::Shutdown()
+	{
+		DestroyActors();
+		OwnerComponent = nullptr;
+		CollisionSources = nullptr;
+		ClothSimulationProxy = nullptr;
+	}
+
+	void FSKMClothingSimulation::CreateActor(USkeletalMeshComponent* InOwnerComponent, const UClothingAssetBase* ClothingAsset, int32 SimDataIndex)
+	{
+		check(InOwnerComponent);
+		if (const UChaosClothAssetSKMClothingAsset* const ClothAssetSKMClothingAsset = Cast<UChaosClothAssetSKMClothingAsset>(ClothingAsset))
+		{
+			if (const int32 ClothSimulationModelIndex = ClothAssetSKMClothingAsset->GetClothSimulationModelIndex(); ClothSimulationModelIndex != INDEX_NONE)
+			{
+				const UChaosClothAssetBase* const ClothAsset = ClothAssetSKMClothingAsset->GetAsset();
+				check(ClothAsset);
+
+				const int32 NumClothSimulationModels = ClothAsset->GetNumClothSimulationModels();
+				check(ClothSimulationModelIndex < NumClothSimulationModels);
+
+				// Keep track of the component
+				if (!OwnerComponent || OwnerComponent != InOwnerComponent)
+				{
+					checkf(!OwnerComponent, TEXT("IClothingSimulation::CreateActor() can't be called from multiple components."));
+					OwnerComponent = InOwnerComponent;
+
+					// Set the collision sources when CreateActor is setting the component
+					checkf(CollisionSources, TEXT("IClothingSimulation::Initialize() must be called before calling IClothingSimulation::CreateActor()."));
+					CollisionSources->SetOwnerComponent(OwnerComponent);
+					CollisionSources->SetCollideWithEnvironment(OwnerComponent->bCollideWithEnvironment);
+				}
+
+				// Add this asset to the asset list, matching its SimDataIndex position in the array
+				const int32 PrevNum = AssetClothSimulationProperties.Num();
+				AssetClothSimulationProperties.Reserve(SimDataIndex + 1);
+				for (int32 Index = PrevNum; Index <= SimDataIndex; ++Index)
+				{
+					AssetClothSimulationProperties.Emplace(nullptr, FChaosClothSimulationProperties());
+				}
+
+				// Set the properties for this SimDataIndex, but only if the ClothAsset/OutfitPiece isn't already used (the same cloth asset can be used in multiple render sections)
+				auto ClothAssetPredicate = [ClothAsset, ClothSimulationModelIndex](const FAssetClothSimulationProperties& Properties)
+					{
+						return Properties.Key &&
+							Properties.Key->GetAsset() == ClothAsset &&
+							Properties.Key->GetClothSimulationModelIndex() == ClothSimulationModelIndex;
+					};
+
+				const int32 FirstSimDataIndex =  AssetClothSimulationProperties.IndexOfByPredicate(ClothAssetPredicate);
+				if (FirstSimDataIndex == INDEX_NONE)
+				{
+					FChaosClothSimulationProperties& ClothSimulationProperties = AssetClothSimulationProperties[SimDataIndex].Value;
+
+					ClothSimulationProperties.Initialize(ClothAsset->GetCollections(ClothSimulationModelIndex));
+				}
+				else if (!ensureMsgf(SimDataIndex > FirstSimDataIndex, TEXT("CreateActor should be called with ascending SimDataIndex")))
+				{
+					Swap(AssetClothSimulationProperties[SimDataIndex].Value, AssetClothSimulationProperties[FirstSimDataIndex].Value);  // Otherwise the properties must be swapped
+				}
+
+				AssetClothSimulationProperties[SimDataIndex].Key = ClothAssetSKMClothingAsset;
+			}
+		}
+	}
+
+	void FSKMClothingSimulation::EndCreateActor()
+	{
+		if (ClothSimulationProxy.IsValid())
+		{
+			// Add all assets to the proxy
+			ClothSimulationProxy->PostConstructor();
+		}
+	}
+
+	void FSKMClothingSimulation::DestroyActors()
+	{
+		AssetClothSimulationProperties.Reset();
+	}
+
+	bool FSKMClothingSimulation::ShouldSimulate() const
+	{
+		if (!OwnerComponent)
+		{
+			return false;
+		}
+		for (const FAssetClothSimulationProperties& Properties : AssetClothSimulationProperties)
+		{
+			if (const UChaosClothAssetSKMClothingAsset* const ClothAssetSKMClothingAsset = Properties.Key)
+			{
+				if (ClothAssetSKMClothingAsset->GetClothSimulationModelIndex() != INDEX_NONE)
+				{
+					return true;
+				}
+			}
+		}
+		return false;
+	}
+
+	void FSKMClothingSimulation::Simulate(const IClothingSimulationContext* Context)
+	{
+		if (const FSKMClothingSimulationContext* const SKMClothingSimulationContext = static_cast<const FSKMClothingSimulationContext*>(Context))
+		{
+			if (ClothSimulationProxy.IsValid())
+			{
+				constexpr bool bUseParallelTask = false;  // The SKM component already deals with calling Simulate() in its own parallel task
+				ClothSimulationProxy->Tick_GameThread(SKMClothingSimulationContext->DeltaSeconds, bUseParallelTask);
+			}
+		}
+		// Reset the bNeedsResetRestLengths at every pass
+		bNeedsResetRestLengths = false;
+	}
+
+	void FSKMClothingSimulation::GetSimulationData(TMap<int32, FClothSimulData>& OutData, const USkeletalMeshComponent* /*InOwnerComponent*/, const USkinnedMeshComponent* /*OverrideComponent*/) const
+	{
+		if (ClothSimulationProxy.IsValid())
+		{
+			OutData = ClothSimulationProxy->GetCurrentSimulationData_AnyThread();  // TODO: Add a way to get the simulationdata without the unnecessary copy: GetCurrentSimulationData_GameThread(TMap<int32, FClothSimulData>& OutData)?
+		}
+		else
+		{
+			OutData.Reset();
+		}
+	}
+
+	FBoxSphereBounds FSKMClothingSimulation::GetBounds(const USkeletalMeshComponent* /*InOwnerComponent*/) const
+	{
+		if (ClothSimulationProxy.IsValid())
+		{
+			ClothSimulationProxy->CalculateBounds_AnyThread();
+		}
+		return FBoxSphereBounds(ForceInit);
+	}
+
+	void FSKMClothingSimulation::AddExternalCollisions(const FClothCollisionData& Data)
+	{
+		if (CollisionSources)
+		{
+			CollisionSources->AddExternalCollisions(Data);
+		}
+	}
+
+	void FSKMClothingSimulation::ClearExternalCollisions()
+	{
+		if (CollisionSources)
+		{
+			CollisionSources->ClearExternalCollisions();
+		}
+	}
+
+	void FSKMClothingSimulation::GetCollisions(FClothCollisionData& OutCollisions, bool bIncludeExternal) const
+	{
+		OutCollisions.Reset();
+	}
+
+	const USkinnedMeshComponent& FSKMClothingSimulation::GetOwnerComponent() const
+	{
+		checkf(OwnerComponent, TEXT("No actors, owner component not set. Call CreateActor first."));
+		return *Cast<USkinnedMeshComponent>(OwnerComponent);
+	}
+
+	TArray<const UChaosClothAssetBase*> FSKMClothingSimulation::GetAssets() const
+	{
+		TArray<const UChaosClothAssetBase*> Assets;
+		Assets.Reserve(AssetClothSimulationProperties.Num());
+		for (const FAssetClothSimulationProperties& Properties : AssetClothSimulationProperties)
+		{
+			if (const UChaosClothAssetSKMClothingAsset* const ClothAssetSKMClothingAsset = Properties.Key)
+			{
+				if (const UChaosClothAssetBase* const Asset = ClothAssetSKMClothingAsset->GetAsset())
+				{
+					Assets.AddUnique(Asset);
+				}
+			}
+		}
+		return Assets;
+	}
+
+	int32 FSKMClothingSimulation::GetSimulationGroupId(const UChaosClothAssetBase* Asset, int32 ModelIndex) const
+	{
+		for (int32 Index = 0; Index < AssetClothSimulationProperties.Num(); ++Index)
+		{
+			const FAssetClothSimulationProperties& Properties = AssetClothSimulationProperties[Index];
+			if (const UChaosClothAssetSKMClothingAsset* const ClothAssetSKMClothingAsset = Properties.Key)
+			{
+				if (ClothAssetSKMClothingAsset->GetAsset() == Asset && 
+					ClothAssetSKMClothingAsset->GetClothSimulationModelIndex() == ModelIndex)
+				{
+					return Index;
+				}
+			}
+		}
+		return INDEX_NONE;
+	}
+
+	const FReferenceSkeleton* FSKMClothingSimulation::GetReferenceSkeleton() const
+	{
+		checkf(OwnerComponent, TEXT("No actors, owner component not set. Call CreateActor first."));
+		return OwnerComponent->GetSkinnedAsset() ? &OwnerComponent->GetSkinnedAsset()->GetRefSkeleton() : nullptr;
+	}
+
+	EClothingTeleportMode FSKMClothingSimulation::GetClothTeleportMode() const
+	{
+		checkf(OwnerComponent, TEXT("No actors, owner component not set. Call CreateActor first."));
+		return OwnerComponent->ClothTeleportMode;
+	}
+
+	bool FSKMClothingSimulation::IsSimulationEnabled() const
+	{
+		return OwnerComponent->CanSimulateClothing();
+	}
+
+	bool FSKMClothingSimulation::IsSimulationSuspended() const
+	{
+		return OwnerComponent->IsClothingSimulationSuspended();
+	}
+
+	float FSKMClothingSimulation::GetClothGeometryScale() const
+	{
+		return OwnerComponent->ClothGeometryScale;
+	}
+
+	const TArray<TSharedPtr<const FManagedArrayCollection>>& FSKMClothingSimulation::GetPropertyCollections(const UChaosClothAssetBase* Asset, int32 ModelIndex) const
+	{
+		const int32 SimulationGroupId = GetSimulationGroupId(Asset, ModelIndex);
+		checkf(SimulationGroupId != INDEX_NONE, TEXT("Invalid arguments"));
+		return AssetClothSimulationProperties[SimulationGroupId].Value.PropertyCollections;
+	}
+
+	const TArray<TSharedPtr<const FManagedArrayCollection>>& FSKMClothingSimulation::GetSolverPropertyCollections() const
+	{
+		for (const FAssetClothSimulationProperties& Properties : AssetClothSimulationProperties)
+		{
+			if (const UChaosClothAssetSKMClothingAsset* const ClothAssetSKMClothingAsset = Properties.Key)
+			{
+				if (ClothAssetSKMClothingAsset->GetAsset())
+				{
+					return Properties.Value.PropertyCollections;
+				}
+			}
+		}
+		static const TArray<TSharedPtr<const FManagedArrayCollection>> EmptyManagedArrayCollectionArray;
+		return EmptyManagedArrayCollectionArray;
+	}
+
+	void FSKMClothingSimulation::AddReferencedObjects(FReferenceCollector& Collector)
+	{
+		for (FAssetClothSimulationProperties& Properties : AssetClothSimulationProperties)
+		{
+			if (Properties.Key && Properties.Value.ClothOutfitInteractor)
+			{
+				Collector.AddReferencedObject(Properties.Value.ClothOutfitInteractor);
+			}
+		}
+	}
+
+	FString FSKMClothingSimulation::GetReferencerName() const
+	{
+		return TEXT("UE::Chaos::ClothAsset::FSKMClothingSimulation");
+	}
+}  // namespace UE::Chaos::ClothAsset
+
+UChaosClothAssetSKMClothingSimulationFactory::UChaosClothAssetSKMClothingSimulationFactory() = default;
+UChaosClothAssetSKMClothingSimulationFactory::~UChaosClothAssetSKMClothingSimulationFactory() = default;
+
+bool UChaosClothAssetSKMClothingSimulationFactory::SupportsAsset(const UClothingAssetBase* Asset) const
+{
+	return Cast<UChaosClothAssetSKMClothingAsset>(Asset) != nullptr;
+}
diff --git a/Engine/Plugins/ChaosClothAsset/Source/ChaosClothAssetEngine/Private/ChaosClothAsset/ClothAssetSKMClothingSimulation.h b/Engine/Plugins/ChaosClothAsset/Source/ChaosClothAssetEngine/Private/ChaosClothAsset/ClothAssetSKMClothingSimulation.h
new file mode 100644
--- /dev/null
+++ b/Engine/Plugins/ChaosClothAsset/Source/ChaosClothAssetEngine/Private/ChaosClothAsset/ClothAssetSKMClothingSimulation.h
@@ -0,0 +1,145 @@
+// Copyright Epic Games, Inc. All Rights Reserved.
+
+#pragma once
+
+#include "ChaosClothAsset/ClothComponentAdapter.h"
+#include "ClothingSimulation.h"
+#include "ClothingSimulationFactory.h"
+#include "UObject/GCObject.h"
+#include "ClothAssetSKMClothingSimulation.generated.h"
+
+class USkinnedMeshComponent;
+class UChaosClothAssetSKMClothingAsset;
+struct FChaosClothSimulationProperties;
+
+namespace UE::Chaos::ClothAsset
+{
+	using FSKMClothingSimulationContext = FClothingSimulationContextCommon;
+	class FClothSimulationProxy;
+
+	class FSKMClothingSimulation final
+		: public FClothingSimulationCommon
+		, public IClothComponentAdapter
+		, public FGCObject
+	{
+	public:
+		FSKMClothingSimulation();
+		virtual ~FSKMClothingSimulation() override;
+
+	private:
+		//~ Begin IClothingSimulation Interface
+		virtual void Initialize() override;
+		virtual void Shutdown() override;
+
+		virtual IClothingSimulationContext* CreateContext() override
+		{
+			return new FSKMClothingSimulationContext();
+		}
+		virtual void DestroyContext(IClothingSimulationContext* Context) override
+		{
+			delete Context;
+		}
+
+		virtual void CreateActor(USkeletalMeshComponent* InOwnerComponent, const UClothingAssetBase* ClothingAsset, int32 SimDataIndex) override;
+		virtual void EndCreateActor() override;
+		virtual void DestroyActors() override;
+
+		virtual bool ShouldSimulate() const override;
+		virtual void Simulate(const IClothingSimulationContext* Context) override;
+		virtual void GetSimulationData(TMap<int32, FClothSimulData>& OutData, const USkeletalMeshComponent* OwnerComponent, const USkinnedMeshComponent* OverrideComponent) const override;
+
+		virtual FBoxSphereBounds GetBounds(const USkeletalMeshComponent* OwnerComponent) const override;
+
+		virtual void AddExternalCollisions(const FClothCollisionData& Data) override;
+		virtual void ClearExternalCollisions() override;
+		virtual void GetCollisions(FClothCollisionData& OutCollisions, bool bIncludeExternal = true) const override;
+		//~ End IClothingSimulation Interface
+
+		//~ Begin IClothComponentAdapter Interface
+		virtual const USkinnedMeshComponent& GetOwnerComponent() const override;
+		virtual FCollisionSources& GetCollisionSources() const override
+		{
+			return *CollisionSources;
+		}
+		virtual TArray<const UChaosClothAssetBase*> GetAssets() const override;
+		virtual int32 GetSimulationGroupId(const UChaosClothAssetBase* Asset, int32 ModelIndex) const override;
+
+		virtual const FReferenceSkeleton* GetReferenceSkeleton() const override;
+		virtual EClothingTeleportMode GetClothTeleportMode() const override;
+		virtual bool IsSimulationEnabled() const override;
+		virtual bool IsSimulationSuspended() const override;
+
+		virtual bool NeedsResetRestLengths() const override  // TODO: Make this a property interactor function
+		{
+			return bNeedsResetRestLengths;
+		}
+		virtual const FString& GetRestLengthsMorphTargetName() const override  // TODO: Make this a property interactor function
+		{
+			return RestLengthsMorphTargetName;
+		}
+		virtual float GetClothGeometryScale() const override;
+		virtual const TArray<TSharedPtr<const FManagedArrayCollection>>& GetPropertyCollections(const UChaosClothAssetBase* Asset, int32 ModelIndex) const override;
+		virtual const TArray<TSharedPtr<const FManagedArrayCollection>>& GetSolverPropertyCollections() const override;
+		//~ End IClothComponentAdapter Interface
+
+		//~ Begin FGCObject Interface
+		virtual void AddReferencedObjects(FReferenceCollector& Collector) override;
+		virtual FString GetReferencerName() const override;
+		//~ End FGCObject Interface
+
+		USkeletalMeshComponent* OwnerComponent = nullptr;
+		TUniquePtr<FCollisionSources> CollisionSources;
+		TSharedPtr<FClothSimulationProxy> ClothSimulationProxy;
+		
+		using FAssetClothSimulationProperties = TPair<const UChaosClothAssetSKMClothingAsset*, FChaosClothSimulationProperties>;
+		TArray<FAssetClothSimulationProperties> AssetClothSimulationProperties;  // Asset properties, matches CreateActor's SimDataIndex, some assets may appear multiple times, some entries might have a null asset
+	
+		bool bNeedsResetRestLengths = false;
+		FString RestLengthsMorphTargetName;
+	};
+}  // namespace UE::Chaos::ClothAsset
+
+UCLASS(MinimalAPI)
+class UChaosClothAssetSKMClothingSimulationFactory final : public UClothingSimulationFactory
+{
+	GENERATED_BODY()
+
+public:
+	UChaosClothAssetSKMClothingSimulationFactory();
+	virtual ~UChaosClothAssetSKMClothingSimulationFactory() override;
+
+private:
+	//~ Begin UClothingSimulationFactory Interface
+	virtual IClothingSimulation* CreateSimulation() override
+	{
+		return static_cast<IClothingSimulation*>(new UE::Chaos::ClothAsset::FSKMClothingSimulation());
+	}
+
+	virtual void DestroySimulation(IClothingSimulation* Simulation) override
+	{
+		delete Simulation;
+	}
+
+	virtual bool SupportsAsset(const UClothingAssetBase* Asset) const override;
+
+	virtual bool SupportsRuntimeInteraction() const override
+	{
+		return false;
+	}
+
+	virtual UClothingSimulationInteractor* CreateInteractor() override
+	{
+		return nullptr;  // NewObject<UChaosClothAssetSimulationInteractor>(GetTransientPackage());
+	}
+
+	virtual TArrayView<const TSubclassOf<UClothConfigBase>> GetClothConfigClasses() const override
+	{
+		return TArrayView<const TSubclassOf<UClothConfigBase>>();
+	}
+
+	const UEnum* GetWeightMapTargetEnum() const override
+	{
+		return nullptr;
+	}
+	//~ End UClothingSimulationFactory Interface
+};
diff --git a/Engine/Plugins/ChaosClothAsset/Source/ChaosClothAssetEngine/Private/ChaosClothAsset/ClothAssetSKMSimulationFactory.cpp b/Engine/Plugins/ChaosClothAsset/Source/ChaosClothAssetEngine/Private/ChaosClothAsset/ClothAssetSKMSimulationFactory.cpp
deleted file mode 100644
diff --git a/Engine/Plugins/ChaosClothAsset/Source/ChaosClothAssetEngine/Private/ChaosClothAsset/ClothAssetSKMSimulationFactory.h b/Engine/Plugins/ChaosClothAsset/Source/ChaosClothAssetEngine/Private/ChaosClothAsset/ClothAssetSKMSimulationFactory.h
deleted file mode 100644
diff --git a/Engine/Plugins/ChaosClothAsset/Source/ChaosClothAssetEngine/Private/ChaosClothAsset/ClothComponent.cpp b/Engine/Plugins/ChaosClothAsset/Source/ChaosClothAssetEngine/Private/ChaosClothAsset/ClothComponent.cpp
--- a/Engine/Plugins/ChaosClothAsset/Source/ChaosClothAssetEngine/Private/ChaosClothAsset/ClothComponent.cpp
+++ b/Engine/Plugins/ChaosClothAsset/Source/ChaosClothAssetEngine/Private/ChaosClothAsset/ClothComponent.cpp
@@ -841,6 +841,20 @@ void UChaosClothComponent::PostProcessSimulation(const float /*DeltaTime*/)
 	}
 }
 
+const FReferenceSkeleton* UChaosClothComponent::GetReferenceSkeleton() const
+{
+	return GetAsset() ? &GetAsset()->GetRefSkeleton() : nullptr;
+}
+
+int32 UChaosClothComponent::GetSimulationGroupId(const UChaosClothAssetBase* InAsset, int32 ModelIndex) const
+{
+	return
+		InAsset == GetAsset() &&
+		ModelIndex >= 0 &&
+		ModelIndex < InAsset->GetNumClothSimulationModels() &&
+		InAsset->GetClothSimulationModel(ModelIndex).IsValid() ? ModelIndex : INDEX_NONE;
+}
+
 FDataflowSimulationProxy* UChaosClothComponent::GetSimulationProxy()
 {
 	return ClothSimulationProxy.Get();
diff --git a/Engine/Plugins/ChaosClothAsset/Source/ChaosClothAssetEngine/Private/ChaosClothAsset/ClothSimulationContext.cpp b/Engine/Plugins/ChaosClothAsset/Source/ChaosClothAssetEngine/Private/ChaosClothAsset/ClothSimulationContext.cpp
--- a/Engine/Plugins/ChaosClothAsset/Source/ChaosClothAssetEngine/Private/ChaosClothAsset/ClothSimulationContext.cpp
+++ b/Engine/Plugins/ChaosClothAsset/Source/ChaosClothAssetEngine/Private/ChaosClothAsset/ClothSimulationContext.cpp
@@ -2,32 +2,46 @@
 
 #include "ChaosClothAsset/ClothSimulationContext.h"
 #include "ChaosClothAsset/ClothAsset.h"
-#include "ChaosClothAsset/ClothComponent.h"
+#include "ChaosClothAsset/ClothComponentAdapter.h"
 #include "ChaosClothAsset/ClothSimulationModel.h"
+#include "Components/SkinnedMeshComponent.h"
 #include "Engine/World.h"
 #include "HAL/IConsoleManager.h"
 #include "Rendering/SkeletalMeshRenderData.h"
+#include "ClothingSystemRuntimeTypes.h"
 #include "SceneInterface.h"
 
 namespace UE::Chaos::ClothAsset
 {
-	void FClothSimulationContext::Fill(const UChaosClothComponent& ClothComponent, float InDeltaTime, float MaxDeltaTime, bool bIsInitialization, FClothingSimulationCacheData* InCacheData)
+	void FClothSimulationContext::Fill(const IClothComponentAdapter& ClothComponentAdapter, float InDeltaTime, float MaxDeltaTime, bool bIsInitialization, FClothingSimulationCacheData* InCacheData)
 	{
 		TRACE_CPUPROFILER_EVENT_SCOPE(FClothSimulationContext_Fill);
 		// Set the time
 		DeltaTime = FMath::Min(InDeltaTime, MaxDeltaTime);
 
 		// Set the current LOD index
-		LodIndex = ClothComponent.GetPredictedLODLevel();
-		const UChaosClothAssetBase* const Asset = ClothComponent.GetAsset();
-		if (!bIsInitialization &&
-			(!Asset ||
-			!Asset->GetResourceForRendering() ||
-			!Asset->GetResourceForRendering()->LODRenderData.IsValidIndex(LodIndex) ||
-			!Asset->GetResourceForRendering()->LODRenderData[LodIndex].HasClothData()))
+		LodIndex = ClothComponentAdapter.GetOwnerComponent().GetPredictedLODLevel();
+
+		// Retrieve the list of assets to simulate
+		const TArray<const UChaosClothAssetBase*> Assets = ClothComponentAdapter.GetAssets();
+		if (!bIsInitialization)
 		{
-			// This Lod will not simulate, so there's no need to update the context.
-			return;
+			bool bHasAssetValidClothRenderData = false;
+			for (const UChaosClothAssetBase* const Asset : Assets)
+			{
+				if (Asset &&
+					Asset->GetResourceForRendering() &&
+					Asset->GetResourceForRendering()->LODRenderData.IsValidIndex(LodIndex) &&
+					Asset->GetResourceForRendering()->LODRenderData[LodIndex].HasClothData())
+				{
+					bHasAssetValidClothRenderData = true;
+					break;  
+				}
+			}
+			if (!bHasAssetValidClothRenderData)
+			{
+				return;  // This Lod will not simulate, so no need to update the context.
+			}
 		}
 
 		// Set the teleport mode
@@ -35,16 +49,16 @@ namespace UE::Chaos::ClothAsset
 		constexpr float MaxDeltaTimeTeleportMultiplierDefault = 1.5f;
 		const float MaxDeltaTimeTeleportMultiplier = CVarMaxDeltaTimeTeleportMultiplier ? CVarMaxDeltaTimeTeleportMultiplier->GetFloat() : MaxDeltaTimeTeleportMultiplierDefault;
 
-		const EClothingTeleportMode ComponentTeleportMode = ClothComponent.GetClothTeleportMode();
+		const EClothingTeleportMode ComponentTeleportMode = ClothComponentAdapter.GetClothTeleportMode();
 
 		bTeleport = (InDeltaTime > MaxDeltaTime * MaxDeltaTimeTeleportMultiplier) ? true : 
 			(ComponentTeleportMode >= EClothingTeleportMode::Teleport);
 		bReset = ComponentTeleportMode == EClothingTeleportMode::TeleportAndReset;
 		
-		bResetRestLengthsWithMorphTarget = ClothComponent.NeedsResetRestLengths();
+		bResetRestLengthsWithMorphTarget = ClothComponentAdapter.NeedsResetRestLengths();
 		if (bResetRestLengthsWithMorphTarget)
 		{
-			ResetRestLengthsMorphTargetName = ClothComponent.GetRestLengthsMorphTargetName();
+			ResetRestLengthsMorphTargetName = ClothComponentAdapter.GetRestLengthsMorphTargetName();
 		}
 
 		VelocityScale = (!bTeleport && !bReset && InDeltaTime > 0.f) ?
@@ -52,18 +66,18 @@ namespace UE::Chaos::ClothAsset
 			bReset ? 1.f : 0.f;  // Set to 0 when teleporting and 1 when resetting to match the internal solver's behavior
 
 		// Copy component transform
-		ComponentTransform = ClothComponent.GetComponentTransform();
+		ComponentTransform = ClothComponentAdapter.GetOwnerComponent().GetComponentTransform();
 
 		// Copy solver geometry scale
-		SolverGeometryScale = ClothComponent.GetClothGeometryScale();;
+		SolverGeometryScale = ClothComponentAdapter.GetClothGeometryScale();
 
 		// Update bone transforms
-		const FReferenceSkeleton* const ReferenceSkeleton = Asset ? &Asset->GetRefSkeleton() : nullptr;
+		const FReferenceSkeleton* const ReferenceSkeleton = ClothComponentAdapter.GetReferenceSkeleton();
 		int32 NumBones = ReferenceSkeleton ? ReferenceSkeleton->GetNum() : 0;
 
-		if (USkinnedMeshComponent* const LeaderComponent = ClothComponent.LeaderPoseComponent.Get())
+		if (USkinnedMeshComponent* const LeaderComponent = ClothComponentAdapter.GetOwnerComponent().LeaderPoseComponent.Get())
 		{
-			const TArray<int32>& LeaderBoneMap = ClothComponent.GetLeaderBoneMap();
+			const TArray<int32>& LeaderBoneMap = ClothComponentAdapter.GetOwnerComponent().GetLeaderBoneMap();
 			if (!LeaderBoneMap.Num())
 			{
 				// This case indicates an invalid leader pose component (e.g. no skeletal mesh)
@@ -107,7 +121,7 @@ namespace UE::Chaos::ClothAsset
 		}
 		else if (!bIsInitialization)  // Initializations must be done in bind pose
 		{
-			BoneTransforms = ClothComponent.GetComponentSpaceTransforms();
+			BoneTransforms = ClothComponentAdapter.GetOwnerComponent().GetComponentSpaceTransforms();
 		}
 		else
 		{
@@ -122,20 +136,23 @@ namespace UE::Chaos::ClothAsset
 		bool bSetRefToLocalMatricesToIdentity = true;
 		if (!bIsInitialization)
 		{
-			if (const FSkeletalMeshRenderData* const RenderData = Asset ? Asset->GetResourceForRendering() : nullptr)
+			for (const UChaosClothAssetBase* const Asset : Assets)
 			{
-				if (RenderData->LODRenderData.IsValidIndex(LodIndex))  // TODO Disable this conditional to drive skeletal meshes
+				if (const FSkeletalMeshRenderData* const RenderData = Asset ? Asset->GetResourceForRendering() : nullptr)
 				{
-					for (int32 ModelIndex = 0; ModelIndex < Asset->GetNumClothSimulationModels(); ++ModelIndex)
+					if (RenderData->LODRenderData.IsValidIndex(LodIndex))  // TODO Disable this conditional to drive skeletal meshes
 					{
-						const TSharedPtr<const FChaosClothSimulationModel>& ClothModel = Asset->GetClothSimulationModel(ModelIndex);
-						if (ClothModel && ClothModel->IsValidLodIndex(LodIndex))
+						for (int32 ModelIndex = 0; ModelIndex < Asset->GetNumClothSimulationModels(); ++ModelIndex)
 						{
-							RequiredExtraBones.Append(ClothModel->ClothSimulationLodModels[LodIndex].RequiredExtraBoneIndices);
+							const TSharedPtr<const FChaosClothSimulationModel>& ClothModel = Asset->GetClothSimulationModel(ModelIndex);
+							if (ClothModel && ClothModel->IsValidLodIndex(LodIndex))
+							{
+								RequiredExtraBones.Append(ClothModel->ClothSimulationLodModels[LodIndex].RequiredExtraBoneIndices);
+							}
 						}
+						ClothComponentAdapter.GetOwnerComponent().GetCurrentRefToLocalMatrices(RefToLocalMatrices, LodIndex, &RequiredExtraBones);
+						bSetRefToLocalMatricesToIdentity = false;
 					}
-					ClothComponent.GetCurrentRefToLocalMatrices(RefToLocalMatrices, LodIndex, &RequiredExtraBones);
-					bSetRefToLocalMatricesToIdentity = false;
 				}
 			}
 		}
@@ -151,14 +168,14 @@ namespace UE::Chaos::ClothAsset
 		}
 		
 		// Update gravity
-		const UWorld* const World = ClothComponent.GetWorld();
+		const UWorld* const World = ClothComponentAdapter.GetOwnerComponent().GetWorld();
 		constexpr float EarthGravity = -981.f;
 		WorldGravity = FVector(0.f, 0.f, World ? World->GetGravityZ() : EarthGravity);
 
 		// Update wind velocity
 		if (World && World->Scene && World->IsGameWorld())
 		{
-			const FVector Position = ClothComponent.GetComponentTransform().GetTranslation();
+			const FVector Position = ClothComponentAdapter.GetOwnerComponent().GetComponentTransform().GetTranslation();
 
 			float WindSpeed;
 			float WindMinGust;
diff --git a/Engine/Plugins/ChaosClothAsset/Source/ChaosClothAssetEngine/Private/ChaosClothAsset/ClothSimulationContext.h b/Engine/Plugins/ChaosClothAsset/Source/ChaosClothAssetEngine/Private/ChaosClothAsset/ClothSimulationContext.h
--- a/Engine/Plugins/ChaosClothAsset/Source/ChaosClothAssetEngine/Private/ChaosClothAsset/ClothSimulationContext.h
+++ b/Engine/Plugins/ChaosClothAsset/Source/ChaosClothAssetEngine/Private/ChaosClothAsset/ClothSimulationContext.h
@@ -4,12 +4,12 @@
 
 #include "Containers/Array.h"
 #include "Math/Transform.h"
-#include "ClothingSimulation.h"
-
-class UChaosClothComponent;
+#include "ClothingSimulationCacheData.h"
 
 namespace UE::Chaos::ClothAsset
 {
+	class IClothComponentAdapter;
+
 	/**
 	 * Cloth simulation context.
 	 * Class used to pass data between the cloth component and its simulation proxy.
@@ -59,7 +59,7 @@ namespace UE::Chaos::ClothAsset
 		FClothingSimulationCacheData CacheData;
 
 		/** Fill the context from data collected from the specified component. */
-		void Fill(const UChaosClothComponent& ClothComponent, float InDeltaTime, float MaxDeltaTime, bool bIsInitialization = false, FClothingSimulationCacheData* CacheData = nullptr);
+		void Fill(const IClothComponentAdapter& ClothComponentAdapter, float InDeltaTime, float MaxDeltaTime, bool bIsInitialization = false, FClothingSimulationCacheData* CacheData = nullptr);
 
 	private:
 		TArray<uint16> RequiredExtraBones;  // Avoids reallocations while calculating RefToLocalMatrices
diff --git a/Engine/Plugins/ChaosClothAsset/Source/ChaosClothAssetEngine/Private/ChaosClothAsset/ClothSimulationProxy.cpp b/Engine/Plugins/ChaosClothAsset/Source/ChaosClothAssetEngine/Private/ChaosClothAsset/ClothSimulationProxy.cpp
--- a/Engine/Plugins/ChaosClothAsset/Source/ChaosClothAssetEngine/Private/ChaosClothAsset/ClothSimulationProxy.cpp
+++ b/Engine/Plugins/ChaosClothAsset/Source/ChaosClothAssetEngine/Private/ChaosClothAsset/ClothSimulationProxy.cpp
@@ -99,7 +99,7 @@ namespace UE::Chaos::ClothAsset
 
 		void DoTask(ENamedThreads::Type CurrentThread, const FGraphEventRef& MyCompletionGraphEvent)
 		{
-			FScopeCycleCounterUObject ContextScope(ClothSimulationProxy.ClothComponent.GetSkinnedAsset());
+			FScopeCycleCounterUObject ContextScope(ClothSimulationProxy.ClothComponentAdapter.GetOwnerComponent().GetSkinnedAsset());
 			SCOPE_CYCLE_COUNTER(STAT_ClothTotalTime);
 			CSV_SCOPED_TIMING_STAT(Animation, Cloth);
 			switch (Type)
@@ -119,10 +119,10 @@ namespace UE::Chaos::ClothAsset
 		EType Type = EType::Tick;
 	};
 
-	FClothSimulationProxy::FClothSimulationProxy(const UChaosClothComponent& InClothComponent)
-		: ClothComponent(InClothComponent)
+	FClothSimulationProxy::FClothSimulationProxy(const IClothComponentAdapter& ClothComponentAdapter)
+		: ClothComponentAdapter(ClothComponentAdapter)
 		, ClothSimulationContext(MakeUnique<FClothSimulationContext>())
-		, CollisionSourcesProxy(MakeUnique<FCollisionSourcesProxy>(InClothComponent.GetCollisionSources()))
+		, CollisionSourcesProxy(MakeUnique<FCollisionSourcesProxy>(ClothComponentAdapter.GetCollisionSources()))
 		, Solver(nullptr)
 		, Visualization(nullptr)
 		, MaxDeltaTime(UPhysicsSettings::Get()->MaxPhysicsDeltaTime)
@@ -169,7 +169,7 @@ namespace UE::Chaos::ClothAsset
 		Cloths.Reset();
 
 		// Create solver config simulation thread object first. Need to know which solver type we're creating.
-		const int32 SolverConfigIndex = Configs.Emplace(MakeUnique<FClothingSimulationConfig>(ClothComponent.GetSolverPropertyCollections()));
+		const int32 SolverConfigIndex = Configs.Emplace(MakeUnique<FClothingSimulationConfig>(ClothComponentAdapter.GetSolverPropertyCollections()));
 
 		// Use new SoftsEvolution, not PBDEvolution.
 		constexpr bool bUseLegacySolver = false;
@@ -181,7 +181,7 @@ PRAGMA_ENABLE_DEPRECATION_WARNINGS
 		// Need a valid context to initialize the mesh
 		constexpr bool bIsInitialization = true;
 		constexpr Softs::FSolverReal NoAdvanceDt = 0.f;
-		ClothSimulationContext->Fill(ClothComponent, NoAdvanceDt, MaxDeltaTime, bIsInitialization);
+		ClothSimulationContext->Fill(ClothComponentAdapter, NoAdvanceDt, MaxDeltaTime, bIsInitialization);
 
 		// Setup startup transforms
 		constexpr bool bNeedsReset = true;
@@ -190,39 +190,51 @@ PRAGMA_ENABLE_DEPRECATION_WARNINGS
 		Solver->SetLocalSpaceLocation((FVec3)ClothSimulationContext->ComponentTransform.GetLocation(), bNeedsReset);
 		Solver->SetLocalSpaceRotation((FQuat)ClothSimulationContext->ComponentTransform.GetRotation());
 
-		// Create mesh simulation thread object
-		const UChaosClothAssetBase* const Asset = ClothComponent.GetAsset();
-		if (!Asset)
+		// Create mesh simulation thread objects
+		for (const UChaosClothAssetBase* const Asset : ClothComponentAdapter.GetAssets())
 		{
-			return;
-		}
-		const FReferenceSkeleton* const ReferenceSkeleton = &Asset->GetRefSkeleton();
-		
-		for (int32 ModelIndex = 0; ModelIndex < Asset->GetNumClothSimulationModels(); ++ModelIndex)
-		{
-			const TSharedPtr<const FChaosClothSimulationModel>& ClothSimulationModel = Asset->GetClothSimulationModel(ModelIndex);
+			const int32 NumClothSimulationModels = Asset ? Asset->GetNumClothSimulationModels() : 0;
+			for (int32 ModelIndex = 0; ModelIndex < NumClothSimulationModels; ++ModelIndex)
+			{
+				const int32 SimulationGroupId = ClothComponentAdapter.GetSimulationGroupId(Asset, ModelIndex);
+				if (SimulationGroupId != INDEX_NONE)  // INDEX_NONE means this model shouldn't be added to the simulation
+				{
+					const TSharedPtr<const FChaosClothSimulationModel> ClothSimulationModel = Asset->GetClothSimulationModel(ModelIndex);
 
-			FString DebugName;
+					FString DebugName;
 #if !UE_BUILD_SHIPPING
-			DebugName = ClothComponent.GetOwner() ?
-				FString::Format(TEXT("{0}|{1}"), { ClothComponent.GetOwner()->GetActorNameOrLabel(), ClothComponent.GetName() }) :
-				ClothComponent.GetName();
+					DebugName = ClothComponentAdapter.GetOwnerComponent().GetOwner() ?
+						FString::Format(TEXT("{0}|{1}|{2}|Model[{3}]"),
+							{
+								ClothComponentAdapter.GetOwnerComponent().GetOwner()->GetActorNameOrLabel(),
+								ClothComponentAdapter.GetOwnerComponent().GetName(),
+								Asset->GetName(),
+								ModelIndex
+							}) :
+						FString::Format(TEXT("{0}|{1}|Model[{2}]"),
+							{
+								ClothComponentAdapter.GetOwnerComponent().GetName(),
+								Asset->GetName(),
+								ModelIndex
+							});
 #endif
-			const int32 MeshIndex = Meshes.Emplace(MakeUnique<FClothSimulationMesh>(*ClothSimulationModel, *ClothSimulationContext, Asset->GetCollections(ModelIndex), DebugName));
+					const int32 MeshIndex = Meshes.Emplace(MakeUnique<FClothSimulationMesh>(*ClothSimulationModel, *ClothSimulationContext, Asset->GetCollections(ModelIndex), DebugName));
 
-			// Create collider simulation thread object
-			const int32 ColliderIndex = Colliders.Emplace(MakeUnique<FClothingSimulationCollider>(Asset->GetPhysicsAssetForModel(ModelIndex), ReferenceSkeleton));
-			Colliders[ColliderIndex]->SetCollisionData(&CollisionSourcesProxy->GetCollisionData());
+					// Create collider simulation thread object
+					const int32 ColliderIndex = Colliders.Emplace(MakeUnique<FClothingSimulationCollider>(Asset->GetPhysicsAssetForModel(ModelIndex), &Asset->GetRefSkeleton()));
+					Colliders[ColliderIndex]->SetCollisionData(&CollisionSourcesProxy->GetCollisionData());
 
-			// Create cloth config simulation thread object
-			const int32 ClothConfigIndex = Configs.Emplace(MakeUnique<FClothingSimulationConfig>(ClothComponent.GetPropertyCollections(ModelIndex)));
+					// Create cloth config simulation thread object
+					const int32 ClothConfigIndex = Configs.Emplace(MakeUnique<FClothingSimulationConfig>(ClothComponentAdapter.GetPropertyCollections(Asset, ModelIndex)));
 
-			// Create cloth simulation thread object
-			const int32 ClothIndex = Cloths.Emplace(MakeUnique<FClothingSimulationCloth>(
-				Configs[ClothConfigIndex].Get(),
-				Meshes[MeshIndex].Get(),
-				ColliderIndex != INDEX_NONE ? TArray<FClothingSimulationCollider*>({ Colliders[ColliderIndex].Get() }) : TArray<FClothingSimulationCollider*>(),
-				ModelIndex));
+					// Create cloth simulation thread object
+					const int32 ClothIndex = Cloths.Emplace(MakeUnique<FClothingSimulationCloth>(
+						Configs[ClothConfigIndex].Get(),
+						Meshes[MeshIndex].Get(),
+						ColliderIndex != INDEX_NONE ? TArray<FClothingSimulationCollider*>({ Colliders[ColliderIndex].Get() }) : TArray<FClothingSimulationCollider*>(),
+						SimulationGroupId));
+				}
+			}
 		}
 	}
 
@@ -273,7 +285,7 @@ PRAGMA_ENABLE_DEPRECATION_WARNINGS
 
 		// Set start pose (update the context, then the solver without advancing the simulation)
 		constexpr Softs::FSolverReal NoAdvanceDt = 0.f;
-		ClothSimulationContext->Fill(ClothComponent, NoAdvanceDt, MaxDeltaTime);
+		ClothSimulationContext->Fill(ClothComponentAdapter, NoAdvanceDt, MaxDeltaTime);
 		Solver->Update((Softs::FSolverReal)ClothSimulationContext->DeltaTime);
 		bIsInitialized = true;
 	}
@@ -311,7 +323,7 @@ PRAGMA_ENABLE_DEPRECATION_WARNINGS
 
 		bIsPreProcessed = true;
 
-		if (!ClothComponent.IsSimulationEnabled())
+		if (!ClothComponentAdapter.IsSimulationEnabled())
 		{
 			return;
 		}
@@ -349,18 +361,27 @@ PRAGMA_ENABLE_DEPRECATION_WARNINGS
 		}
 
 		// Check that the render mesh current LOD isn't just fully skinned
+		const TArray<const UChaosClothAssetBase*> Assets = ClothComponentAdapter.GetAssets();
 		const int32 LodIndex = ClothSimulationContext->LodIndex;
-		const UChaosClothAssetBase* const Asset = ClothComponent.GetAsset();
-		if (!Asset ||
-			!Asset->GetResourceForRendering() ||
-			!Asset->GetResourceForRendering()->LODRenderData.IsValidIndex(LodIndex) ||
-			!Asset->GetResourceForRendering()->LODRenderData[LodIndex].HasClothData())
+		bool bHasAssetValidClothRenderData = false;
+		for (const UChaosClothAssetBase* const Asset : Assets)
+		{
+			if (Asset &&
+				Asset->GetResourceForRendering() &&
+				Asset->GetResourceForRendering()->LODRenderData.IsValidIndex(LodIndex) &&
+				Asset->GetResourceForRendering()->LODRenderData[LodIndex].HasClothData())
+			{
+				bHasAssetValidClothRenderData = true;
+				break;
+			}
+		}
+		if (!bHasAssetValidClothRenderData)
 		{
 			return false;
 		}
 
 		// Check whether there some actual simulation needs to happen
-		bIsSimulating = (DeltaTime > 0.f && !ClothComponent.IsSimulationSuspended() && bIsSolverEnabled);
+		bIsSimulating = (DeltaTime > 0.f && !ClothComponentAdapter.IsSimulationSuspended() && bIsSolverEnabled);
 		if (!bIsSimulating)
 		{
 			return false;
@@ -370,11 +391,11 @@ PRAGMA_ENABLE_DEPRECATION_WARNINGS
 		InitializeConfigs();
 
 		// Update world forces
-		if (UWorld* const World = ClothComponent.GetWorld())
+		if (UWorld* const World = ClothComponentAdapter.GetOwnerComponent().GetWorld())
 		{
 			if (UPhysicsFieldComponent* const PhysicsField = World->PhysicsField)
 			{
-				const FBox BoundingBox = CalculateBounds_AnyThread().GetBox().TransformBy(ClothComponent.GetComponentTransform());
+				const FBox BoundingBox = CalculateBounds_AnyThread().GetBox().TransformBy(ClothComponentAdapter.GetOwnerComponent().GetComponentTransform());
 
 				PhysicsField->FillTransientCommands(false, BoundingBox, Solver->GetTime(), Solver->GetPerSolverField().GetTransientCommands());
 				PhysicsField->FillPersistentCommands(false, BoundingBox, Solver->GetTime(), Solver->GetPerSolverField().GetPersistentCommands());
@@ -406,7 +427,7 @@ PRAGMA_ENABLE_DEPRECATION_WARNINGS
 		{
 
 			// If the simulation is enabled, then it is suspended or reading from cache and the simulation data still needs updating (transforms and LODs)
-			if (ClothComponent.IsSimulationEnabled())
+			if (ClothComponentAdapter.IsSimulationEnabled())
 			{
 				// Take care of the LOD switching, as the simulation won't do it
 				UpdateClothLODs();
@@ -419,7 +440,7 @@ PRAGMA_ENABLE_DEPRECATION_WARNINGS
 		}
 	}
 
-	bool FClothSimulationProxy::Tick_GameThread(float DeltaTime)
+	bool FClothSimulationProxy::Tick_GameThread(float DeltaTime, bool bUseParallelTask)
 	{
 		SCOPE_CYCLE_COUNTER(STAT_ClothSimulationProxy_TickGame);
 
@@ -428,18 +449,19 @@ PRAGMA_ENABLE_DEPRECATION_WARNINGS
 
 		if (bIsSimulating)
 		{
-			// Start the the cloth simulation thread
-			ParallelTask = TGraphTask<FClothSimulationProxyParallelTask>::CreateTask(nullptr, ENamedThreads::GameThread).ConstructAndDispatchWhenReady(*this);
+			if (bUseParallelTask)
+			{
+				ParallelTask = TGraphTask<FClothSimulationProxyParallelTask>::CreateTask(nullptr, ENamedThreads::GameThread).ConstructAndDispatchWhenReady(*this);
 
-			// Note: Post simulate and post process will be handled when the parallel task completes
-		}
-		else
-		{
-			// The post simulate doesn't do much right now, but better still call it for completion in case it one day does
-			PostSimulate_GameThread();
-			// The post process only happens when there isn't any simulation
-			PostProcess_GameThread();
+				return true;  // Note: Post simulate and post process will be handled when the parallel task completes
+			}
+
+			Tick();
 		}
+
+		PostSimulate_GameThread();
+		PostProcess_GameThread();
+
 		return bIsSimulating;
 	}
 
@@ -651,7 +673,7 @@ PRAGMA_ENABLE_DEPRECATION_WARNINGS
 			const int32 LODIndex = Cloth->GetMesh()->GetLODIndex();
 			if (LODIndex != Cloth->GetLODIndex(Solver.Get()))
 			{
-				if (!ClothComponent.IsSimulationEnabled())
+				if (!ClothComponentAdapter.IsSimulationEnabled())
 				{
 					// Mark the cloth as needing to be reset so it doesn't both proxy-deforming lod transitions.
 					Cloth->Reset();
@@ -674,31 +696,46 @@ PRAGMA_ENABLE_DEPRECATION_WARNINGS
 		CSV_SCOPED_TIMING_STAT(Animation, Cloth);
 		TRACE_CPUPROFILER_EVENT_SCOPE(FClothSimulationProxy_WriteSimulationData);
 		SCOPE_CYCLE_COUNTER(STAT_ClothSimulationProxy_WriteSimulationData);
-		SCOPE_CYCLE_COUNTER(STAT_ClothWriteback)
+		SCOPE_CYCLE_COUNTER(STAT_ClothWriteback);
+
+		// Check if our bone map is actually valid, if not there is no clothing data to build
+		const USkinnedMeshComponent& OwnerComponent = ClothComponentAdapter.GetOwnerComponent();
+		const USkinnedMeshComponent* const LeaderPoseComponent = OwnerComponent.LeaderPoseComponent.Get();
 
-		USkinnedMeshComponent* LeaderPoseComponent = nullptr;
-		if (ClothComponent.LeaderPoseComponent.IsValid())
+		if (LeaderPoseComponent && !OwnerComponent.GetLeaderBoneMap().Num())
 		{
-			LeaderPoseComponent = ClothComponent.LeaderPoseComponent.Get();
+			CurrentSimulationData.Reset();
+			return;
+		}
+
+		// Get the list of used indices and cache the result of the ClothSimulData map find
+		TArray<int32, TInlineAllocator<16>> AssetIndices;
+		AssetIndices.Reserve(Cloths.Num());
 
-			// Check if our bone map is actually valid, if not there is no clothing data to build
-			if (!ClothComponent.GetLeaderBoneMap().Num())
+		for (const TUniquePtr<FClothingSimulationCloth>& Cloth : Cloths)
+		{
+			if (Cloth &&
+				Cloth->GetMesh() &&
+				Cloth->GetParticleRangeId(Solver.Get()) != INDEX_NONE &&
+				Cloth->GetLODIndex(Solver.Get()) != INDEX_NONE)  // Needs a valid LOD
 			{
-				CurrentSimulationData.Reset();
-				return;
+				AssetIndices.AddUnique((int32)Cloth->GetGroupId());
+				// TODO: Asset indices should be taken from the ClothComponentAdapter so multiple ids can be recovered from this group
 			}
 		}
 
-		if (!Cloths.Num())
+		// Cleanup all past keys (ensures that the cloth vertex factory won't run unnecessarily)
+		if (!AssetIndices.Num())
 		{
 			CurrentSimulationData.Reset();
 			return;
 		}
-
-		// Reset map when new cloths have appeared
-		if (CurrentSimulationData.Num() != Cloths.Num())
+		for (TMap<int32, FClothSimulData>::TIterator It = CurrentSimulationData.CreateIterator(); It; ++It)
 		{
-			CurrentSimulationData.Reset();
+			if (!AssetIndices.Contains(It->Key))
+			{
+				It.RemoveCurrent();
+			}
 		}
 
 		// Get the solver's local space
@@ -706,131 +743,163 @@ PRAGMA_ENABLE_DEPRECATION_WARNINGS
 		const FReal LocalSpaceScale = Solver->GetLocalSpaceScale();
 
 		// Retrieve the component's bones transforms
-		const TArray<FTransform>& ComponentSpaceTransforms = LeaderPoseComponent ? LeaderPoseComponent->GetComponentSpaceTransforms() : ClothComponent.GetComponentSpaceTransforms();
+		const TArray<FTransform>& ComponentSpaceTransforms = LeaderPoseComponent ? LeaderPoseComponent->GetComponentSpaceTransforms() : OwnerComponent.GetComponentSpaceTransforms();
 
 		// Set the simulation data for each of the cloths
-		for (const TUniquePtr<FClothingSimulationCloth>& Cloth : Cloths)
+		for (const int32 AssetIndex : AssetIndices)
 		{
-			const int32 AssetIndex = Cloth->GetGroupId();
+			bool bHasCloth = false;
+			int32 ReferenceBoneIndex = INDEX_NONE;
+			FRigidTransform3 ReferenceSpaceTransform;
 
-			if (!Cloth->GetMesh())
+			FClothSimulData* Data = CurrentSimulationData.Find(AssetIndex);
+			if (Data)
 			{
-				CurrentSimulationData.Remove(AssetIndex);  // Ensures that the cloth vertex factory won't run unnecessarily
-				continue;  // Invalid or empty cloth
+				Data->Reset();  // Clear the data to prepare them for the append operations to come
 			}
 
-			// If the LOD has changed while the simulation is suspended, the cloth still needs to be updated with the correct LOD data
-			// This should be handled by calling UpdateClothLODs when not ticking/simulating.
-			const int32 LODIndex = Cloth->GetMesh()->GetLODIndex();
-			ensure(LODIndex == Cloth->GetLODIndex(Solver.Get()));
-
-			if (Cloth->GetParticleRangeId(Solver.Get()) == INDEX_NONE || Cloth->GetLODIndex(Solver.Get()) == INDEX_NONE)
+			for (const TUniquePtr<FClothingSimulationCloth>& Cloth : Cloths)
 			{
-				CurrentSimulationData.Remove(AssetIndex);  // Ensures that the cloth vertex factory won't run unnecessarily
-				continue;  // No valid LOD, there's nothing to write out
-			}
+				if (Cloth->GetGroupId() != AssetIndex)
+				{
+					continue;
+				}
 
-			// Get the reference bone index for this cloth
-			const int32 ReferenceBoneIndex = LeaderPoseComponent ? ClothComponent.GetLeaderBoneMap()[Cloth->GetReferenceBoneIndex()] : Cloth->GetReferenceBoneIndex();
-			if (!ComponentSpaceTransforms.IsValidIndex(ReferenceBoneIndex))
-			{
-				UE_CLOG(!bHasInvalidReferenceBoneTransforms, LogChaosClothAsset, Warning, TEXT("Failed to write back clothing simulation data for component %s as bone transforms are invalid."), *ClothComponent.GetName());
-				bHasInvalidReferenceBoneTransforms = true;
-				CurrentSimulationData.Reset();
-				return;
-			}
+				// Get the reference bone index for this cloth
+				const int32 ClothReferenceBoneIndex = LeaderPoseComponent ? OwnerComponent.GetLeaderBoneMap()[Cloth->GetReferenceBoneIndex()] : Cloth->GetReferenceBoneIndex();
+				if (bHasCloth && ClothReferenceBoneIndex != ReferenceBoneIndex)
+				{
+					UE_CLOG(bHasNonMatchingReferenceBones, LogChaosClothAsset, Warning, TEXT("Multiple cloth simulation models are used with different reference bone indices within the same simulation group"),
+						OwnerComponent.GetSkinnedAsset() ? *OwnerComponent.GetSkinnedAsset()->GetName() : *OwnerComponent.GetName());
+					bHasNonMatchingReferenceBones = true;
+					bHasCloth = false;
+					break;
+				}
 
-			// Get the reference transform used in the current animation pose
-			FTransform ReferenceBoneTransform = ComponentSpaceTransforms[ReferenceBoneIndex];
-			ReferenceBoneTransform *= ClothSimulationContext->ComponentTransform;
-			ReferenceBoneTransform.SetScale3D(FVector(1.0f));  // Scale is already baked in the cloth mesh
+				if (!bHasCloth)
+				{
+					// Set the ReferenceBoneIndex for this group
+					ReferenceBoneIndex = ClothReferenceBoneIndex;
 
-			// Set the world space transform to be this cloth's reference bone
-			FClothSimulData& Data = CurrentSimulationData.FindOrAdd(AssetIndex);
-			Data.Transform = ReferenceBoneTransform;
-			Data.ComponentRelativeTransform = ReferenceBoneTransform.GetRelativeTransform(ClothSimulationContext->ComponentTransform);
+					// Get the reference bone transform for this cloth
+					if (!ComponentSpaceTransforms.IsValidIndex(ReferenceBoneIndex))
+					{
+						UE_CLOG(!bHasInvalidReferenceBoneTransforms, LogChaosClothAsset, Warning, TEXT("Failed to write back clothing simulation data for component %s as bone transforms are invalid."),
+							OwnerComponent.GetSkinnedAsset() ? *OwnerComponent.GetSkinnedAsset()->GetName() : *OwnerComponent.GetName());
+						bHasInvalidReferenceBoneTransforms = true;
+						CurrentSimulationData.Reset();
+						bHasCloth = false;
+						break;
+					}
 
-			// Retrieve the last reference space transform used for this cloth
-			// Note: This won't necessary match the current bone reference transform when the simulation is paused,
-			//       and still allows for the correct positioning of the sim data while the component is animated.
-			FRigidTransform3 ReferenceSpaceTransform = Cloth->GetReferenceSpaceTransform();
-			ReferenceSpaceTransform.AddToTranslation(-LocalSpaceLocation);
+					// Get the reference transform used in the current animation pose
+					FTransform ReferenceBoneTransform = ComponentSpaceTransforms[ReferenceBoneIndex];
+					ReferenceBoneTransform *= ClothSimulationContext->ComponentTransform;
+					ReferenceBoneTransform.SetScale3D(FVector(1.0f));  // Scale is already baked in the cloth mesh
 
-			// Copy positions and normals
-			Data.Positions = Cloth->GetParticlePositions(Solver.Get());
-			Data.Normals = Cloth->GetParticleNormals(Solver.Get());
+					// Set the FClothSimulData common properties for this group
+					if (!Data)
+					{
+						Data = &CurrentSimulationData.Add(AssetIndex);
+					}
 
-			check(Data.Positions.Num() == Data.Normals.Num());
+					// Set the world space transform to be this cloth's reference bone
+					Data->Transform = ReferenceBoneTransform;
+					Data->ComponentRelativeTransform = ReferenceBoneTransform.GetRelativeTransform(ClothSimulationContext->ComponentTransform);
+
+					// Set the current LOD these data apply to, so that the correct deformer mappings can be applied
+					// Note: If the LOD has changed while the simulation is suspended, the cloth still needs to be updated with the correct LOD data
+					//       This should be handled by calling UpdateClothLODs when not ticking/simulating.
+					const int32 LODIndex = Cloth->GetMesh()->GetLODIndex();
+					ensure(LODIndex == Cloth->GetLODIndex(Solver.Get()));
+					Data->LODIndex = Cloth->GetMesh()->GetOwnerLODIndex(LODIndex);  // The owner component LOD index can be different to the cloth mesh LOD index
+
+					// Retrieve the last reference space transform used for this cloth
+					// Note: This won't necessary match the current bone reference transform when the simulation is paused,
+					//       and still allows for the correct positioning of the sim data while the component is animated.
+					ReferenceSpaceTransform = Cloth->GetReferenceSpaceTransform();
+					ReferenceSpaceTransform.AddToTranslation(-LocalSpaceLocation);
+				}
 
-			// Transform into the cloth reference simulation space used at the time of simulation
-			bool bSimDataIsFinite = true;
-#if INTEL_ISPC
-			if (bTransformClothSimulData_ISPC_Enabled)
-			{
-				// ISPC is assuming float input here
-				check(sizeof(ispc::FVector3f) == Data.Positions.GetTypeSize());
-				check(sizeof(ispc::FVector3f) == Data.Normals.GetTypeSize());
+				// Copy positions and normals
+				Data->Positions.Append(Cloth->GetParticlePositions(Solver.Get()));
+				Data->Normals.Append(Cloth->GetParticleNormals(Solver.Get()));
 
-				if (bSimulationDataCheckFinite)
-				{
-					bSimDataIsFinite = ispc::TransformClothSimulDataAndCheckFinite(
-						(ispc::FVector3f*)Data.Positions.GetData(),
-						(ispc::FVector3f*)Data.Normals.GetData(),
-						(ispc::FTransform&)ReferenceSpaceTransform,
-						LocalSpaceScale,
-						Data.Positions.Num());
-				}
-				else
-				{
-					ispc::TransformClothSimulData(
-						(ispc::FVector3f*)Data.Positions.GetData(),
-						(ispc::FVector3f*)Data.Normals.GetData(),
-						(ispc::FTransform&)ReferenceSpaceTransform,
-						LocalSpaceScale,
-						Data.Positions.Num());
-				}
+				check(Data->Positions.Num() == Data->Normals.Num());
+				bHasCloth = true;
 			}
-			else
-#endif
+
+			// Transform into the cloth reference simulation space used at the time of simulation
+			bool bSimDataIsFinite = true;
+			if (bHasCloth)
 			{
-				if (bSimulationDataCheckFinite)
+	#if INTEL_ISPC
+				if (bTransformClothSimulData_ISPC_Enabled)
 				{
-					auto VectorIsFinite = [](const FSolverVec3& Vector)
-						{
-							return FMath::IsFinite(Vector.X) && FMath::IsFinite(Vector.Y) && FMath::IsFinite(Vector.Z);
-						};
-					for (int32 Index = 0; Index < Data.Positions.Num(); ++Index)
+					// ISPC is assuming float input here
+					check(sizeof(ispc::FVector3f) == Data->Positions.GetTypeSize());
+					check(sizeof(ispc::FVector3f) == Data->Normals.GetTypeSize());
+
+					if (bSimulationDataCheckFinite)
+					{
+						bSimDataIsFinite = ispc::TransformClothSimulDataAndCheckFinite(
+							(ispc::FVector3f*)Data->Positions.GetData(),
+							(ispc::FVector3f*)Data->Normals.GetData(),
+							(ispc::FTransform&)ReferenceSpaceTransform,
+							LocalSpaceScale,
+							Data->Positions.Num());
+					}
+					else
 					{
-						Data.Positions[Index] = FVec3f(ReferenceSpaceTransform.InverseTransformPosition(LocalSpaceScale * FVec3(Data.Positions[Index])));
-						Data.Normals[Index] = FVec3f(ReferenceSpaceTransform.InverseTransformVector(FVec3(-Data.Normals[Index])));
-						bSimDataIsFinite = bSimDataIsFinite && VectorIsFinite(Data.Positions[Index]) && VectorIsFinite(Data.Normals[Index]);
+						ispc::TransformClothSimulData(
+							(ispc::FVector3f*)Data->Positions.GetData(),
+							(ispc::FVector3f*)Data->Normals.GetData(),
+							(ispc::FTransform&)ReferenceSpaceTransform,
+							LocalSpaceScale,
+							Data->Positions.Num());
 					}
 				}
 				else
+	#endif
 				{
-					for (int32 Index = 0; Index < Data.Positions.Num(); ++Index)
+					if (bSimulationDataCheckFinite)
 					{
-						Data.Positions[Index] = FVec3f(ReferenceSpaceTransform.InverseTransformPosition(LocalSpaceScale * FVec3(Data.Positions[Index])));
-						Data.Normals[Index] = FVec3f(ReferenceSpaceTransform.InverseTransformVector(FVec3(-Data.Normals[Index])));
+						auto VectorIsFinite = [](const FSolverVec3& Vector)
+							{
+								return FMath::IsFinite(Vector.X) && FMath::IsFinite(Vector.Y) && FMath::IsFinite(Vector.Z);
+							};
+						for (int32 Index = 0; Index < Data->Positions.Num(); ++Index)
+						{
+							Data->Positions[Index] = FVec3f(ReferenceSpaceTransform.InverseTransformPosition(LocalSpaceScale * FVec3(Data->Positions[Index])));
+							Data->Normals[Index] = FVec3f(ReferenceSpaceTransform.InverseTransformVector(FVec3(-Data->Normals[Index])));
+							bSimDataIsFinite = bSimDataIsFinite && VectorIsFinite(Data->Positions[Index]) && VectorIsFinite(Data->Normals[Index]);
+						}
+					}
+					else
+					{
+						for (int32 Index = 0; Index < Data->Positions.Num(); ++Index)
+						{
+							Data->Positions[Index] = FVec3f(ReferenceSpaceTransform.InverseTransformPosition(LocalSpaceScale * FVec3(Data->Positions[Index])));
+							Data->Normals[Index] = FVec3f(ReferenceSpaceTransform.InverseTransformVector(FVec3(-Data->Normals[Index])));
+						}
 					}
 				}
 			}
-
-
 			if (bSimulationDataCheckFinite)
 			{
-				UE_CLOG(!bSimDataIsFinite, LogChaosClothAsset, Warning, TEXT("Failed to write back clothing simulation data for component %s as simulation contains non-finite values."), *ClothComponent.GetName());
+				UE_CLOG(!bSimDataIsFinite, LogChaosClothAsset, Warning, TEXT("Failed to write back clothing simulation data for component %s as simulation contains non-finite values."), *OwnerComponent.GetName());
 				ensure(bSimDataIsFinite);
 				if (!bSimDataIsFinite && bSimulationDataResetOnNonFinite)
 				{
 					bNeedResetClothFromNanCheck = true;
-					CurrentSimulationData.Reset();
-					return;
 				}
 			}
 
-			// Set the current LOD these data apply to, so that the correct deformer mappings can be applied
-			Data.LODIndex = Cloth->GetMesh()->GetOwnerLODIndex(LODIndex);  // The owner component LOD index can be different to the cloth mesh LOD index
+			// Ensures that the cloth vertex factory won't run unnecessarily
+			if (Data && (!bHasCloth || !bSimDataIsFinite))
+			{
+				CurrentSimulationData.Remove(AssetIndex);
+			}
 		}
 	}
 
@@ -876,20 +945,35 @@ PRAGMA_ENABLE_DEPRECATION_WARNINGS
 	void FClothSimulationProxy::InitializeConfigs()
 	{
 		TRACE_CPUPROFILER_EVENT_SCOPE(FClothSimulationProxy_InitializeConfigs);
-		// Replace physics thread's configs with the game thread's configs
-		for (int32 ConfigIndex = 0; ConfigIndex < Configs.Num(); ++ConfigIndex)
+		
+		if (Configs.Num())
 		{
-			const int32 ModelIndex = ConfigIndex - 1;
-			Configs[ConfigIndex]->Initialize(ModelIndex < 0 ?
-				ClothComponent.GetSolverPropertyCollections() :  // Config 0 is the solver config
-				ClothComponent.GetPropertyCollections(ModelIndex));
+			// Replace physics thread's configs with the game thread's configs
+			int32 ConfigIndex = 0;
+			Configs[ConfigIndex++]->Initialize(ClothComponentAdapter.GetSolverPropertyCollections());
+
+			for (const UChaosClothAssetBase* const Asset : ClothComponentAdapter.GetAssets())
+			{
+				if (Asset)
+				{
+					for (int32 ModelIndex = 0; ModelIndex < Asset->GetNumClothSimulationModels(); ++ModelIndex)
+					{
+						const int32 SimulationGroupId = ClothComponentAdapter.GetSimulationGroupId(Asset, ModelIndex);
+						if (SimulationGroupId != INDEX_NONE &&  // INDEX_NONE means this model hasn't been added to the simulation
+							ensure(Configs.IsValidIndex(ConfigIndex)))
+						{
+							Configs[ConfigIndex++]->Initialize(ClothComponentAdapter.GetPropertyCollections(Asset, ModelIndex));
+						}
+					}
+				}
+			}
 		}
 		Solver->SetSolverLOD(ClothSimulationContext->LodIndex);
 	}
 
 	void FClothSimulationProxy::FillSimulationContext(float DeltaTime, bool bIsInitialization)
 	{
-		ClothSimulationContext->Fill(ClothComponent, DeltaTime, MaxDeltaTime, bIsInitialization, CacheData.Get());
+		ClothSimulationContext->Fill(ClothComponentAdapter, DeltaTime, MaxDeltaTime, bIsInitialization, CacheData.Get());
 		CacheData.Reset();
 	}
 
diff --git a/Engine/Plugins/ChaosClothAsset/Source/ChaosClothAssetEngine/Public/ChaosClothAsset/ClothAssetSKMClothingAsset.h b/Engine/Plugins/ChaosClothAsset/Source/ChaosClothAssetEngine/Public/ChaosClothAsset/ClothAssetSKMClothingAsset.h
--- a/Engine/Plugins/ChaosClothAsset/Source/ChaosClothAssetEngine/Public/ChaosClothAsset/ClothAssetSKMClothingAsset.h
+++ b/Engine/Plugins/ChaosClothAsset/Source/ChaosClothAssetEngine/Public/ChaosClothAsset/ClothAssetSKMClothingAsset.h
@@ -20,8 +20,29 @@ class UChaosClothAssetSKMClothingAsset final : public UClothingAssetBase
 public:
 	UE_API UChaosClothAssetSKMClothingAsset(const FObjectInitializer& ObjectInitializer);
 
+#if WITH_EDITOR
+	/**
+	 * Assign a Cloth or an Outfit Asset to this clothing data.
+	 * @Note this is only possible WITH_EDITOR since the Skeletal Mesh might need to be rebuild when the asset is changed.
+	 */
 	UE_API void SetAsset(const UChaosClothAssetBase* InAsset);
 
+	/** Call this function whenever the Cloth or Outfit Asset has changed (e.g. once rebuilt in Dataflow). */
+	UE_API void OnAssetChanged();
+#endif
+	
+	/** Return the underlaying Cloth or Outfit Asset used in this clothing data. */
+	const UChaosClothAssetBase* GetAsset() const
+	{
+		return Asset;
+	}
+
+	/** Return the index of the simulation model used from the Cloth or Outfit Asset. */
+	int32 GetClothSimulationModelIndex() const
+	{
+		return ClothSimulationModelIndex;
+	}
+
 private:
 #if WITH_EDITOR
 	//~ Begin UObject Interface
@@ -30,50 +51,43 @@ private:
 
 	//~ Begin UClothingAssetBase Interface
 	virtual bool BindToSkeletalMesh(USkeletalMesh* SkeletalMesh, const int32 MeshLodIndex, const int32 SectionIndex, const int32 AssetLodIndex) override;
-
-	virtual void UnbindFromSkeletalMesh(USkeletalMesh* SkeletalMesh) override
-	{
-		// TODO
-	}
-
-	virtual void UnbindFromSkeletalMesh(USkeletalMesh* SkeletalMesh, const int32 MeshLodIndex) override
-	{
-		// TODO
-	}
-
-	virtual void UpdateAllLODBiasMappings(USkeletalMesh* SkeletalMesh) override
-	{
-		// TODO
-	}
+	virtual void UnbindFromSkeletalMesh(USkeletalMesh* SkeletalMesh) override;
+	virtual void UnbindFromSkeletalMesh(USkeletalMesh* SkeletalMesh, const int32 MeshLodIndex) override;
+	virtual void UpdateAllLODBiasMappings(USkeletalMesh* SkeletalMesh) override;
 #endif  // #if WITH_EDITOR
-
-	virtual void RefreshBoneMapping(USkeletalMesh* SkeletalMesh) override
-	{
-	}
-
-	virtual void PostUpdateAllAssets() override
+	virtual void RefreshBoneMapping(USkeletalMesh* SkeletalMesh) override {}
+	virtual bool IsValid() const override
 	{
-		// TODO
+		return GetClothSimulationModelIndex() != INDEX_NONE;
 	}
+	virtual void PostUpdateAllAssets() override {}
 	//~ End UClothingAssetBase Interface
 
-	UFUNCTION(CallInEditor)
-	TArray<FName> GetClothSimulationModelNames() const;
+#if WITH_EDITOR
+	void OnModelChanged();
 
-	int32 GetClothSimulationModelIndex() const
-	{
-		return GetClothSimulationModelNames().Find(ClothSimulationModelName);
-	}
+	FName GetClothSimulationModelName() const;
+	FGuid GetClothSimulationModelGuid() const;
+#endif  // #if WITH_EDITOR
 
-	void OnAssetChanged();
+	UFUNCTION(CallInEditor)
+	TArray<FString> GetClothSimulationModelIds() const;
 
 	/** Cloth or Outfit Asset to use with this Skeletal Mesh. */
-	UPROPERTY(EditAnywhere, Setter = "SetAsset", Category = "Chaos Cloth Asset")
+	UPROPERTY(EditAnywhere, Category = "Chaos Cloth Asset")
 	TObjectPtr<const UChaosClothAssetBase> Asset;
 
-	/** The Simulation Model to use for this Clothing Data. Useful for Outfit Assets that may contain multiple models. */
-	UPROPERTY(EditAnywhere, Category = "Chaos Cloth Asset", Meta = (DisplayName = "Cloth Simulation Model", GetOptions = "GetClothSimulationModelNames"))
-	FName ClothSimulationModelName;
+#if WITH_EDITORONLY_DATA
+	/**
+	 * The Simulation Model to use for this Clothing Data. Useful for Outfit Assets that may contain multiple models.
+	 * The name is followed by the model GUID, since it is technically possible to have two pieces of clothing with the same name in an Outfit Asset.
+	 */
+	UPROPERTY(EditAnywhere, Category = "Chaos Cloth Asset", Meta = (DisplayName = "Cloth Simulation Model", GetOptions = "GetClothSimulationModelIds"))
+	FString ClothSimulationModelId;
+#endif  // #if WITH_EDITORONLY_DATA
+
+	UPROPERTY()
+	int32 ClothSimulationModelIndex = INDEX_NONE;
 };
 
 #undef UE_API
diff --git a/Engine/Plugins/ChaosClothAsset/Source/ChaosClothAssetEngine/Public/ChaosClothAsset/ClothComponent.h b/Engine/Plugins/ChaosClothAsset/Source/ChaosClothAssetEngine/Public/ChaosClothAsset/ClothComponent.h
--- a/Engine/Plugins/ChaosClothAsset/Source/ChaosClothAssetEngine/Public/ChaosClothAsset/ClothComponent.h
+++ b/Engine/Plugins/ChaosClothAsset/Source/ChaosClothAssetEngine/Public/ChaosClothAsset/ClothComponent.h
@@ -2,6 +2,8 @@
 
 #pragma once
 
+#include "ChaosClothAsset/ClothComponentAdapter.h"
+#include "Containers/Array.h"
 #include "Components/SkinnedMeshComponent.h"
 #include "Dataflow/Interfaces/DataflowPhysicsSolver.h"
 #include "ClothComponent.generated.h"
@@ -41,25 +43,32 @@ USTRUCT()
 struct FChaosClothSimulationProperties final
 {
 	GENERATED_BODY()
-private:
-	friend UChaosClothComponent;
+
 	UPROPERTY(Transient)
 	TObjectPtr<UChaosClothAssetInteractor> ClothOutfitInteractor;
 	TArray<TSharedPtr<const FManagedArrayCollection>> PropertyCollections;
 	TArray<TSharedPtr<::Chaos::Softs::FCollectionPropertyFacade>> CollectionPropertyFacades;
 
+	FChaosClothSimulationProperties() = default;
+	explicit FChaosClothSimulationProperties(const TArray<TSharedRef<const FManagedArrayCollection>>& AssetPropertyCollections)
+	{
+		Initialize(AssetPropertyCollections);
+	}
 	void Initialize(const TArray<TSharedRef<const FManagedArrayCollection>>& AssetPropertyCollections);
 };
 
 /**
  * Cloth simulation component.
  */
-UCLASS(MinimalAPI, 
-	ClassGroup = Physics, 
+UCLASS(MinimalAPI,
+	ClassGroup = Physics,
 	Meta = (BlueprintSpawnableComponent, ToolTip = "Chaos Cloth Component"),
 	DisplayName = "Chaos Cloth Component",
 	HideCategories = (Object, "Mesh|SkeletalAsset", Constraints, Advanced, Cooking, Collision, Navigation))
-class UChaosClothComponent : public USkinnedMeshComponent, public IDataflowPhysicsSolverInterface
+	class UChaosClothComponent
+	: public USkinnedMeshComponent
+	, public IDataflowPhysicsSolverInterface
+	, public UE::Chaos::ClothAsset::IClothComponentAdapter
 {
 	GENERATED_BODY()
 public:
@@ -107,11 +116,15 @@ public:
 	/** Return whether reseting the pose is currently requested. Use GetClothTeleportMode to get teleport mode including any auto-teleport based on teleport thresholds*/
 	bool NeedsReset() const { return bReset || bResetOnce; }
 
-	bool NeedsResetRestLengths() const { return bResetRestLengthsFromMorphTarget; }
-	const FString& GetRestLengthsMorphTargetName() const { return ResetRestLengthsMorphTargetName; }
+	/** Part of the IClothComponentAdapter implementation. */
+	virtual bool NeedsResetRestLengths() const override { return bResetRestLengthsFromMorphTarget; }
+	virtual const FString& GetRestLengthsMorphTargetName() const override { return ResetRestLengthsMorphTargetName; }
 
-	/** Get currently calculated teleport mode. */
-	EClothingTeleportMode GetClothTeleportMode() const { return ClothTeleportMode; }
+	/**
+	 * Get currently calculated teleport mode.
+	 * Part of the IClothComponentAdapter implementation.
+	 */
+	virtual EClothingTeleportMode GetClothTeleportMode() const override { return ClothTeleportMode; }
 
 	/** Stop the simulation, and keep the cloth in its last pose. */
 	UFUNCTION(CallInEditor, BlueprintCallable, Category = "ClothComponent", Meta = (UnsafeDuringActorConstruction, Keywords = "Chaos Cloth Simulation Suspend"))
@@ -121,17 +134,23 @@ public:
 	UFUNCTION(CallInEditor, BlueprintCallable, Category = "ClothComponent", Meta = (UnsafeDuringActorConstruction, Keywords = "Chaos Cloth Simulation Resume"))
 	void ResumeSimulation() { bSuspendSimulation = false; }
 
-	/** Return whether or not the simulation is currently suspended. */
+	/**
+	 * Return whether or not the simulation is currently suspended.
+	 * Part of the IClothComponentAdapter implementation.
+	 */
 	UFUNCTION(BlueprintCallable, Category = "ClothComponent", Meta = (Keywords = "Chaos Cloth Simulation Suspend"))
-	UE_API bool IsSimulationSuspended() const;
+	UE_API virtual bool IsSimulationSuspended() const override;
 
 	/** Set whether or not to enable simulation. */
 	UFUNCTION(BlueprintCallable, Category = "ClothComponent", Meta = (UnsafeDuringActorConstruction, Keywords = "Chaos Cloth Simulation Enable"))
 	void SetEnableSimulation(bool bEnable) { bEnableSimulation = bEnable; }
 
-	/** Return whether or not the simulation is currently enabled. */
+	/**
+	 * Return whether or not the simulation is currently enabled.
+	 * Part of the IClothComponentAdapter implementation.
+	 */
 	UFUNCTION(BlueprintCallable, Category = "ClothComponent", Meta = (Keywords = "Chaos Cloth Simulation Enable"))
-	UE_API bool IsSimulationEnabled() const;
+	UE_API virtual bool IsSimulationEnabled() const override;
 
 	/** Reset all cloth simulation config properties to the values stored in the original cloth asset.*/
 	UFUNCTION(BlueprintCallable, Category = "ClothComponent", Meta = (Keywords = "Chaos Cloth Config Property"))
@@ -174,8 +193,11 @@ public:
 	UFUNCTION(BlueprintCallable, Category = "ClothComponent", Meta = (Keywords = "Chaos Cloth Collision Source"))
 	UE_API void ResetCollisionSources();
 
-	/** Return all collision sources currently assigned to this component. */
-	UE::Chaos::ClothAsset::FCollisionSources& GetCollisionSources() const { return *CollisionSources; }
+	/**
+	 * Return all collision sources currently assigned to this component.
+	 * Part of the IClothComponentAdapter implementation.
+	 */
+	virtual UE::Chaos::ClothAsset::FCollisionSources& GetCollisionSources() const override { return *CollisionSources; }
 
 	/** Set whether or not to collide with the environment. */
 	UFUNCTION(BlueprintCallable, Category = "ClothComponent", Meta = (Keywords = "Chaos Cloth Collide Environment"))
@@ -201,8 +223,11 @@ public:
 	{
 		return ClothSimulationProperties[ModelIndex].PropertyCollections;
 	}
-	/** Return the solver property collections. */
-	const TArray<TSharedPtr<const FManagedArrayCollection>>& GetSolverPropertyCollections() const
+	/**
+	 * Return the solver property collections.
+	 * Part of the IClothComponentAdapter implementation.
+	 */
+	virtual const TArray<TSharedPtr<const FManagedArrayCollection>>& GetSolverPropertyCollections() const override
 	{
 		return ClothSimulationProperties[0].PropertyCollections;
 	}
@@ -214,10 +239,13 @@ public:
 
 	const UE::Chaos::ClothAsset::FClothSimulationProxy* GetClothSimulationProxy() const { return ClothSimulationProxy.Get(); }
 
-	/** This scale is applied to all cloth geometry (e.g., cloth meshes and collisions) in order to simulate in a different scale space than world.This scale is not applied to distance-based simulation parameters such as MaxDistance.
-	* This property is currently only read by the cloth solver when creating cloth actors, but may become animatable in the future.
-	*/
-	float GetClothGeometryScale() const { return ClothGeometryScale; }
+	/**
+	 * This scale is applied to all cloth geometry (e.g., cloth meshes and collisions) in order to simulate in a different scale space than world.
+	 * This scale is not applied to distance-based simulation parameters such as MaxDistance.
+	 * This property is currently only read by the cloth solver when creating cloth actors, but may become animatable in the future.
+	 * Part of the IClothComponentAdapter implementation.
+	 */
+	virtual float GetClothGeometryScale() const override { return ClothGeometryScale; }
 	void SetClothGeometryScale(float Scale) { ClothGeometryScale = Scale; }
 
 	/**
@@ -317,7 +345,7 @@ protected:
 	UE_API virtual void UnregisterOnBoneTransformsFinalizedDelegate(const FDelegateHandle& DelegateHandle) override;
 	//~ End USkinnedMeshComponent Interface
 
-	// Begin IDataflowPhysicsSolverInterface overrides
+	//~ Begin IDataflowPhysicsSolverInterface Interface
 	virtual FString GetSimulationName() const override {return GetName();};
 	virtual FDataflowSimulationAsset& GetSimulationAsset() override {return SimulationAsset;};
 	virtual const FDataflowSimulationAsset& GetSimulationAsset() const override {return SimulationAsset;};
@@ -329,7 +357,25 @@ protected:
 	UE_API virtual void ReadFromSimulation(const float DeltaTime, const bool bAsyncTask) override;
 	UE_API virtual void PreProcessSimulation(const float DeltaTime) override;
 	UE_API virtual void PostProcessSimulation(const float DeltaTime) override;
-	// End IDataflowPhysicsSolverInterface overrides
+	//~ End IDataflowPhysicsSolverInterface Interface
+
+	//~ Begin IClothComponentAdapter Interface
+	virtual const USkinnedMeshComponent& GetOwnerComponent() const override
+	{
+		return *this;
+	}
+	UE_API virtual const FReferenceSkeleton* GetReferenceSkeleton() const override;
+	virtual TArray<const UChaosClothAssetBase*> GetAssets() const override
+	{
+		return GetAsset() ? TArray<const UChaosClothAssetBase*>{ GetAsset() } : TArray<const UChaosClothAssetBase*>{};
+	}
+	UE_API virtual int32 GetSimulationGroupId(const UChaosClothAssetBase* InAsset, int32 ModelIndex) const override;
+	virtual const TArray<TSharedPtr<const FManagedArrayCollection>>& GetPropertyCollections(const UChaosClothAssetBase* InAsset, int32 ModelIndex) const override
+	{
+		checkf(GetSimulationGroupId(InAsset, ModelIndex) != INDEX_NONE, TEXT("Invalid arguments"));
+		return GetPropertyCollections(ModelIndex);  // There's only this one asset for now, therefore this is the same as calling GetPropertyCollections(ModelIndex)
+	}
+	//~ End IClothComponentAdapter Interface
 
 	/** Override this function for setting up custom simulation proxies when the component is registered. */
 	UE_API virtual TSharedPtr<UE::Chaos::ClothAsset::FClothSimulationProxy> CreateClothSimulationProxy();
diff --git a/Engine/Plugins/ChaosClothAsset/Source/ChaosClothAssetEngine/Public/ChaosClothAsset/ClothComponentAdapter.h b/Engine/Plugins/ChaosClothAsset/Source/ChaosClothAssetEngine/Public/ChaosClothAsset/ClothComponentAdapter.h
new file mode 100644
--- /dev/null
+++ b/Engine/Plugins/ChaosClothAsset/Source/ChaosClothAssetEngine/Public/ChaosClothAsset/ClothComponentAdapter.h
@@ -0,0 +1,38 @@
+// Copyright Epic Games, Inc. All Rights Reserved.
+
+#pragma once
+
+#include "Containers/Array.h"
+#include "Templates/SharedPointer.h"
+
+class UChaosClothAssetBase;
+class USkinnedMeshComponent;
+struct FManagedArrayCollection;
+struct FReferenceSkeleton;
+enum class EClothingTeleportMode : uint8;
+
+namespace UE::Chaos::ClothAsset
+{
+	class FCollisionSources;
+
+	class IClothComponentAdapter
+	{
+	public:
+		virtual const USkinnedMeshComponent& GetOwnerComponent() const = 0;
+		virtual FCollisionSources& GetCollisionSources() const = 0;
+		virtual TArray<const UChaosClothAssetBase*> GetAssets() const = 0;
+		/** Returns the index of the model assigned to the render section, or INDEX_NONE when the model shouldn't be simulated. */
+		virtual int32 GetSimulationGroupId(const UChaosClothAssetBase* Asset, int32 ModelId) const = 0;
+		virtual const FReferenceSkeleton* GetReferenceSkeleton() const = 0;
+		virtual EClothingTeleportMode GetClothTeleportMode() const = 0;
+		virtual bool IsSimulationEnabled() const = 0;
+		virtual bool IsSimulationSuspended() const = 0;
+		virtual bool NeedsResetRestLengths() const = 0;
+		virtual const FString& GetRestLengthsMorphTargetName() const = 0;
+		virtual float GetClothGeometryScale() const = 0;
+		virtual const TArray<TSharedPtr<const FManagedArrayCollection>>& GetPropertyCollections(const UChaosClothAssetBase* Asset, int32 ModelIndex) const = 0;
+		virtual const TArray<TSharedPtr<const FManagedArrayCollection>>& GetSolverPropertyCollections() const = 0;
+	protected:
+		virtual ~IClothComponentAdapter() = default;
+	};
+}
diff --git a/Engine/Plugins/ChaosClothAsset/Source/ChaosClothAssetEngine/Public/ChaosClothAsset/ClothSimulationProxy.h b/Engine/Plugins/ChaosClothAsset/Source/ChaosClothAssetEngine/Public/ChaosClothAsset/ClothSimulationProxy.h
--- a/Engine/Plugins/ChaosClothAsset/Source/ChaosClothAssetEngine/Public/ChaosClothAsset/ClothSimulationProxy.h
+++ b/Engine/Plugins/ChaosClothAsset/Source/ChaosClothAssetEngine/Public/ChaosClothAsset/ClothSimulationProxy.h
@@ -11,10 +11,9 @@
 
 #define UE_API CHAOSCLOTHASSETENGINE_API
 
-#if UE_ENABLE_INCLUDE_ORDER_DEPRECATED_IN_5_5
-namespace Dataflow = UE::Dataflow;
-#else
-namespace UE_DEPRECATED(5.5, "Use UE::Dataflow instead.") Dataflow {}
+#if UE_ENABLE_INCLUDE_ORDER_DEPRECATED_IN_5_7
+class UChaosClothAsset;
+class UChaosClothComponent;
 #endif
 
 namespace Chaos
@@ -28,8 +27,7 @@ namespace Chaos
 	class FClothVisualizationNoGC;
 }
 
-class UChaosClothComponent;
-class UChaosClothAsset;
+class UChaosClothAssetBase;
 struct FChaosClothSimulationModel;
 struct FReferenceSkeleton;
 struct FClothingSimulationCacheData;
@@ -37,19 +35,20 @@ struct FClothingSimulationCacheData;
 namespace UE::Chaos::ClothAsset
 {
 	class FCollisionSourcesProxy;
+	class IClothComponentAdapter;
 	struct FClothSimulationContext;
 
 	/**
 	 * Cloth simulation proxy.
-	 * Class used to share data between the cloth simulation and the cloth component.
+	 * Class used to share data between the cloth simulation and the component.
 	 */
 	class FClothSimulationProxy : public FDataflowPhysicsSolverProxy
 	{
 	public:
-		UE_API explicit FClothSimulationProxy(const UChaosClothComponent& InClothComponent);
+		UE_API explicit FClothSimulationProxy(const IClothComponentAdapter& ClothComponentAdapter);
 		UE_API ~FClothSimulationProxy();
 
-		FClothSimulationProxy() = delete;  // This object cannot be created without a valid reference to a parent UChaosClothComponent 
+		FClothSimulationProxy() = delete;  // This object cannot be created without a valid reference to an owner component 
 		FClothSimulationProxy(const FClothSimulationProxy&) = delete;  // Disable the copy as there must be a single unique proxy per component
 		FClothSimulationProxy(FClothSimulationProxy&&) = delete;  // Disable the move to force it to be associated with a valid component reference
 		FClothSimulationProxy& operator=(const FClothSimulationProxy&) = delete;
@@ -58,7 +57,7 @@ namespace UE::Chaos::ClothAsset
 		UE_API virtual void PostConstructor();
 
 		/** Start the simulation if it isn't paused or suspended and return true, or simply update the existing simulation data and return false otherwise. */
-		UE_API bool Tick_GameThread(float DeltaTime);
+		UE_API bool Tick_GameThread(float DeltaTime, bool bUseParallelTask = true);
 
 		/** Wait for the parallel task to complete if one was running, and update the simulation data. */
 		UE_API void CompleteParallelSimulation_GameThread();
@@ -153,7 +152,7 @@ UE_API PRAGMA_ENABLE_DEPRECATION_WARNINGS
 		TMap<int32, FClothSimulData> CurrentSimulationData;
 
 		// Owner component
-		const UChaosClothComponent& ClothComponent;
+		const IClothComponentAdapter& ClothComponentAdapter;
 
 		// Simulation context used to store the required component data for the duration of the simulation
 		TUniquePtr<FClothSimulationContext> ClothSimulationContext;
@@ -198,20 +197,15 @@ PRAGMA_ENABLE_DEPRECATION_WARNINGS
 		std::atomic<int32> LastLinearSolveIterations; // For single cloth only.
 		std::atomic<float> LastLinearSolveError; // For single cloth only.
 
-		mutable bool bHasInvalidReferenceBoneTransforms = false;
-		bool bNeedResetClothFromNanCheck = false;
-
 		// Cached value of the MaxPhysicsDeltaTime setting for the life of this proxy
 		const float MaxDeltaTime;
 
-		// Whether the proxy has been running the simulation, for internal use only
-		bool bIsSimulating;
-
-		// Whether the proxy has finished initialization, for internal use only
-		bool bIsInitialized = false;
-
-		// Whether the proxy was initialized when Tick_GameThread/PreProcess_GameThread was called, for internal use only
-		bool bIsPreProcessed = false;
+		mutable bool bHasInvalidReferenceBoneTransforms : 1 = false;
+		mutable bool bHasNonMatchingReferenceBones : 1 = false;
+		bool bNeedResetClothFromNanCheck : 1 = false;
+		bool bIsSimulating : 1 = false;  // Whether the proxy has been running the simulation, for internal use only
+		bool bIsInitialized : 1 = false;  // Whether the proxy has finished initialization, for internal use only
+		bool bIsPreProcessed : 1 = false;  // Whether the proxy was initialized when Tick_GameThread/PreProcess_GameThread was called, for internal use only
 	};
 
 #if !defined(CHAOS_TRANSFORM_CLOTH_SIMUL_DATA_ISPC_ENABLED_DEFAULT)
diff --git a/Engine/Plugins/ChaosClothAssetEditor/Source/ChaosClothAssetTools/Private/ChaosClothAsset/ChaosClothAssetToolsModule.cpp b/Engine/Plugins/ChaosClothAssetEditor/Source/ChaosClothAssetTools/Private/ChaosClothAsset/ChaosClothAssetToolsModule.cpp
--- a/Engine/Plugins/ChaosClothAssetEditor/Source/ChaosClothAssetTools/Private/ChaosClothAsset/ChaosClothAssetToolsModule.cpp
+++ b/Engine/Plugins/ChaosClothAssetEditor/Source/ChaosClothAssetTools/Private/ChaosClothAsset/ChaosClothAssetToolsModule.cpp
@@ -3,13 +3,19 @@
 #include "Features/IModularFeatures.h"
 #include "Modules/ModuleManager.h"
 #include "ChaosClothAsset/ClothAssetEditorStyle.h"
+#include "ChaosClothAsset/ClothAssetSKMClothingAssetFactory.h"
 #include "ChaosClothAsset/ClothComponentEditorStyle.h"
 #include "ChaosClothAsset/ClothingAssetToClothAssetExporter.h"
 #include "ChaosClothAsset/SkeletalMeshConverter.h"
 
 namespace UE::Chaos::ClothAsset
 {
-	class FChaosClothAssetToolsModule : public IModuleInterface, public IClothingAssetExporterClassProvider, public IClothAssetSkeletalMeshConverterClassProvider
+	class FChaosClothAssetToolsModule
+		: public IModuleInterface
+		, public IClothingAssetExporterClassProvider
+		, public IClothAssetSkeletalMeshConverterClassProvider
+		, public IClothingAssetFactoryProvider
+
 	{
 	public:
 		// IModuleInterface implementation
@@ -22,6 +28,7 @@ namespace UE::Chaos::ClothAsset
 			// Register modular features
 			IModularFeatures::Get().RegisterModularFeature(IClothingAssetExporterClassProvider::FeatureName, static_cast<IClothingAssetExporterClassProvider*>(this));
 			IModularFeatures::Get().RegisterModularFeature(IClothAssetSkeletalMeshConverterClassProvider::FeatureName, static_cast<IClothAssetSkeletalMeshConverterClassProvider*>(this));
+			IModularFeatures::Get().RegisterModularFeature(IClothingAssetFactoryProvider::FeatureName, static_cast<IClothingAssetFactoryProvider*>(this));
 		}
 
 		virtual void ShutdownModule() override
@@ -31,6 +38,7 @@ namespace UE::Chaos::ClothAsset
 				// Unregister modular features
 				IModularFeatures::Get().UnregisterModularFeature(IClothingAssetExporterClassProvider::FeatureName, static_cast<IClothingAssetExporterClassProvider*>(this));
 				IModularFeatures::Get().UnregisterModularFeature(IClothAssetSkeletalMeshConverterClassProvider::FeatureName, static_cast<IClothAssetSkeletalMeshConverterClassProvider*>(this));
+				IModularFeatures::Get().UnregisterModularFeature(IClothingAssetFactoryProvider::FeatureName, static_cast<IClothingAssetFactoryProvider*>(this));
 			}
 		}
 
@@ -47,6 +55,13 @@ namespace UE::Chaos::ClothAsset
 			return UClothAssetEditorSkeletalMeshConverter::StaticClass();
 		}
 		//~ End IClothAssetSkeletalMeshConverterClassProvider implementation
+
+		//~ Begin IClothingAssetFactoryProvider
+		virtual UClothingAssetFactoryBase* GetFactory() override
+		{
+			return UChaosClothAssetSKMClothingAssetFactory::StaticClass()->GetDefaultObject<UClothingAssetFactoryBase>();
+		}
+		//~ End IClothingSimulationFactoryClassProvider
 	};
 }
 
diff --git a/Engine/Plugins/MovieScene/MovieRenderPipeline/Source/MovieRenderPipelineCore/Private/MoviePipelineUtils.cpp b/Engine/Plugins/MovieScene/MovieRenderPipeline/Source/MovieRenderPipelineCore/Private/MoviePipelineUtils.cpp
--- a/Engine/Plugins/MovieScene/MovieRenderPipeline/Source/MovieRenderPipelineCore/Private/MoviePipelineUtils.cpp
+++ b/Engine/Plugins/MovieScene/MovieRenderPipeline/Source/MovieRenderPipelineCore/Private/MoviePipelineUtils.cpp
@@ -38,6 +38,7 @@
 #include "RHI.h"
 #include "CineCameraComponent.h"
 #include "EngineUtils.h"
+#include "ClothingSimulationInstance.h"
 #include "ClothingSimulationInteractor.h"
 #include "ClothingSimulationInterface.h"
 #include "ChaosClothAsset/ClothComponent.h"
@@ -1275,26 +1276,26 @@ namespace UE
 					TArray<USkeletalMeshComponent*> SkeletalMeshComponents;
 					FoundActor->GetComponents(SkeletalMeshComponents);
 
-					for (USkeletalMeshComponent* Component : SkeletalMeshComponents)
+					for (const USkeletalMeshComponent* const Component : SkeletalMeshComponents)
 					{
-						UClothingSimulationInteractor* ClothInteractor = Component->GetClothingSimulationInteractor();
-						if (ClothInteractor)
+						for (const FClothingSimulationInstance& ClothingSimulationInstance : Component->GetClothingSimulationInstances())
 						{
-							const int32 LODIndex = 0;  // There is only a NumSubSteps for LOD 0 in the Skeletal Mesh clothing system
-							TWeakObjectPtr<UObject> WeakPtr = TWeakObjectPtr<UClothingSimulationInteractor>(ClothInteractor);
-
-							TArray<::MoviePipeline::FClothSimSettingsCache>* ExistingCacheEntry = InClothSimCache.Find(WeakPtr);
-							if (!ExistingCacheEntry)
+							const IClothingSimulation* const ClothingSimulation = ClothingSimulationInstance.GetClothingSimulation();
+							UClothingSimulationInteractor* const ClothingSimulationInteractor = ClothingSimulationInstance.GetClothingSimulationInteractor();
+							if (ClothingSimulation && ClothingSimulationInteractor)
 							{
-								InClothSimCache.Add(WeakPtr);
-								ExistingCacheEntry = InClothSimCache.Find(WeakPtr);
-								ExistingCacheEntry->SetNumZeroed(1);  // Only store LOD 0
-								const int32 NumSubsteps = Component->GetClothingSimulation() ? FMath::Max(Component->GetClothingSimulation()->GetNumSubsteps(), 1)
-									: 1; // If there's no clothing simulation component just fall back to assuming they only had 1.
-								(*ExistingCacheEntry)[LODIndex].NumSubSteps = NumSubsteps;
-							}
+								constexpr int32 LODIndex = 0;  // There is only a NumSubSteps for LOD 0 in the Skeletal Mesh clothing system
+								const TWeakObjectPtr<UObject> WeakPtr(ClothingSimulationInteractor);
 
-							ClothInteractor->SetNumSubsteps((*ExistingCacheEntry)[LODIndex].NumSubSteps * InSubdivisionCount);
+								TArray<::MoviePipeline::FClothSimSettingsCache>* ExistingCacheEntry = InClothSimCache.Find(WeakPtr);
+								if (!ExistingCacheEntry)
+								{
+									ExistingCacheEntry = &InClothSimCache.Add(WeakPtr);
+									ExistingCacheEntry->SetNumZeroed(1);  // Only store LOD 0
+									(*ExistingCacheEntry)[LODIndex] = { ClothingSimulation->GetNumSubsteps(), 0.f };
+								}
+								ClothingSimulationInteractor->SetNumSubsteps((*ExistingCacheEntry)[LODIndex].NumSubSteps * InSubdivisionCount);
+							}
 						}
 					}
 
diff --git a/Engine/Plugins/VirtualProduction/LevelSnapshots/Source/LevelSnapshots/Private/SpecialSupport/EngineTypesRestorationFence.cpp b/Engine/Plugins/VirtualProduction/LevelSnapshots/Source/LevelSnapshots/Private/SpecialSupport/EngineTypesRestorationFence.cpp
--- a/Engine/Plugins/VirtualProduction/LevelSnapshots/Source/LevelSnapshots/Private/SpecialSupport/EngineTypesRestorationFence.cpp
+++ b/Engine/Plugins/VirtualProduction/LevelSnapshots/Source/LevelSnapshots/Private/SpecialSupport/EngineTypesRestorationFence.cpp
@@ -119,12 +119,14 @@ namespace UE::LevelSnapshots::Private::EngineTypesRestorationFence
 
 	static void DisableClothingSimulationFactory(ILevelSnapshotsModule& Module)
 	{
+PRAGMA_DISABLE_DEPRECATION_WARNINGS
 		// This property is directly set by RegisterComponent and similar functions... we have no chance of restoring this property properly
 		const FProperty* ClothingSimulationFactory = USkeletalMeshComponent::StaticClass()->FindPropertyByName(GET_MEMBER_NAME_CHECKED(USkeletalMeshComponent, ClothingSimulationFactory));
 		if (ensure(ClothingSimulationFactory))
 		{
 			Module.AddExplicitlyUnsupportedProperties({ ClothingSimulationFactory });
 		}
+PRAGMA_ENABLE_DEPRECATION_WARNINGS
 	}
 	
 	void RegisterSpecialEngineTypeSupport(FLevelSnapshotsModule& Module)
diff --git a/Engine/Source/Editor/Persona/Private/AnimationEditorViewportClient.cpp b/Engine/Source/Editor/Persona/Private/AnimationEditorViewportClient.cpp
--- a/Engine/Source/Editor/Persona/Private/AnimationEditorViewportClient.cpp
+++ b/Engine/Source/Editor/Persona/Private/AnimationEditorViewportClient.cpp
@@ -55,6 +55,7 @@
 #include "Animation/AnimMontage.h"
 #include "Animation/AnimSequenceHelpers.h"
 #include "AnimationBlueprintLibrary.h"
+#include "ClothingSimulationInstance.h"
 
 TAutoConsoleVariable<bool> CVarShowEngineDebugMessageOnAnimViewport(TEXT("Persona.AnimationEditorViewportClient.ShowEngineDebugMessageOnAnimViewport"), true, TEXT("When true show engine onscreen debug messages on animation editor viewport clients."));
 
@@ -1381,46 +1382,52 @@ FText FAnimationViewportClient::GetDisplayInfo(bool bDisplayAllInfo) const
 		}
 		
 		const bool bMirroring = PreviewMeshComponent->PreviewInstance && PreviewMeshComponent->PreviewInstance->GetMirrorDataTable();
-        if (bMirroring)
-        {
-        	TextValue = ConcatenateLine(TextValue,FText::Format(LOCTEXT("Preview_mirrored", "Mirrored with {0} "), FText::FromString(PreviewMeshComponent->PreviewInstance->GetMirrorDataTable()->GetName())));
-        }
+		if (bMirroring)
+		{
+			TextValue = ConcatenateLine(TextValue,FText::Format(LOCTEXT("Preview_mirrored", "Mirrored with {0} "), FText::FromString(PreviewMeshComponent->PreviewInstance->GetMirrorDataTable()->GetName())));
+		}
 	}
 
-	if (const IClothingSimulation* const ClothingSimulation = PreviewMeshComponent->GetClothingSimulation())
+	const TArray<FClothingSimulationInstance>& ClothingSimulationInstances = PreviewMeshComponent->GetClothingSimulationInstances();
+	for (int32 Index = 0; Index < ClothingSimulationInstances.Num(); ++Index)
 	{
-		// Cloth stats
-		if (const int32 NumActiveCloths = ClothingSimulation->GetNumCloths())
-		{
-			TextValue = ConcatenateLine(TextValue, FText::Format(LOCTEXT("NumActiveCloths", "Active Cloths: {0}"), NumActiveCloths));
-		}
-		if (const int32 NumKinematicParticles = ClothingSimulation->GetNumKinematicParticles())
-		{
-			TextValue = ConcatenateLine(TextValue, FText::Format(LOCTEXT("NumKinematicParticles", "Kinematic Particles: {0}"), NumKinematicParticles));
-		}
-		if (const int32 NumDynamicParticles = ClothingSimulation->GetNumDynamicParticles())
-		{
-			TextValue = ConcatenateLine(TextValue, FText::Format(LOCTEXT("NumDynamicParticles", "Dynamic Particles: {0}"), NumDynamicParticles));
-		}
-		if (const int32 NumIterations = ClothingSimulation->GetNumIterations())
-		{
-			TextValue = ConcatenateLine(TextValue, FText::Format(LOCTEXT("NumIterations", "Iterations: {0}"), NumIterations));
-		}
-		if (const float SimulationTime = ClothingSimulation->GetSimulationTime())
-		{
-			FNumberFormattingOptions NumberFormatOptions;
-			NumberFormatOptions.AlwaysSign = false;
-			NumberFormatOptions.UseGrouping = false;
-			NumberFormatOptions.RoundingMode = ERoundingMode::HalfFromZero;
-			NumberFormatOptions.MinimumIntegralDigits = 1;
-			NumberFormatOptions.MaximumIntegralDigits = 6;
-			NumberFormatOptions.MinimumFractionalDigits = 2;
-			NumberFormatOptions.MaximumFractionalDigits = 2;
-			TextValue = ConcatenateLine(TextValue, FText::Format(LOCTEXT("SimulationTime", "Simulation Time: {0}ms"), FText::AsNumber(SimulationTime, &NumberFormatOptions)));
-		}
-		if (ClothingSimulation->IsTeleported())
+		const FClothingSimulationInstance& ClothingSimulationInstance = ClothingSimulationInstances[Index];
+
+		if (const IClothingSimulation* const ClothingSimulation = ClothingSimulationInstance.GetClothingSimulation())
 		{
-			TextValue = ConcatenateLine(TextValue, LOCTEXT("IsTeleported", "Simulation Teleport Activated"));
+			// Cloth stats
+			if (const int32 NumActiveCloths = ClothingSimulation->GetNumCloths())
+			{
+				TextValue = ConcatenateLine(TextValue, FText::Format(LOCTEXT("NumActiveCloths", "{0}>Active Cloths: {1}"), Index, NumActiveCloths));
+			}
+			if (const int32 NumKinematicParticles = ClothingSimulation->GetNumKinematicParticles())
+			{
+				TextValue = ConcatenateLine(TextValue, FText::Format(LOCTEXT("NumKinematicParticles", "{0}>Kinematic Particles: {1}"), Index, NumKinematicParticles));
+			}
+			if (const int32 NumDynamicParticles = ClothingSimulation->GetNumDynamicParticles())
+			{
+				TextValue = ConcatenateLine(TextValue, FText::Format(LOCTEXT("NumDynamicParticles", "#{0}>Dynamic Particles: {1}"), Index, NumDynamicParticles));
+			}
+			if (const int32 NumIterations = ClothingSimulation->GetNumIterations())
+			{
+				TextValue = ConcatenateLine(TextValue, FText::Format(LOCTEXT("NumIterations", "{0}:>Iterations: {1}"), Index, NumIterations));
+			}
+			if (const float SimulationTime = ClothingSimulation->GetSimulationTime())
+			{
+				FNumberFormattingOptions NumberFormatOptions;
+				NumberFormatOptions.AlwaysSign = false;
+				NumberFormatOptions.UseGrouping = false;
+				NumberFormatOptions.RoundingMode = ERoundingMode::HalfFromZero;
+				NumberFormatOptions.MinimumIntegralDigits = 1;
+				NumberFormatOptions.MaximumIntegralDigits = 6;
+				NumberFormatOptions.MinimumFractionalDigits = 2;
+				NumberFormatOptions.MaximumFractionalDigits = 2;
+				TextValue = ConcatenateLine(TextValue, FText::Format(LOCTEXT("SimulationTime", "{0}>Simulation Time: {1}ms"), Index, FText::AsNumber(SimulationTime, &NumberFormatOptions)));
+			}
+			if (ClothingSimulation->IsTeleported())
+			{
+				TextValue = ConcatenateLine(TextValue, FText::Format(LOCTEXT("IsTeleported", "{0}>Simulation Teleport Activated"), Index));
+			}
 		}
 	}
 
diff --git a/Engine/Source/Editor/Persona/Private/ViewportToolbar/AnimationEditorMenus.cpp b/Engine/Source/Editor/Persona/Private/ViewportToolbar/AnimationEditorMenus.cpp
--- a/Engine/Source/Editor/Persona/Private/ViewportToolbar/AnimationEditorMenus.cpp
+++ b/Engine/Source/Editor/Persona/Private/ViewportToolbar/AnimationEditorMenus.cpp
@@ -13,6 +13,8 @@
 #include "BoneSelectionWidget.h"
 #include "BufferVisualizationMenuCommands.h"
 #include "NaniteVisualizationMenuCommands.h"
+#include "ClothingSimulationFactory.h"
+#include "ClothingSimulationInstance.h"
 #include "ClothingSystemEditorInterfaceModule.h"
 #include "ContentBrowserModule.h"
 #include "EditorViewportCommands.h"
@@ -264,7 +266,7 @@ void FillCharacterClothingMenu(UToolMenu* InMenu)
 		TSharedRef<IPersonaPreviewScene> PreviewScene = AnimationViewportClient->GetPreviewScene();
 		if (UDebugSkelMeshComponent* PreviewComponent = PreviewScene->GetPreviewMeshComponent())
 		{
-			if (PreviewComponent->ClothingSimulationFactory) // The cloth plugin could be disabled, and the factory would be null in this case
+			if (PreviewComponent->GetClothingSimulationInstances().Num())
 			{
 				// Currently using FNewToolMenuDelegateLegacy since this extension is being done via
 				// ISimulationEditorExtender::ExtendViewportShowMenu(FMenuBuilder&InMenuBuilder, TSharedRef<IPersonaPreviewScene> InPreviewScene)
@@ -295,15 +297,25 @@ void FillCharacterClothingMenu(UToolMenu* InMenu)
 								FModuleManager::LoadModuleChecked<FClothingSystemEditorInterfaceModule>(
 									TEXT("ClothingSystemEditorInterface")
 								);
-
-							if (ISimulationEditorExtender* Extender = ClothingEditorModule.GetSimulationEditorExtender(
-									PreviewComponent->ClothingSimulationFactory->GetFName()
-								))
+							bool bHasExtender = false;
+							for (const FClothingSimulationInstance& ClothingSimulationInstance : PreviewComponent->GetClothingSimulationInstances())
 							{
-								// Calling end section will set bSectionNeedsToBeApplied to false.
-								// Without doing so, calling ExtendViewportShowMenu will end up triggering an ensure
-								InMenuBuilder.EndSection();
-								Extender->ExtendViewportShowMenu(InMenuBuilder, PreviewScene);
+								UClothingSimulationFactory* const ClothingSimulationFactory = ClothingSimulationInstance.GetClothingSimulationFactory();
+								if (ClothingSimulationFactory && ClothingSimulationFactory->GetClass())
+								{
+									if (ISimulationEditorExtender* const Extender = ClothingEditorModule.GetSimulationEditorExtender(
+										ClothingSimulationFactory->GetClass()->GetFName()))
+									{
+										if (!bHasExtender)
+										{
+											// Calling end section will set bSectionNeedsToBeApplied to false.
+											// Without doing so, calling ExtendViewportShowMenu will end up triggering an ensure
+											InMenuBuilder.EndSection();
+											bHasExtender = true;
+										}
+										Extender->ExtendViewportShowMenu(InMenuBuilder, PreviewScene);
+									}
+								}
 							}
 						}
 					)
diff --git a/Engine/Source/Editor/PhysicsAssetEditor/Private/PhysicsAssetEditorSharedData.cpp b/Engine/Source/Editor/PhysicsAssetEditor/Private/PhysicsAssetEditorSharedData.cpp
--- a/Engine/Source/Editor/PhysicsAssetEditor/Private/PhysicsAssetEditorSharedData.cpp
+++ b/Engine/Source/Editor/PhysicsAssetEditor/Private/PhysicsAssetEditorSharedData.cpp
@@ -45,6 +45,7 @@
 #include "Widgets/Text/STextBlock.h"
 #include "Widgets/Notifications/SNotificationList.h"
 #include "Framework/Notifications/NotificationManager.h"
+#include "ClothingSimulationInstance.h"
 #include "ClothingSimulationInteractor.h"
 #include "UnrealExporter.h"
 #include "Exporters/Exporter.h"
@@ -3732,9 +3733,15 @@ void FPhysicsAssetEditorSharedData::ForceDisableSimulation()
 
 void FPhysicsAssetEditorSharedData::UpdateClothPhysics()
 {
-	if(EditorSkelComp && EditorSkelComp->GetClothingSimulationInteractor())
+	if (EditorSkelComp)
 	{
-		EditorSkelComp->GetClothingSimulationInteractor()->PhysicsAssetUpdated();
+		for (const FClothingSimulationInstance& ClothingSimulationInstance : EditorSkelComp->GetClothingSimulationInstances())
+		{
+			if (ClothingSimulationInstance.GetClothingSimulationInteractor())
+			{
+				ClothingSimulationInstance.GetClothingSimulationInteractor()->PhysicsAssetUpdated();
+			}
+		}
 	}
 }
 
diff --git a/Engine/Source/Editor/UnrealEd/Classes/Animation/DebugSkelMeshComponent.h b/Engine/Source/Editor/UnrealEd/Classes/Animation/DebugSkelMeshComponent.h
--- a/Engine/Source/Editor/UnrealEd/Classes/Animation/DebugSkelMeshComponent.h
+++ b/Engine/Source/Editor/UnrealEd/Classes/Animation/DebugSkelMeshComponent.h
@@ -594,6 +594,7 @@ public:
 	 * To normal game/runtime code we don't want to expose a non-const pointer to the simulation, so we can only get
 	 * one from this editor-only component. Intended for debug options/visualisations/editor-only code to poke the sim
 	 */
+	UE_DEPRECATED(5.7, "Use GetClothingSimulationInstances instead.")
 	UNREALED_API IClothingSimulation* GetMutableClothingSimulation();
 
 	/** to avoid clothing reset while modifying properties in Persona */
diff --git a/Engine/Source/Editor/UnrealEd/Private/Animation/DebugSkelMeshComponent.cpp b/Engine/Source/Editor/UnrealEd/Private/Animation/DebugSkelMeshComponent.cpp
--- a/Engine/Source/Editor/UnrealEd/Private/Animation/DebugSkelMeshComponent.cpp
+++ b/Engine/Source/Editor/UnrealEd/Private/Animation/DebugSkelMeshComponent.cpp
@@ -22,7 +22,7 @@
 #include "Rendering/SkeletalMeshModel.h"
 
 #include "ClothingAsset.h"
-#include "ClothingSimulation.h"
+#include "ClothingSimulationInstance.h"  // UE_DEPRECATED(5.7, "Use GetClothingSimulationInstances instead.")
 #include "Utils/ClothingMeshUtils.h"
 #include "DynamicMeshBuilder.h"
 #include "SkeletalDebugRendering.h"
@@ -1284,9 +1284,12 @@ void UDebugSkelMeshComponent::GetUsedMaterials(TArray<UMaterialInterface *>& Out
 	}
 }
 
+// UE_DEPRECATED(5.7, "Use GetClothingSimulationInstances instead.")
 IClothingSimulation* UDebugSkelMeshComponent::GetMutableClothingSimulation()
 {
-	return ClothingSimulation;
+	// This function seems unused, and therefore a const_cast is probably a better idea than compromising the const version of GetClothingSimulation()
+	return GetClothingSimulationInstances().Num() ?
+		const_cast<IClothingSimulation*>(GetClothingSimulationInstances()[0].GetClothingSimulation()) : nullptr;
 }
 
 void UDebugSkelMeshComponent::RebuildCachedClothBounds()
diff --git a/Engine/Source/Editor/UnrealEd/Private/ApexClothingUtils.cpp b/Engine/Source/Editor/UnrealEd/Private/ApexClothingUtils.cpp
--- a/Engine/Source/Editor/UnrealEd/Private/ApexClothingUtils.cpp
+++ b/Engine/Source/Editor/UnrealEd/Private/ApexClothingUtils.cpp
@@ -55,7 +55,9 @@ void RestoreAllClothingSections(USkeletalMesh* SkelMesh, uint32 LODIndex, uint32
 			{
 				if(Section.HasClothingData())
 				{
+PRAGMA_DISABLE_DEPRECATION_WARNINGS
 					ClothingAssetUtils::ClearSectionClothingData(Section);
+PRAGMA_ENABLE_DEPRECATION_WARNINGS
 					if (FSkelMeshSourceSectionUserData* UserSectionData = LodModel.UserSectionsData.Find(Section.OriginalDataSectionIndex))
 					{
 						UserSectionData->CorrespondClothAssetIndex = INDEX_NONE;
diff --git a/Engine/Source/Runtime/ClothingSystemRuntimeCommon/Private/ClothingAsset.cpp b/Engine/Source/Runtime/ClothingSystemRuntimeCommon/Private/ClothingAsset.cpp
--- a/Engine/Source/Runtime/ClothingSystemRuntimeCommon/Private/ClothingAsset.cpp
+++ b/Engine/Source/Runtime/ClothingSystemRuntimeCommon/Private/ClothingAsset.cpp
@@ -19,6 +19,7 @@
 
 #include "Components/SkeletalMeshComponent.h"
 
+#include "ClothingSimulationFactory.h"
 #include "ClothingSimulationInteractor.h"
 #include "ComponentReregisterContext.h"
 #include "UObject/UObjectIterator.h"
@@ -40,6 +41,91 @@
 //==============================================================================
 
 #if WITH_EDITOR
+void ClothingAssetUtils::UnbindFromSkeletalMesh(const UClothingAssetBase& Asset, USkeletalMesh& SkeletalMesh)
+{
+	FScopedSkeletalMeshPostEditChange SkeletalMeshPostEditChange(&SkeletalMesh);
+
+	if (const FSkeletalMeshModel* const Mesh = SkeletalMesh.GetImportedModel())
+	{
+		const int32 NumLods = Mesh->LODModels.Num();
+		for (int32 LodIndex = 0; LodIndex < NumLods; ++LodIndex)
+		{
+			UnbindFromSkeletalMesh(Asset, SkeletalMesh, LodIndex);
+		}
+	}
+}
+
+void ClothingAssetUtils::UnbindFromSkeletalMesh(const UClothingAssetBase& Asset, USkeletalMesh& SkeletalMesh, const int32 MeshLodIndex)
+{
+	// Find the chunk(s) we created
+	if (FSkeletalMeshModel* const Mesh = SkeletalMesh.GetImportedModel())
+	{
+		bool bChangedMesh = false;
+
+		if (!Mesh->LODModels.IsValidIndex(MeshLodIndex))
+		{
+			UE_LOG(LogClothingAsset, Warning, TEXT("Failed to remove clothing asset [%s] from mesh LOD%d as that LOD doesn't exist."),
+				*Asset.GetName(),
+				MeshLodIndex);
+			return;
+		}
+
+		FSkeletalMeshLODModel& LodModel = Mesh->LODModels[MeshLodIndex];
+
+		for (int32 SectionIdx = LodModel.Sections.Num() - 1; SectionIdx >= 0; --SectionIdx)
+		{
+			FSkelMeshSection& Section = LodModel.Sections[SectionIdx];
+			if (Section.HasClothingData() && Section.ClothingData.AssetGuid == Asset.GetAssetGuid())
+			{
+				// Don't do this when in async task, this should have been done before initiating the build
+				if (IsInGameThread())
+				{
+					SkeletalMesh.PreEditChange(nullptr);
+				}
+
+				// Log the LOD bias mapping data that will be removed
+				for (int32 LodIndex = 0; LodIndex < MeshLodIndex; ++LodIndex)
+				{
+					const int32 LODBias = MeshLodIndex - LodIndex;
+
+					UE_CLOG(Section.ClothMappingDataLODs.IsValidIndex(LODBias) && Section.ClothMappingDataLODs[LODBias].Num(),
+						LogClothingAsset, Verbose, TEXT("Removed deformer data for [%s/%s], section %d LODBias = %d, render mesh LOD = %d, sim mesh LOD = %d"),
+						*SkeletalMesh.GetName(),
+						*Asset.GetName(),
+						SectionIdx,
+						LODBias,
+						MeshLodIndex,
+						LodIndex);
+				}
+
+				Section.ClothingData.AssetGuid = FGuid();
+				Section.ClothingData.AssetLodIndex = INDEX_NONE;
+				Section.CorrespondClothAssetIndex = INDEX_NONE;
+				Section.ClothMappingDataLODs.Empty();
+
+				if (FSkelMeshSourceSectionUserData* const UserSectionData = LodModel.UserSectionsData.Find(Section.OriginalDataSectionIndex))
+				{
+					UserSectionData->CorrespondClothAssetIndex = INDEX_NONE;
+					UserSectionData->ClothingData.AssetLodIndex = INDEX_NONE;
+					UserSectionData->ClothingData.AssetGuid = FGuid();
+				}
+
+				// Clear all remaining LOD bias mapping data from other LODs that relied on this section
+				ClothingAssetUtils::ClearLODBiasMappings(SkeletalMesh, MeshLodIndex, SectionIdx);
+
+				bChangedMesh = true;
+			}
+		}
+
+		// If the mesh changed we need to re-register any components that use it to reflect the changes
+		if (bChangedMesh)
+		{
+			//Register the scope post edit change
+			FScopedSkeletalMeshPostEditChange SkeletalMeshPostEditChange(&SkeletalMesh);
+		}
+	}
+}
+
 void ClothingAssetUtils::GetAllMeshClothingAssetBindings(const USkeletalMesh* SkeletalMesh, TArray<FClothingAssetMeshBinding>& OutBindings)
 {
 	OutBindings.Empty();
@@ -98,13 +184,48 @@ void ClothingAssetUtils::GetAllLodMeshClothingAssetBindings(const USkeletalMesh*
 	}
 }
 
-void ClothingAssetUtils::ClearSectionClothingData(FSkelMeshSection& InSection)
+// UE_DEPRECATED(5.7, "Unused low level function, will be removed or re-added as a member to FSkelMeshSection.")
+void ClothingAssetUtils::ClearSectionClothingData(FSkelMeshSection& Section)
 {
-	InSection.ClothingData.AssetGuid = FGuid();
-	InSection.ClothingData.AssetLodIndex = INDEX_NONE;
-	InSection.CorrespondClothAssetIndex = INDEX_NONE;
+	Section.ClothingData.AssetGuid = FGuid();
+	Section.ClothingData.AssetLodIndex = INDEX_NONE;
+	Section.CorrespondClothAssetIndex = INDEX_NONE;
 
-	InSection.ClothMappingDataLODs.Empty();
+	Section.ClothMappingDataLODs.Empty();
+}
+
+void ClothingAssetUtils::ClearLODBiasMappings(USkeletalMesh& SkeletalMesh, const int32 UpdatedLodIndex, const int32 SectionIndex)
+{
+	if (FSkeletalMeshModel* const SkeletalMeshModel = SkeletalMesh.GetImportedModel())
+	{
+		if (SkeletalMeshModel->LODModels.IsValidIndex(UpdatedLodIndex))
+		{
+			for (int32 LODIndex = UpdatedLodIndex + 1; LODIndex < SkeletalMeshModel->LODModels.Num(); ++LODIndex)
+			{
+				FSkeletalMeshLODModel& BiasLodModel = SkeletalMeshModel->LODModels[LODIndex];
+				if (BiasLodModel.Sections.IsValidIndex(SectionIndex))
+				{
+					FSkelMeshSection& BiasSection = BiasLodModel.Sections[SectionIndex];
+					const int32 LODBias = LODIndex - UpdatedLodIndex;
+
+					if (BiasSection.ClothMappingDataLODs.IsValidIndex(LODBias))
+					{
+						BiasSection.ClothMappingDataLODs[LODBias].Empty();
+
+						const UClothingAssetBase* const Asset = SkeletalMesh.GetClothingAsset(BiasSection.ClothingData.AssetGuid);
+
+						UE_LOG(LogClothingAsset, Verbose, TEXT("Removed deformer data for [%s/%s], section %d LODBias = %d, render mesh LOD = %d, sim mesh LOD = %d"),
+							*SkeletalMesh.GetName(),
+							Asset ? *Asset->GetName() : TEXT("?"),
+							SectionIndex,
+							LODBias,
+							LODIndex,
+							UpdatedLodIndex);
+					}
+				}
+			}
+		}
+	}
 }
 #endif
 
@@ -558,126 +679,32 @@ void UClothingAssetCommon::UpdateLODBiasMappings(const USkeletalMesh* SkeletalMe
 	}
 }
 
-void UClothingAssetCommon::ClearLODBiasMappings(const USkeletalMesh* SkeletalMesh, int32 UpdatedLODIndex, int32 SectionIndex)
+void UClothingAssetCommon::UnbindFromSkeletalMesh(USkeletalMesh* SkeletalMesh)
 {
-	check(SkeletalMesh);
-	check(UpdatedLODIndex >= 0);
-
-	FSkeletalMeshModel* const SkeletalMeshModel = SkeletalMesh->GetImportedModel();
-	check(SkeletalMeshModel);
-
-	for (int32 LODIndex = UpdatedLODIndex + 1; LODIndex < SkeletalMeshModel->LODModels.Num(); ++LODIndex)
+	// Clear the entire LOD map to allow for rebinding a new asset
+	for (int32& LodMapIndex : LodMap)
 	{
-		FSkeletalMeshLODModel& BiasLodModel = SkeletalMeshModel->LODModels[LODIndex];
-		if (BiasLodModel.Sections.IsValidIndex(SectionIndex))
-		{
-			FSkelMeshSection& BiasSection = BiasLodModel.Sections[SectionIndex];
-			const int32 LODBias = LODIndex - UpdatedLODIndex;
-
-			if (BiasSection.ClothMappingDataLODs.IsValidIndex(LODBias))
-			{
-				BiasSection.ClothMappingDataLODs[LODBias].Empty();
-
-				UE_LOG(LogClothingAsset, Verbose, TEXT("Removed deformer data for [%s/%s], section %d LODBias = %d, render mesh LOD = %d, sim mesh LOD = %d"),
-					*SkeletalMesh->GetName(),
-					*GetName(),
-					SectionIndex,
-					LODBias,
-					LODIndex,
-					UpdatedLODIndex);
-			}
-		}
+		LodMapIndex = INDEX_NONE;  // Note: This must be done prior calling UnbindFromSkeletalMesh(), so that the info is correct when rebuilding the mesh
 	}
-}
 
-void UClothingAssetCommon::UnbindFromSkeletalMesh(USkeletalMesh* InSkelMesh)
-{
-	if(FSkeletalMeshModel* Mesh = InSkelMesh->GetImportedModel())
+	// Unbind from the Skeletal Mesh
+	if (SkeletalMesh)
 	{
-		const int32 NumLods = Mesh->LODModels.Num();
-		for(int32 LodIndex = 0; LodIndex < NumLods; ++LodIndex)
-		{
-			UnbindFromSkeletalMesh(InSkelMesh, LodIndex);
-		}
+		ClothingAssetUtils::UnbindFromSkeletalMesh(*this, *SkeletalMesh);
 	}
 }
 
-void UClothingAssetCommon::UnbindFromSkeletalMesh(
-	USkeletalMesh* InSkelMesh, 
-	const int32 InMeshLodIndex)
+void UClothingAssetCommon::UnbindFromSkeletalMesh(USkeletalMesh* SkeletalMesh, const int32 MeshLodIndex)
 {
-	bool bChangedMesh = false;
-
-	// Find the chunk(s) we created
-	if(FSkeletalMeshModel* Mesh = InSkelMesh->GetImportedModel())
+	// Clear the LOD map entry to allow for rebinding a new asset
+	if (LodMap.IsValidIndex(MeshLodIndex))
 	{
-		if(!Mesh->LODModels.IsValidIndex(InMeshLodIndex))
-		{
-			const FText Error = FText::Format(
-				LOCTEXT("Error_UnbindNoMeshLod", "Failed to remove clothing asset {0} from mesh LOD{1} as that LOD doesn't exist."), 
-				FText::FromString(GetName()), 
-				InMeshLodIndex);
-			WarningNotification(Error);
-
-			return;
-		}
-
-		FSkeletalMeshLODModel& LodModel = Mesh->LODModels[InMeshLodIndex];
-
-		for(int32 SectionIdx = LodModel.Sections.Num() - 1; SectionIdx >= 0; --SectionIdx)
-		{
-			FSkelMeshSection& Section = LodModel.Sections[SectionIdx];
-			if(Section.HasClothingData() && Section.ClothingData.AssetGuid == AssetGuid)
-			{
-				// Don't do this when in async task, this should have been done before initiating the build
-				if (IsInGameThread())
-				{
-					InSkelMesh->PreEditChange(nullptr);
-				}
-
-				// Log the LOD bias mapping data that will be removed through the call to ClearSectionClothingData
-				for (int32 LodIndex = 0; LodIndex < InMeshLodIndex; ++LodIndex)
-				{
-					const int32 LODBias = InMeshLodIndex - LodIndex;
-
-					UE_CLOG(Section.ClothMappingDataLODs.IsValidIndex(LODBias) && Section.ClothMappingDataLODs[LODBias].Num(),
-						LogClothingAsset, Verbose, TEXT("Removed deformer data for [%s/%s], section %d LODBias = %d, render mesh LOD = %d, sim mesh LOD = %d"),
-						*InSkelMesh->GetName(),
-						*GetName(),
-						SectionIdx,
-						LODBias,
-						InMeshLodIndex,
-						LodIndex);
-				}
-
-				ClothingAssetUtils::ClearSectionClothingData(Section);
-				if (FSkelMeshSourceSectionUserData* UserSectionData = LodModel.UserSectionsData.Find(Section.OriginalDataSectionIndex))
-				{
-					UserSectionData->CorrespondClothAssetIndex = INDEX_NONE;
-					UserSectionData->ClothingData.AssetLodIndex = INDEX_NONE;
-					UserSectionData->ClothingData.AssetGuid = FGuid();
-				}
-				
-				// Clear all remaining LOD bias mapping data from other LODs that relied on this section
-				ClearLODBiasMappings(InSkelMesh, InMeshLodIndex, SectionIdx);
-
-				bChangedMesh = true;
-			}
-		}
-
-		// Clear the LOD map entry for this asset LOD, after a unbind we must be able to bind any asset
-		if (LodMap.IsValidIndex(InMeshLodIndex))
-		{
-			LodMap[InMeshLodIndex] = INDEX_NONE;
-			bChangedMesh = true;
-		}
+		LodMap[MeshLodIndex] = INDEX_NONE;  // Note: This must be done prior calling UnbindFromSkeletalMesh(), so that the info is correct when rebuilding the mesh
 	}
-
-	// If the mesh changed we need to re-register any components that use it to reflect the changes
-	if(bChangedMesh)
+	// Unbind from the Skeletal Mesh
+	if (SkeletalMesh)
 	{
-		//Register the scope post edit change
-		FScopedSkeletalMeshPostEditChange SkeletalMeshPostEditChange(InSkelMesh);
+		ClothingAssetUtils::UnbindFromSkeletalMesh(*this, *SkeletalMesh, MeshLodIndex);
 	}
 }
 
@@ -709,7 +736,7 @@ void UClothingAssetCommon::ForEachInteractorUsingClothing(TFunction<void(UClothi
 			{
 				if(Component->GetSkeletalMeshAsset() == OwnerMesh)
 				{
-					if(UClothingSimulationInteractor* CurInteractor = Component->GetClothingSimulationInteractor())
+					if(UClothingSimulationInteractor* CurInteractor = Component->GetClothingSimulationInteractor(GetFName()))
 					{
 						Func(CurInteractor);
 					}
@@ -1227,8 +1254,6 @@ void UClothingAssetCommon::DeclareConstructClasses(TArray<FTopLevelAssetPath>& O
 			}
 		}
 	}
-	// OutConstructClasses.Add(FTopLevelAssetPath(UChaosClothSharedSimConfig::StaticClass()));
-	//OutConstructClasses.Add(FTopLevelAssetPath(UChaosClothConfig::StaticClass()));
 }
 #endif
 
@@ -1255,6 +1280,10 @@ bool UClothingAssetCommon::AddClothConfigs()
 		if (UClass* const ClothingSimulationFactoryClass = *TSubclassOf<class UClothingSimulationFactory>(Provider->GetClothingSimulationFactoryClass()))
 		{
 			const UClothingSimulationFactory* const ClothingSimulationFactory = ClothingSimulationFactoryClass->GetDefaultObject<UClothingSimulationFactory>();
+			if (!ClothingSimulationFactory || !ClothingSimulationFactory->SupportsAsset(this))
+			{
+				continue;  // No need to add configs for simulations that don't support this asset
+			}
 			for (TSubclassOf<UClothConfigBase> ClothConfigClass : ClothingSimulationFactory->GetClothConfigClasses())
 			{
 				const FName ClothConfigName = ClothConfigClass->GetFName();
diff --git a/Engine/Source/Runtime/ClothingSystemRuntimeCommon/Private/ClothingSimulation.cpp b/Engine/Source/Runtime/ClothingSystemRuntimeCommon/Private/ClothingSimulation.cpp
--- a/Engine/Source/Runtime/ClothingSystemRuntimeCommon/Private/ClothingSimulation.cpp
+++ b/Engine/Source/Runtime/ClothingSystemRuntimeCommon/Private/ClothingSimulation.cpp
@@ -203,14 +203,7 @@ FClothingSimulationCommon::FClothingSimulationCommon()
 FClothingSimulationCommon::~FClothingSimulationCommon()
 {}
 
-void FClothingSimulationCommon::FillContext(USkeletalMeshComponent* InComponent, float InDeltaTime, IClothingSimulationContext* InOutContext)
-{
-	// Deprecated version always fills RefToLocals with current animation results instead of using reference pose on initialization
-	const bool bIsInitialization = false;
-	FillContext(InComponent, InDeltaTime, InOutContext, bIsInitialization);
-}
-
-void FClothingSimulationCommon::FillContext(USkeletalMeshComponent* InComponent, float InDeltaTime, IClothingSimulationContext* InOutContext, bool bIsInitialization)
+void FClothingSimulationCommon::FillContext(const USkeletalMeshComponent* InComponent, float InDeltaTime, IClothingSimulationContext* InOutContext, bool bIsInitialization) const
 {
 	check(InOutContext);
 	FClothingSimulationContextCommon* const Context = static_cast<FClothingSimulationContextCommon*>(InOutContext);
diff --git a/Engine/Source/Runtime/ClothingSystemRuntimeCommon/Public/ClothingAsset.h b/Engine/Source/Runtime/ClothingSystemRuntimeCommon/Public/ClothingAsset.h
--- a/Engine/Source/Runtime/ClothingSystemRuntimeCommon/Public/ClothingAsset.h
+++ b/Engine/Source/Runtime/ClothingSystemRuntimeCommon/Public/ClothingAsset.h
@@ -10,6 +10,8 @@
 #include "Algo/AnyOf.h"
 #include "ClothingAsset.generated.h"
 
+#define UE_API CLOTHINGSYSTEMRUNTIMECOMMON_API
+
 class UPhysicsAsset;
 struct FSkelMeshSection;
 
@@ -30,29 +32,47 @@ namespace ClothingAssetUtils
 
 #if WITH_EDITOR
 	/**
-	 * Given a skeletal mesh model, find all of the currently bound clothing assets and their binding information
-	 * @param MeshModel - The skeletal mesh imported model to extract binding from
-	 * @param OutBindings - The list of bindings to write to
+	 * Unbind this clothing asset from the provided skeletal mesh, will remove all LODs.
+	 * @param SkeletalMesh skeletal mesh to remove from.
+	 */
+	UE_API void UnbindFromSkeletalMesh(const UClothingAssetBase& Asset, USkeletalMesh& SkeletalMesh);
+
+	/**
+	 * Unbind this clothing asset from the provided skeletal mesh.
+	 * @param Asset The clothing data to unbind.
+	 * @param SkeletalMesh The Skeletal mesh to remove from.
+	 * @param MeshLodIndex Mesh LOD to remove this asset from (could still be bound to other LODs).
 	 */
-	void CLOTHINGSYSTEMRUNTIMECOMMON_API
-	GetAllMeshClothingAssetBindings(
-		const USkeletalMesh* InSkelMesh,
-		TArray<FClothingAssetMeshBinding>& OutBindings);
+	UE_API void UnbindFromSkeletalMesh(const UClothingAssetBase& Asset, USkeletalMesh& SkeletalMesh, const int32 MeshLodIndex);
 
 	/**
-	 * Similar to above, but only inspects the specified LOD.
+	 * Given a skeletal mesh model, find all of the currently bound clothing assets and their binding information.
+	 * @param SkeletalMesh The skeletal mesh to extract binding from.
+	 * @param OutBindings The list of bindings to write to.
 	 */
-	void CLOTHINGSYSTEMRUNTIMECOMMON_API
-	GetAllLodMeshClothingAssetBindings(
-		const USkeletalMesh* InSkelMesh,
+	UE_API void GetAllMeshClothingAssetBindings(const USkeletalMesh* SkeletalMesh, TArray<FClothingAssetMeshBinding>& OutBindings);
+
+	/**
+	 * Given a skeletal mesh model, find all of the currently bound clothing assets and their binding information.
+	 * @param SkeletalMesh The skeletal mesh to extract binding from.
+	 * @param OutBindings The list of bindings to write to.
+	 * @param LodIndex The specified LOD.
+	 */
+	UE_API void GetAllLodMeshClothingAssetBindings(
+		const USkeletalMesh* SkeletalMesh,
 		TArray<FClothingAssetMeshBinding>& OutBindings,
-		int32 InLodIndex);
+		int32 LodIndex);
 
 	/**
-	 * Clears the clothing tracking struct of a section.
+	 * Clear all defomer section that relies on the specified cloth sim data bound to a LOD section.
+	 * @param SkeletalMesh The skeletal mesh to clear.
+	 * @param UpdatedLodIndex The recently updated Skeletal Mesh LOD index from which the bias mapping needs to be cleared.
+	 * @param SectionIndex The section to clear.
 	 */
-	CLOTHINGSYSTEMRUNTIMECOMMON_API void 
-	ClearSectionClothingData(FSkelMeshSection& InSection);
+	UE_API void ClearLODBiasMappings(USkeletalMesh& SkeletalMesh, const int32 UpdatedLodIndex, const int32 SectionIndex);
+
+	UE_DEPRECATED(5.7, "Unused low level function, will be removed or re-added as a member to FSkelMeshSection.")
+	UE_API void ClearSectionClothingData(FSkelMeshSection& Section);
 #endif
 }
 
@@ -94,7 +114,7 @@ class UClothingAssetCommon : public UClothingAssetBase
 	GENERATED_BODY()
 public:
 
-	CLOTHINGSYSTEMRUNTIMECOMMON_API UClothingAssetCommon(const FObjectInitializer& ObjectInitializer);
+	UE_API UClothingAssetCommon(const FObjectInitializer& ObjectInitializer);
 
 #if WITH_EDITOR
 
@@ -102,45 +122,45 @@ public:
 	 * Create weights for skinning the render mesh to our simulation mesh, and 
 	 * weights to drive our sim mesh from the skeleton.
 	 */
-	CLOTHINGSYSTEMRUNTIMECOMMON_API virtual bool BindToSkeletalMesh(USkeletalMesh* InSkelMesh, const int32 InMeshLodIndex, const int32 InSectionIndex, const int32 InAssetLodIndex) override;
+	UE_API virtual bool BindToSkeletalMesh(USkeletalMesh* InSkelMesh, const int32 InMeshLodIndex, const int32 InSectionIndex, const int32 InAssetLodIndex) override;
 
 	/**
 	 * Helper that invokes \c UnbindFromSkeletalMesh() for each avilable entry in 
 	 * \p InSkelMesh->GetImportedModel()'s LODModel.
 	 */
-	CLOTHINGSYSTEMRUNTIMECOMMON_API virtual void UnbindFromSkeletalMesh(USkeletalMesh* InSkelMesh) override;
-	CLOTHINGSYSTEMRUNTIMECOMMON_API virtual void UnbindFromSkeletalMesh(USkeletalMesh* InSkelMesh, const int32 InMeshLodIndex) override;
+	UE_API virtual void UnbindFromSkeletalMesh(USkeletalMesh* InSkelMesh) override;
+	UE_API virtual void UnbindFromSkeletalMesh(USkeletalMesh* InSkelMesh, const int32 InMeshLodIndex) override;
 
 	/**
 	 * Update all extra LOD deformer mappings.
 	 * This should be called whenever the raytracing LOD bias has changed.
 	 */
-	CLOTHINGSYSTEMRUNTIMECOMMON_API virtual void UpdateAllLODBiasMappings(USkeletalMesh* SkeletalMesh) override;
+	UE_API virtual void UpdateAllLODBiasMappings(USkeletalMesh* SkeletalMesh) override;
 
 	/** 
 	 * Callback envoked after weights have been edited.
 	 * Calls \c PushWeightsToMesh() on each \c ClothLodData, and invalidates cached data. 
 	 * Optionaly recalculate the owner's sections fixed vertex data based on this asset masks, and invalidate the DDC for this asset.
 	 */
-	CLOTHINGSYSTEMRUNTIMECOMMON_API void ApplyParameterMasks(bool bUpdateFixedVertData = false, bool bInvalidateDerivedDataCache = true);
+	UE_API void ApplyParameterMasks(bool bUpdateFixedVertData = false, bool bInvalidateDerivedDataCache = true);
 
 	/**
 	 * Builds the LOD transition data.
 	 * When we transition between LODs we skin the incoming mesh to the outgoing mesh
 	 * in exactly the same way the render mesh is skinned to create a smooth swap
 	 */
-	CLOTHINGSYSTEMRUNTIMECOMMON_API void BuildLodTransitionData();
+	UE_API void BuildLodTransitionData();
 
 	//~ Begin UObject interface
 	/**
 	 * Stop any simulation from using this asset.
 	 */
-	CLOTHINGSYSTEMRUNTIMECOMMON_API virtual void PreEditUndo() override;
+	UE_API virtual void PreEditUndo() override;
 
 	/**
 	 * Restart simulation using this asset after undo change.
 	 */
-	CLOTHINGSYSTEMRUNTIMECOMMON_API virtual void PostEditUndo() override;
+	UE_API virtual void PostEditUndo() override;
 	//~ End UObject interface
 #endif // WITH_EDITOR
 
@@ -148,16 +168,16 @@ public:
 	 * Rebuilds the \c UsedBoneIndices array from looking up the entries in the
 	 * \c UsedBoneNames array, in the \p InSkelMesh's reference skeleton.
 	 */
-	CLOTHINGSYSTEMRUNTIMECOMMON_API virtual void RefreshBoneMapping(USkeletalMesh* InSkelMesh) override;
+	UE_API virtual void RefreshBoneMapping(USkeletalMesh* InSkelMesh) override;
 
 	/** Calculates the preferred root bone for the simulation. */
-	CLOTHINGSYSTEMRUNTIMECOMMON_API void CalculateReferenceBoneIndex();
+	UE_API void CalculateReferenceBoneIndex();
 
 	/** Returns \c true if \p InLodIndex is a valid LOD id (index into \c ClothLodData). */
-	CLOTHINGSYSTEMRUNTIMECOMMON_API virtual bool IsValidLod(int32 InLodIndex) const override;
+	UE_API virtual bool IsValidLod(int32 InLodIndex) const override;
 
 	/** Returns the number of valid LOD's (length of the \c ClothLodData array). */
-	CLOTHINGSYSTEMRUNTIMECOMMON_API virtual int32 GetNumLods() const override;
+	UE_API virtual int32 GetNumLods() const override;
 
 #if WITH_EDITORONLY_DATA
 	/**
@@ -170,15 +190,15 @@ public:
 	 * Called on the clothing asset when the base data (physical mesh, config etc.)
 	 * has changed, so any intermediate generated data can be regenerated.
 	 */
-	CLOTHINGSYSTEMRUNTIMECOMMON_API void InvalidateFlaggedCachedData(EClothingCachedDataFlagsCommon Flags);
+	UE_API void InvalidateFlaggedCachedData(EClothingCachedDataFlagsCommon Flags);
 #endif // WITH_EDITORONLY_DATA
 
 #if WITH_EDITOR
 	/** * Add a new LOD class instance. */
-	CLOTHINGSYSTEMRUNTIMECOMMON_API virtual int32 AddNewLod() override;
+	UE_API virtual int32 AddNewLod() override;
 
 	/* Called after changes in any of the asset properties. */
-	CLOTHINGSYSTEMRUNTIMECOMMON_API virtual void PostEditChangeChainProperty(FPropertyChangedChainEvent& ChainEvent) override;
+	UE_API virtual void PostEditChangeChainProperty(FPropertyChangedChainEvent& ChainEvent) override;
 #endif // WITH_EDITOR
 
 	/** Return a const cloth config pointer of the desired cloth config type, or nullptr if there isn't any suitable. */
@@ -198,16 +218,16 @@ public:
 	}
 
 	/** Migrate deprecated objects. */
-	CLOTHINGSYSTEMRUNTIMECOMMON_API virtual void PostLoad() override;
+	UE_API virtual void PostLoad() override;
 #if WITH_EDITORONLY_DATA
-	static CLOTHINGSYSTEMRUNTIMECOMMON_API void DeclareConstructClasses(TArray<FTopLevelAssetPath>& OutConstructClasses, const UClass* SpecificSubclass);
+	static UE_API void DeclareConstructClasses(TArray<FTopLevelAssetPath>& OutConstructClasses, const UClass* SpecificSubclass);
 #endif
 
 	/** Serialize deprecated objects. */
-	CLOTHINGSYSTEMRUNTIMECOMMON_API virtual void Serialize(FArchive& Ar) override;
+	UE_API virtual void Serialize(FArchive& Ar) override;
 
 	/** Propagate the shared simulation configs between assets. Called after all cloth assets sharing the same simulation are loaded. */
-	CLOTHINGSYSTEMRUNTIMECOMMON_API virtual void PostUpdateAllAssets() override;
+	UE_API virtual void PostUpdateAllAssets() override;
 
 	// The physics asset to extract collisions from when building a simulation.
 	UPROPERTY(EditAnywhere, Category = "Skeletal Mesh Clothing")
@@ -263,12 +283,6 @@ public:
 	int32 ReferenceBoneIndex;
 
 #if WITH_EDITORONLY_DATA
-PRAGMA_DISABLE_DEPRECATION_WARNINGS
-	UE_DEPRECATED(5.5, "This property is obsolete.")
-	UPROPERTY(Meta = (DeprecatedProperty, DeprecationMessage = "This property is obsolete."))
-	TObjectPtr<UClothingAssetCustomData> CustomData;
-PRAGMA_ENABLE_DEPRECATION_WARNINGS
-
 	/** 
 	 * Deprecated property for transitioning the \c FClothConfig struct to the 
 	 * \c UClothConfigBase array, in a new property called \c ClothConfigs.
@@ -291,12 +305,12 @@ private:
 	// If a config from a different factory exists already, the newly
 	// created config will attempt to initialize its parameters from it.
 	// Return true when at least one config has been added, false otherwise.
-	CLOTHINGSYSTEMRUNTIMECOMMON_API bool AddClothConfigs();
+	bool AddClothConfigs();
 
 	// Propagate the shared simulation configs between assets.
 	// Also migrate all deprecated shared parameters which have been moved to the per cloth configs if required.
 	// Called after a cloth asset is created or loaded.
-	CLOTHINGSYSTEMRUNTIMECOMMON_API void PropagateSharedConfigs(bool bMigrateSharedConfigToConfig=false);
+	void PropagateSharedConfigs(bool bMigrateSharedConfigToConfig=false);
 
 	// Return true when any one of the cloth configs fullfill the predicate.
 	// Used to select which type of data to cache.
@@ -311,13 +325,12 @@ private:
 
 #if WITH_EDITOR
 	// Add extra cloth deformer mappings to cope with a different raytracing LOD than the one currently rendered.
-	CLOTHINGSYSTEMRUNTIMECOMMON_API void UpdateLODBiasMappings(const USkeletalMesh* SkeletalMesh, int32 UpdatedLODIndex, int32 SectionIndex);
-
-	// Clear all defomer section that relies on the specified cloth sim data bound to the specified LOD section 
-	CLOTHINGSYSTEMRUNTIMECOMMON_API void ClearLODBiasMappings(const USkeletalMesh* SkeletalMesh, int32 UpdatedLODIndex, int32 SectionIndex);
+	void UpdateLODBiasMappings(const USkeletalMesh* SkeletalMesh, int32 UpdatedLODIndex, int32 SectionIndex);
 
 	// Helper functions used in PostPropertyChangeCb
-	CLOTHINGSYSTEMRUNTIMECOMMON_API void ReregisterComponentsUsingClothing();
-	CLOTHINGSYSTEMRUNTIMECOMMON_API void ForEachInteractorUsingClothing(TFunction<void (UClothingSimulationInteractor*)> Func);
+	void ReregisterComponentsUsingClothing();
+	void ForEachInteractorUsingClothing(TFunction<void (UClothingSimulationInteractor*)> Func);
 #endif // WITH_EDITOR
 };
+
+#undef UE_API
diff --git a/Engine/Source/Runtime/ClothingSystemRuntimeCommon/Public/ClothingSimulation.h b/Engine/Source/Runtime/ClothingSystemRuntimeCommon/Public/ClothingSimulation.h
--- a/Engine/Source/Runtime/ClothingSystemRuntimeCommon/Public/ClothingSimulation.h
+++ b/Engine/Source/Runtime/ClothingSystemRuntimeCommon/Public/ClothingSimulation.h
@@ -8,39 +8,40 @@
 #include "Math/Vector.h"
 #include "Stats/Stats.h"
 
+#define UE_API CLOTHINGSYSTEMRUNTIMECOMMON_API
 class USkeletalMeshComponent;
 
 // Common simulation stats
-DECLARE_CYCLE_STAT_EXTERN(TEXT("Compute Clothing Normals"), STAT_ClothComputeNormals, STATGROUP_Physics, CLOTHINGSYSTEMRUNTIMECOMMON_API);
-DECLARE_CYCLE_STAT_EXTERN(TEXT("Internal Solve"), STAT_ClothInternalSolve, STATGROUP_Physics, CLOTHINGSYSTEMRUNTIMECOMMON_API);
-DECLARE_CYCLE_STAT_EXTERN(TEXT("Update Collisions"), STAT_ClothUpdateCollisions, STATGROUP_Physics, CLOTHINGSYSTEMRUNTIMECOMMON_API);
-DECLARE_CYCLE_STAT_EXTERN(TEXT("Skin Physics Mesh"), STAT_ClothSkinPhysMesh, STATGROUP_Physics, CLOTHINGSYSTEMRUNTIMECOMMON_API);
-DECLARE_CYCLE_STAT_EXTERN(TEXT("Fill Context"), STAT_ClothFillContext, STATGROUP_Physics, CLOTHINGSYSTEMRUNTIMECOMMON_API);
+DECLARE_CYCLE_STAT_EXTERN(TEXT("Compute Clothing Normals"), STAT_ClothComputeNormals, STATGROUP_Physics, UE_API);
+DECLARE_CYCLE_STAT_EXTERN(TEXT("Internal Solve"), STAT_ClothInternalSolve, STATGROUP_Physics, UE_API);
+DECLARE_CYCLE_STAT_EXTERN(TEXT("Update Collisions"), STAT_ClothUpdateCollisions, STATGROUP_Physics, UE_API);
+DECLARE_CYCLE_STAT_EXTERN(TEXT("Skin Physics Mesh"), STAT_ClothSkinPhysMesh, STATGROUP_Physics, UE_API);
+DECLARE_CYCLE_STAT_EXTERN(TEXT("Fill Context"), STAT_ClothFillContext, STATGROUP_Physics, UE_API);
 
 /** Base simulation data that just about every simulation would need. */
 class FClothingSimulationContextCommon : public IClothingSimulationContext
 {
 public:
-	CLOTHINGSYSTEMRUNTIMECOMMON_API FClothingSimulationContextCommon();
+	UE_API FClothingSimulationContextCommon();
 	FClothingSimulationContextCommon(const FClothingSimulationContextCommon&) = default;
 	FClothingSimulationContextCommon& operator=(const FClothingSimulationContextCommon&) = default;
 
-	CLOTHINGSYSTEMRUNTIMECOMMON_API virtual ~FClothingSimulationContextCommon() override;
+	UE_API virtual ~FClothingSimulationContextCommon() override;
 
 	// Fill this context using the given skeletal mesh component
-	CLOTHINGSYSTEMRUNTIMECOMMON_API virtual void Fill(const USkeletalMeshComponent* InComponent, float InDeltaSeconds, float InMaxPhysicsDelta, bool bIsInitialization);
+	UE_API virtual void Fill(const USkeletalMeshComponent* InComponent, float InDeltaSeconds, float InMaxPhysicsDelta, bool bIsInitialization);
 
 protected:
 	// Default fill behavior as expected to be used by every simulation
-	CLOTHINGSYSTEMRUNTIMECOMMON_API virtual void FillBoneTransforms(const USkeletalMeshComponent* InComponent);
-	CLOTHINGSYSTEMRUNTIMECOMMON_API virtual void FillRefToLocals(const USkeletalMeshComponent* InComponent, bool bIsInitialization);
-	CLOTHINGSYSTEMRUNTIMECOMMON_API virtual void FillComponentToWorld(const USkeletalMeshComponent* InComponent);
-	CLOTHINGSYSTEMRUNTIMECOMMON_API virtual void FillWorldGravity(const USkeletalMeshComponent* InComponent);
-	CLOTHINGSYSTEMRUNTIMECOMMON_API virtual void FillWindVelocity(const USkeletalMeshComponent* InComponent);
-	CLOTHINGSYSTEMRUNTIMECOMMON_API virtual void FillDeltaSeconds(float InDeltaSeconds, float InMaxPhysicsDelta);
-	CLOTHINGSYSTEMRUNTIMECOMMON_API virtual void FillTeleportMode(const USkeletalMeshComponent* InComponent, float InDeltaSeconds, float InMaxPhysicsDelta);
-	CLOTHINGSYSTEMRUNTIMECOMMON_API virtual void FillMaxDistanceScale(const USkeletalMeshComponent* InComponent);
-	CLOTHINGSYSTEMRUNTIMECOMMON_API virtual void FillSolverGeometryScale(const USkeletalMeshComponent* InComponent);
+	UE_API virtual void FillBoneTransforms(const USkeletalMeshComponent* InComponent);
+	UE_API virtual void FillRefToLocals(const USkeletalMeshComponent* InComponent, bool bIsInitialization);
+	UE_API virtual void FillComponentToWorld(const USkeletalMeshComponent* InComponent);
+	UE_API virtual void FillWorldGravity(const USkeletalMeshComponent* InComponent);
+	UE_API virtual void FillWindVelocity(const USkeletalMeshComponent* InComponent);
+	UE_API virtual void FillDeltaSeconds(float InDeltaSeconds, float InMaxPhysicsDelta);
+	UE_API virtual void FillTeleportMode(const USkeletalMeshComponent* InComponent, float InDeltaSeconds, float InMaxPhysicsDelta);
+	UE_API virtual void FillMaxDistanceScale(const USkeletalMeshComponent* InComponent);
+	UE_API virtual void FillSolverGeometryScale(const USkeletalMeshComponent* InComponent);
 
 public:
 	// Component space bone transforms of the owning component
@@ -88,17 +89,16 @@ public:
 class FClothingSimulationCommon : public IClothingSimulation
 {
 public:
-	CLOTHINGSYSTEMRUNTIMECOMMON_API FClothingSimulationCommon();
-	CLOTHINGSYSTEMRUNTIMECOMMON_API virtual ~FClothingSimulationCommon();
+	UE_API FClothingSimulationCommon();
+	UE_API virtual ~FClothingSimulationCommon();
 
 protected:
 	/** Fills in the base data for a clothing simulation */
-	CLOTHINGSYSTEMRUNTIMECOMMON_API virtual void FillContext(USkeletalMeshComponent* InComponent, float InDeltaTime, IClothingSimulationContext* InOutContext) override;
-
-	/** Fills in the base data for a clothing simulation */
-	CLOTHINGSYSTEMRUNTIMECOMMON_API virtual void FillContext(USkeletalMeshComponent* InComponent, float InDeltaTime, IClothingSimulationContext* InOutContext, bool bIsInitialization) override;
+	UE_API virtual void FillContext(const USkeletalMeshComponent* InComponent, float InDeltaTime, IClothingSimulationContext* InOutContext, bool bIsInitialization) const override;
 
 protected:
 	/** Maximum physics time, incoming deltas will be clamped down to this value on long frames */
 	float MaxPhysicsDelta;
 };
+
+#undef UE_API
diff --git a/Engine/Source/Runtime/ClothingSystemRuntimeInterface/Private/ClothingSimulationFactory.cpp b/Engine/Source/Runtime/ClothingSystemRuntimeInterface/Private/ClothingSimulationFactory.cpp
--- a/Engine/Source/Runtime/ClothingSystemRuntimeInterface/Private/ClothingSimulationFactory.cpp
+++ b/Engine/Source/Runtime/ClothingSystemRuntimeInterface/Private/ClothingSimulationFactory.cpp
@@ -52,3 +52,25 @@ TSubclassOf<class UClothingSimulationFactory> UClothingSimulationFactory::GetDef
 	return DefaultClothingSimulationFactoryClass;
 }
 
+UClothingSimulationFactory* UClothingSimulationFactory::GetClothingSimulationFactory(const UClothingAssetBase* InAsset)
+{
+	IModularFeatures::Get().LockModularFeatureList();
+	const TArray<IClothingSimulationFactoryClassProvider*> ClassProviders = IModularFeatures::Get().GetModularFeatureImplementations<IClothingSimulationFactoryClassProvider>(IClothingSimulationFactoryClassProvider::FeatureName);
+	IModularFeatures::Get().UnlockModularFeatureList();
+
+	for (const IClothingSimulationFactoryClassProvider* const ClassProvider : ClassProviders)
+	{
+		if (ClassProvider)
+		{
+			const TSubclassOf<UClothingSimulationFactory> ClothingSimulationFactoryClass = ClassProvider->GetClothingSimulationFactoryClass();
+			UClothingSimulationFactory* const ClothingSimulationFactory = ClothingSimulationFactoryClass.GetDefaultObject();
+			if (ClothingSimulationFactory && ClothingSimulationFactory->SupportsAsset(InAsset))
+			{
+				return ClothingSimulationFactory;
+			}
+		}
+	}
+
+	return nullptr;
+}
+
diff --git a/Engine/Source/Runtime/ClothingSystemRuntimeInterface/Private/ClothingSimulationInstance.cpp b/Engine/Source/Runtime/ClothingSystemRuntimeInterface/Private/ClothingSimulationInstance.cpp
new file mode 100644
--- /dev/null
+++ b/Engine/Source/Runtime/ClothingSystemRuntimeInterface/Private/ClothingSimulationInstance.cpp
@@ -0,0 +1,84 @@
+// Copyright Epic Games, Inc. All Rights Reserved.
+
+#include "ClothingSimulationInstance.h"
+#include "ClothCollisionData.h"
+#include "ClothingSimulationInteractor.h"
+#include "ClothingSimulationFactory.h"
+
+#include UE_INLINE_GENERATED_CPP_BY_NAME(ClothingSimulationInstance)
+
+FClothingSimulationInstance::FClothingSimulationInstance(USkeletalMeshComponent* SkeletalMeshComponent, const FFactoryAssetGroup& FactoryAssetGroup)
+	: ClothingSimulation(FactoryAssetGroup.ClothingSimulationFactory->CreateSimulation())
+	, ClothingSimulationContext(ClothingSimulation->CreateContext())
+	, ClothingSimulationInteractor(static_cast<const UClothingSimulationFactory*>(FactoryAssetGroup.ClothingSimulationFactory)->SupportsRuntimeInteraction() ? FactoryAssetGroup.ClothingSimulationFactory->CreateInteractor() : nullptr)
+	, ClothingSimulationFactory(FactoryAssetGroup.ClothingSimulationFactory)
+{
+	check(ClothingSimulation);
+	ClothingSimulation->Initialize();
+
+	const bool bIsInitialization = true;
+	FillContext(SkeletalMeshComponent, 0.f, bIsInitialization);
+
+	for (int32 ClothingAssetIndex = 0; ClothingAssetIndex < FactoryAssetGroup.ClothingAssets.Num(); ++ClothingAssetIndex)
+	{
+		if (const UClothingAssetBase* const ClothingAsset = FactoryAssetGroup.ClothingAssets[ClothingAssetIndex])
+		{
+			ClothingSimulation->CreateActor(SkeletalMeshComponent, ClothingAsset, ClothingAssetIndex);
+
+			if (ClothingSimulationInteractor)
+			{
+				ClothingSimulationInteractor->CreateClothingInteractor(ClothingAsset, ClothingAssetIndex);
+			}
+		}
+	}
+	ClothingSimulation->EndCreateActor();
+}
+
+FClothingSimulationInstance::FClothingSimulationInstance() = default;
+
+FClothingSimulationInstance::~FClothingSimulationInstance()
+{
+	if (ClothingSimulation)
+	{
+		if (ClothingSimulationContext)
+		{
+			ClothingSimulation->DestroyContext(ClothingSimulationContext);
+		}
+		ClothingSimulation->DestroyActors();
+		ClothingSimulation->Shutdown();
+		if (ClothingSimulationFactory)
+		{
+			ClothingSimulationFactory->DestroySimulation(ClothingSimulation);
+		}
+	}
+}
+
+void FClothingSimulationInstance::RemoveAllClothingActors()
+{
+	if (ClothingSimulationInteractor)
+	{
+		ClothingSimulationInteractor->DestroyClothingInteractors();
+	}
+	ClothingSimulation->DestroyActors();
+}
+
+void FClothingSimulationInstance::SyncClothingInteractor(bool bNeedsHardReset)
+{
+	if (ClothingSimulationInteractor)
+	{
+		if (bNeedsHardReset)
+		{
+			ClothingSimulationInteractor->ClothConfigUpdated();
+		}
+
+		ClothingSimulationInteractor->Sync(ClothingSimulation, ClothingSimulationContext);
+	}
+}
+
+void FClothingSimulationInstance::AppendUniqueCollisions(FClothCollisionData& OutCollisions, bool bIncludeExternal) const
+{
+	FClothCollisionData CollisionsToAppend;
+	ClothingSimulation->GetCollisions(CollisionsToAppend, bIncludeExternal);
+
+	OutCollisions.AppendUnique(CollisionsToAppend);
+}
diff --git a/Engine/Source/Runtime/ClothingSystemRuntimeInterface/Public/ClothingAssetBase.h b/Engine/Source/Runtime/ClothingSystemRuntimeInterface/Public/ClothingAssetBase.h
--- a/Engine/Source/Runtime/ClothingSystemRuntimeInterface/Public/ClothingAssetBase.h
+++ b/Engine/Source/Runtime/ClothingSystemRuntimeInterface/Public/ClothingAssetBase.h
@@ -89,6 +89,12 @@ public:
 	virtual void RefreshBoneMapping(USkeletalMesh* InSkelMesh)
 	PURE_VIRTUAL(UClothingAssetBase::RefreshBoneMapping, );
 
+	/** Check the validity of this clothing data, for example for when it is dependent from an asset that isn't ready or accessible. */
+	virtual bool IsValid() const
+	{
+		return true;
+	}
+
 	/** Check the validity of a LOD index */
 	UE_DEPRECATED(5.7, "Use UClothingAssetCommon::IsValidLod() instead.")
 	virtual bool IsValidLod(int32 InLodIndex) const
diff --git a/Engine/Source/Runtime/ClothingSystemRuntimeInterface/Public/ClothingSimulationFactory.h b/Engine/Source/Runtime/ClothingSystemRuntimeInterface/Public/ClothingSimulationFactory.h
--- a/Engine/Source/Runtime/ClothingSystemRuntimeInterface/Public/ClothingSimulationFactory.h
+++ b/Engine/Source/Runtime/ClothingSystemRuntimeInterface/Public/ClothingSimulationFactory.h
@@ -14,6 +14,8 @@
 
 #include "ClothingSimulationFactory.generated.h"
 
+#define UE_API CLOTHINGSYSTEMRUNTIMEINTERFACE_API
+
 class FName;
 class IClothingSimulation;
 class UClass;
@@ -29,7 +31,7 @@ class IClothingSimulationFactoryClassProvider : public IModularFeature
 public:
 
 	// The feature name to register against for providers
-	static CLOTHINGSYSTEMRUNTIMEINTERFACE_API const FName FeatureName;
+	static UE_API const FName FeatureName;
 
 	// Called by the engine to get the clothing simulation factory associated with this
 	// provider for skeletal mesh components (see USkeletalMeshComponent constructor).
@@ -47,7 +49,12 @@ public:
 	// Return the default clothing simulation factory class as set by the build or by
 	// the p.Cloth.DefaultClothingSimulationFactoryClass console variable if any available.
 	// Otherwise return the last registered factory.
-	static CLOTHINGSYSTEMRUNTIMEINTERFACE_API TSubclassOf<class UClothingSimulationFactory> GetDefaultClothingSimulationFactoryClass();
+	static UE_API TSubclassOf<class UClothingSimulationFactory> GetDefaultClothingSimulationFactoryClass();
+
+	// Return the clothing simulation factory class for.the specified asset.
+	// Calls SupportsAsset() internally.
+	// If multiple simulation classes support this asset then this returns the first one registered through IClothingSimulationFactoryClassProvider.
+	static UE_API UClothingSimulationFactory* GetClothingSimulationFactory(const UClothingAssetBase* InAsset);
 
 	// Create a simulation object for a skeletal mesh to use (see IClothingSimulation)
 	virtual IClothingSimulation* CreateSimulation()
@@ -59,14 +66,26 @@ public:
 
 	// Given an asset, decide whether this factory can create a simulation to use the data inside
 	// (return false if data is invalid or missing in the case of custom data)
-	virtual bool SupportsAsset(UClothingAssetBase* InAsset)
+	virtual bool SupportsAsset(const UClothingAssetBase* InAsset) const
 	PURE_VIRTUAL(UClothingSimulationFactory::SupportsAsset, return false;);
 
+	UE_DEPRECATED(5.7, "Use const version instead.")
+	virtual bool SupportsAsset(UClothingAssetBase* InAsset)
+	{
+		return static_cast<const UClothingSimulationFactory*>(this)->SupportsAsset(InAsset);
+	}
+
 	// Whether or not we provide an interactor object to manipulate the simulation at runtime.
 	// If true is returned then CreateInteractor *must* create a valid object to handle this
-	virtual bool SupportsRuntimeInteraction()
+	virtual bool SupportsRuntimeInteraction() const
 	PURE_VIRTUAL(UClothingSimulationFactory::SupportsRuntimeInteraction, return false;);
 
+	UE_DEPRECATED(5.7, "Use const version instead.")
+	virtual bool SupportsRuntimeInteraction()
+	{
+		return static_cast<const UClothingSimulationFactory*>(this)->SupportsRuntimeInteraction();
+	}
+
 	// Creates the runtime interactor object for a clothing simulation. This object will
 	// receive events allowing it to write data to the simulation context in a safe manner
 	virtual UClothingSimulationInteractor* CreateInteractor()
@@ -80,3 +99,5 @@ public:
 	virtual const UEnum* GetWeightMapTargetEnum() const
 	PURE_VIRTUAL(UClothingSimulationFactory::GetWeightMapTargetEnum, return nullptr;);
 };
+
+#undef UE_API
diff --git a/Engine/Source/Runtime/ClothingSystemRuntimeInterface/Public/ClothingSimulationInstance.h b/Engine/Source/Runtime/ClothingSystemRuntimeInterface/Public/ClothingSimulationInstance.h
new file mode 100644
--- /dev/null
+++ b/Engine/Source/Runtime/ClothingSystemRuntimeInterface/Public/ClothingSimulationInstance.h
@@ -0,0 +1,200 @@
+// Copyright Epic Games, Inc. All Rights Reserved.
+#pragma once
+
+#include "ClothingSimulationInterface.h"
+#include "UObject/ObjectMacros.h"
+#include "UObject/ObjectPtr.h"
+#include "ClothingSimulationInstance.generated.h"
+
+#define UE_API CLOTHINGSYSTEMRUNTIMEINTERFACE_API
+
+struct FClothSimulData;
+class UClothingSimulationFactory;
+class UClothingSimulationInteractor;
+class USkeletalMesh;
+class USkeletalMeshComponent;
+class USkinnedMeshComponent;
+
+/**
+ * Clothing simulation objects.
+ * Facade to the simulation interfaces to allow for a simplified implementation and multiple simulations to coexist with each other.
+ */
+USTRUCT()
+struct FClothingSimulationInstance
+{
+public:
+	GENERATED_BODY()
+
+	/**
+	 * Factory asset group.
+	 * Group of clothing data that run under the same simulation type, and the clothing simulation factory associated with them.
+	 * The group storage array always matches the order of the SkeletalMesh clothing data so that the indices are the same,
+	 * with a nullptr entry into the array to signify that the asset is part of a different group.
+	 * This is to avoid an additional array of indices, knowing that most Skeletal Meshes have rarely more than half a dozen clothing data.
+	 */
+	struct FFactoryAssetGroup
+	{
+	public:
+		FFactoryAssetGroup(UClothingSimulationFactory* InClothingSimulationFactory, const int32 NumTotalClothAssets)
+			: ClothingSimulationFactory(InClothingSimulationFactory)
+		{
+			check(ClothingSimulationFactory);
+			ClothingAssets.SetNumZeroed(NumTotalClothAssets);
+		}
+
+		void AddClothingAsset(const UClothingAssetBase* ClothingAsset, const int32 Index)
+		{
+			ClothingAssets[Index] = ClothingAsset;
+		}
+
+		const UClothingSimulationFactory* GetClothingSimulationFactory() const
+		{
+			return ClothingSimulationFactory;
+		}
+
+	private:
+		friend FClothingSimulationInstance;
+		UClothingSimulationFactory* ClothingSimulationFactory;
+		TArray<const UClothingAssetBase*> ClothingAssets;
+	};
+
+	/** Construct a new clothing simulation instance for the specified asset group and component. */
+	UE_API FClothingSimulationInstance(USkeletalMeshComponent* SkeletalMeshComponent, const FFactoryAssetGroup& FactoryAssetGroup);
+	
+	UE_API FClothingSimulationInstance();
+	UE_API ~FClothingSimulationInstance();
+
+	FClothingSimulationInstance(const FClothingSimulationInstance&) = delete;
+	FClothingSimulationInstance(FClothingSimulationInstance&&) = default;
+	FClothingSimulationInstance& operator=(const FClothingSimulationInstance&) = delete;
+	FClothingSimulationInstance& operator=(FClothingSimulationInstance&&) = default;
+
+	/** Return the clothing simulation factory used to initialize this clothing simulation instance. */
+	UClothingSimulationFactory* GetClothingSimulationFactory() const
+	{
+		return ClothingSimulationFactory;
+	}
+
+	/** Return the clothing simulation pointer. */
+	const IClothingSimulation* GetClothingSimulation() const
+	{
+		return ClothingSimulation;
+	}
+
+	/** Return the clothing simulation pointer. */
+	IClothingSimulation* GetClothingSimulation()
+	{
+		return ClothingSimulation;
+	}
+
+	/** Return the clothing simulation context. */
+	const IClothingSimulationContext* GetClothingSimulationContext() const
+	{
+		return ClothingSimulationContext;
+	}
+
+	/** Return the clothing simulation context. */
+	IClothingSimulationContext* GetClothingSimulationContext()
+	{
+		return ClothingSimulationContext;
+	}
+
+	/** Return a non const version of the clothing interactor (non const methods are used to update simulation parameters). */
+	UClothingSimulationInteractor* GetClothingSimulationInteractor() const
+	{
+		return ClothingSimulationInteractor;
+	}
+
+	/** Remove all actors. */
+	UE_API void RemoveAllClothingActors();
+
+	/** Synchronize the simulation with the changes requested by the interactor. */
+	UE_API void SyncClothingInteractor(bool bNeedsHardReset);
+
+	/** Return the collision data from this simulation instance. */
+	void GetCollisions(FClothCollisionData& Collisions, bool bIncludeExternal) const
+	{
+		ClothingSimulation->GetCollisions(Collisions, bIncludeExternal);
+	}
+
+	/** Append unique collision data from this simulation instance to the specified collision data. */
+	UE_API void AppendUniqueCollisions(FClothCollisionData& OutCollisions, bool bIncludeExternal) const;
+
+	/** Add external dynamic collisions. */
+	void AddExternalCollisions(const FClothCollisionData& Collisions)
+	{
+		ClothingSimulation->AddExternalCollisions(Collisions);
+	}
+
+	/** Remove all dynamically added external collisions. */
+	void ClearExternalCollisions()
+	{
+		ClothingSimulation->ClearExternalCollisions();
+	}
+
+	/** Update external forces. */
+	void UpdateWorldForces(const USkeletalMeshComponent* OwnerComponent)
+	{
+		ClothingSimulation->UpdateWorldForces(OwnerComponent);
+	}
+
+	/** Fill the simulation context. */
+	void FillContext(const USkeletalMeshComponent* OwnerComponent, float DeltaTime, bool bIsInitialization)
+	{
+		ClothingSimulation->FillContext(OwnerComponent, DeltaTime, ClothingSimulationContext, bIsInitialization);
+	}
+
+	/** Advance the simulation. */
+	void Simulate()
+	{
+		ClothingSimulation->Simulate(static_cast<const IClothingSimulationContext*>(ClothingSimulationContext));
+	}
+
+	/** Return the simulation data. */
+	void GetSimulationData(TMap<int32, FClothSimulData>& CurrentSimulationData, const USkeletalMeshComponent* OwnerComponent, const USkinnedMeshComponent* LeaderPoseComponent) const
+	{
+		ClothingSimulation->GetSimulationData(CurrentSimulationData, OwnerComponent, LeaderPoseComponent);
+	}
+
+	/** Return whether this simulation has anything to simulate. */
+	bool ShouldSimulate() const
+	{
+		return ClothingSimulation->ShouldSimulate();
+	}
+
+	/** Return this simulation bounds. */
+	FBoxSphereBounds GetBounds(const USkeletalMeshComponent* OwnerComponent) const
+	{
+		return ClothingSimulation->GetBounds(OwnerComponent);
+	}
+
+private:
+	/** ClothingSimulations are responsible for maintaining and simulating clothing actors. */
+	IClothingSimulation* ClothingSimulation = nullptr;
+
+	/** ClothingSimulationContexts are a datastore for simulation data sent to the clothing thread. */
+	IClothingSimulationContext* ClothingSimulationContext = nullptr;
+
+	/**
+	 * Objects responsible for interacting with the clothing simulation.
+	 * Blueprints and code can call/set data on this from the game thread and the next time
+	 * it is safe to do so the interactor will sync to the simulation context.
+	 */
+	UPROPERTY(Transient)
+	TObjectPtr<UClothingSimulationInteractor> ClothingSimulationInteractor;
+
+	/** Simulation factory. */
+	UPROPERTY(Transient)
+	TObjectPtr<UClothingSimulationFactory> ClothingSimulationFactory;
+};
+
+template<>
+struct TStructOpsTypeTraits<FClothingSimulationInstance> : public TStructOpsTypeTraitsBase2<FClothingSimulationInstance>
+{
+	enum
+	{
+		WithCopy = false
+	};
+};
+
+#undef UE_API
diff --git a/Engine/Source/Runtime/ClothingSystemRuntimeInterface/Public/ClothingSimulationInterface.h b/Engine/Source/Runtime/ClothingSystemRuntimeInterface/Public/ClothingSimulationInterface.h
--- a/Engine/Source/Runtime/ClothingSystemRuntimeInterface/Public/ClothingSimulationInterface.h
+++ b/Engine/Source/Runtime/ClothingSystemRuntimeInterface/Public/ClothingSimulationInterface.h
@@ -40,6 +40,7 @@ protected:
 
 	friend class USkeletalMeshComponent;
 	friend class FParallelClothTask;
+	friend struct FClothingSimulationInstance;
 	friend struct FEnvironmentalCollisions;
 
 	/**
@@ -49,25 +50,38 @@ protected:
 	 * @param InAsset - The asset to create an actor from
 	 * @param SimDataIndex - the sim data index to use when doing the writeback in GetSimulationData
 	 */
-	virtual void CreateActor(USkeletalMeshComponent* InOwnerComponent, UClothingAssetBase* InAsset, int32 SimDataIndex) = 0;
+	virtual void CreateActor(USkeletalMeshComponent* InOwnerComponent, const UClothingAssetBase* InAsset, int32 SimDataIndex) {}
+
+	UE_DEPRECATED(5.7, "Use const arguments version instead,")
+	virtual void CreateActor(USkeletalMeshComponent* InOwnerComponent, UClothingAssetBase* InAsset, int32 SimDataIndex)
+	{
+		return CreateActor(InOwnerComponent, static_cast<const UClothingAssetBase*>(InAsset), SimDataIndex);
+	}
+
+	/** Called once all CreateActor have been called. */
+	virtual void EndCreateActor() {}
 
 	/** Create a new context, will not be filled, call FillContext before simulating with this context */
 	virtual IClothingSimulationContext* CreateContext() = 0;
 
-	/**
-	 * Fills an existing context for a single simulation step, called by the engine on the game thread prior to simulation 
-	 * @param InComponent - The component to fill the context for
-	 * @param InOutContext - The context to fill
-	 */
-	virtual void FillContext(USkeletalMeshComponent* InComponent, float InDeltaTime, IClothingSimulationContext* InOutContext) = 0;
-
 	/**
 	 * Fills an existing context for a single simulation step, called by the engine on the game thread prior to simulation 
 	 * @param InComponent - The component to fill the context for
 	 * @param InOutContext - The context to fill
 	 * @param bIsInitialization - Whether this fill is occurring as part of the actor creation stage
 	 */
-	virtual void FillContext(USkeletalMeshComponent* InComponent, float InDeltaTime, IClothingSimulationContext* InOutContext, bool bIsInitialization) = 0;
+	virtual void FillContext(const USkeletalMeshComponent* InComponent, float InDeltaTime, IClothingSimulationContext* InOutContext, bool bIsInitialization) const {}
+
+	UE_DEPRECATED(5.7, "Use const version instead.")
+	virtual void FillContext(USkeletalMeshComponent* InComponent, float InDeltaTime, IClothingSimulationContext* InOutContext, bool bIsInitialization)
+	{
+		return static_cast<const IClothingSimulation*>(this)->FillContext(InComponent, InDeltaTime, InOutContext, bIsInitialization);
+	}
+	UE_DEPRECATED(5.7, "Use const version instead.")
+	virtual void FillContext(USkeletalMeshComponent* InComponent, float InDeltaTime, IClothingSimulationContext* InOutContext)
+	{
+		return static_cast<const IClothingSimulation*>(this)->FillContext(InComponent, InDeltaTime, InOutContext, false);
+	}
 
 	/** Initialize the simulation, will be called before any Simulate calls */
 	virtual void Initialize() = 0;
@@ -91,7 +105,13 @@ protected:
 	 * New callers should make care to make sure only the correct context is ever passed through
 	 * @param InContext - The context to use during simulation, will have been filled in FillContext
 	 */
-	virtual void Simulate(IClothingSimulationContext* InContext) = 0;
+	virtual void Simulate(const IClothingSimulationContext* InContext) {}
+	
+	UE_DEPRECATED(5.7, "Use const argument version instead,")
+	virtual void Simulate(IClothingSimulationContext* InContext)
+	{
+		return Simulate(static_cast<const IClothingSimulationContext*>(InContext));
+	}
 
 	/** Simulation should remove all of it's actors when next possible and free them */
 	virtual void DestroyActors() = 0;
@@ -105,10 +125,16 @@ protected:
 	/** 
 	 * Fill FClothSimulData map for the clothing simulation. Should fill a map pair per-actor 
 	 * @param OutData - The simulation data to write to
-	 * @param InOwnerComponent - the component that owns the simulation
-	 * @param InOverrideComponent - An override component if bound to a leader pose component
+	 * @param OwnerComponent - the component that owns the simulation
+	 * @param OverrideComponent - An override component if bound to a leader pose component
 	 */
-	virtual void GetSimulationData(TMap<int32, FClothSimulData>& OutData, USkeletalMeshComponent* InOwnerComponent, USkinnedMeshComponent* InOverrideComponent) const = 0;
+	virtual void GetSimulationData(TMap<int32, FClothSimulData>& OutData, const USkeletalMeshComponent* OwnerComponent, const USkinnedMeshComponent* OverrideComponent) const {}
+
+	UE_DEPRECATED(5.7, "Use const version instead.")
+	virtual void GetSimulationData(TMap<int32, FClothSimulData>& OutData, USkeletalMeshComponent* InOwnerComponent, USkinnedMeshComponent* InOverrideComponent) const
+	{
+		return GetSimulationData(OutData, static_cast<const USkeletalMeshComponent*>(InOwnerComponent), static_cast<const USkinnedMeshComponent*>(InOverrideComponent));
+	}
 
 	/** Get the bounds of the simulation mesh in local simulation space */
 	virtual FBoxSphereBounds GetBounds(const USkeletalMeshComponent* InOwnerComponent) const = 0;
diff --git a/Engine/Source/Runtime/ClothingSystemRuntimeNv/Private/ClothingSimulationFactoryNv.cpp b/Engine/Source/Runtime/ClothingSystemRuntimeNv/Private/ClothingSimulationFactoryNv.cpp
--- a/Engine/Source/Runtime/ClothingSystemRuntimeNv/Private/ClothingSimulationFactoryNv.cpp
+++ b/Engine/Source/Runtime/ClothingSystemRuntimeNv/Private/ClothingSimulationFactoryNv.cpp
@@ -16,12 +16,12 @@ void UClothingSimulationFactoryNv::DestroySimulation(IClothingSimulation* InSimu
 {
 }
 
-bool UClothingSimulationFactoryNv::SupportsAsset(UClothingAssetBase* InAsset)
+bool UClothingSimulationFactoryNv::SupportsAsset(const UClothingAssetBase* InAsset) const
 {
 	return false;
 }
 
-bool UClothingSimulationFactoryNv::SupportsRuntimeInteraction()
+bool UClothingSimulationFactoryNv::SupportsRuntimeInteraction() const
 {
 	return true;
 }
diff --git a/Engine/Source/Runtime/ClothingSystemRuntimeNv/Public/ClothingSimulationFactoryNv.h b/Engine/Source/Runtime/ClothingSystemRuntimeNv/Public/ClothingSimulationFactoryNv.h
--- a/Engine/Source/Runtime/ClothingSystemRuntimeNv/Public/ClothingSimulationFactoryNv.h
+++ b/Engine/Source/Runtime/ClothingSystemRuntimeNv/Public/ClothingSimulationFactoryNv.h
@@ -28,9 +28,9 @@ public:
 
 	CLOTHINGSYSTEMRUNTIMENV_API virtual IClothingSimulation* CreateSimulation() override;
 	CLOTHINGSYSTEMRUNTIMENV_API virtual void DestroySimulation(IClothingSimulation* InSimulation) override;
-	CLOTHINGSYSTEMRUNTIMENV_API virtual bool SupportsAsset(UClothingAssetBase* InAsset) override;
+	CLOTHINGSYSTEMRUNTIMENV_API virtual bool SupportsAsset(const UClothingAssetBase* InAsset) const override;
 
-	CLOTHINGSYSTEMRUNTIMENV_API virtual bool SupportsRuntimeInteraction() override;
+	CLOTHINGSYSTEMRUNTIMENV_API virtual bool SupportsRuntimeInteraction() const override;
 	CLOTHINGSYSTEMRUNTIMENV_API virtual UClothingSimulationInteractor* CreateInteractor() override;
 
 	CLOTHINGSYSTEMRUNTIMENV_API virtual TArrayView<const TSubclassOf<UClothConfigBase>> GetClothConfigClasses() const override;
diff --git a/Engine/Source/Runtime/Engine/Classes/Components/SkeletalMeshComponent.h b/Engine/Source/Runtime/Engine/Classes/Components/SkeletalMeshComponent.h
--- a/Engine/Source/Runtime/Engine/Classes/Components/SkeletalMeshComponent.h
+++ b/Engine/Source/Runtime/Engine/Classes/Components/SkeletalMeshComponent.h
@@ -17,13 +17,15 @@
 #include "SingleAnimationPlayData.h"
 #include "Animation/PoseSnapshot.h"
 #include "ClothingSystemRuntimeTypes.h"
-#include "ClothingSimulationInterface.h"
-#include "ClothingSimulationFactory.h"
 #include "Animation/AttributesRuntime.h"
 #include "Chaos/Serialization/SerializedDataBuffer.h"
 #if WITH_ENGINE
 #include "Engine/PoseWatchRenderData.h"
 #endif
+#if UE_ENABLE_INCLUDE_ORDER_DEPRECATED_IN_5_7
+#include "ClothingSimulationInterface.h"
+#include "ClothingSimulationFactory.h"
+#endif
 
 #include "SkeletalMeshComponent.generated.h"
 
@@ -36,7 +38,9 @@ class UAnimInstance;
 class UPhysicalMaterial;
 class USkeletalMesh;
 class USkeletalMeshComponent;
+struct FClothCollisionData;
 struct FClothCollisionSource;
+struct FClothingSimulationInstance;
 struct FConstraintInstance;
 struct FConstraintProfileProperties;
 struct FNavigableGeometryExport;
@@ -948,7 +952,8 @@ public:
 	FPlasticDeformationEventSignature OnPlasticDeformation;
 
 	/** Class of the object responsible for  */
-	UPROPERTY(EditAnywhere, Category = Clothing)
+	UE_DEPRECATED(5.7, "This property will be removed since the ClothingSimulationFactory is no longer set per component.")
+	UPROPERTY()
 	TSubclassOf<class UClothingSimulationFactory> ClothingSimulationFactory;
 
 	struct FPendingRadialForces
@@ -1539,20 +1544,39 @@ public:
 	bool IsClothBoundToMasterComponent() const { return IsClothBoundToLeaderComponent(); }
 
 	/** Get the current clothing simulation (read only) */
-	ENGINE_API const IClothingSimulation* GetClothingSimulation() const;
-	
+	UE_DEPRECATED(5.7, "Use GetClothingSimulationInstances instead.")
+	ENGINE_API const class IClothingSimulation* GetClothingSimulation() const;
+
 	/** Get the current clothing simulation (read/write) */
-	ENGINE_API IClothingSimulation* GetClothingSimulation();
+	UE_DEPRECATED(5.7, "Use GetClothingSimulationInstances instead.")
+	ENGINE_API class IClothingSimulation* GetClothingSimulation();
 
 	/** Get the current clothing simulation context (read only) */
-	ENGINE_API const IClothingSimulationContext* GetClothingSimulationContext() const;
+	UE_DEPRECATED(5.7, "Use GetClothingSimulationInstances instead.")
+	ENGINE_API const class IClothingSimulationContext* GetClothingSimulationContext() const;
 	
 	/** Get the current clothing simulation context (read/write) */
-	ENGINE_API IClothingSimulationContext* GetClothingSimulationContext();
+	ENGINE_API class IClothingSimulationContext* GetClothingSimulationContext();
 	
-	/** Get the current interactor for a clothing simulation, if the current simulation supports runtime interaction. */
-	UFUNCTION(BlueprintCallable, Category=Clothing)
-	ENGINE_API UClothingSimulationInteractor* GetClothingSimulationInteractor() const;
+	/** Get the current clothing simulation instances (read only) */
+	const TArray<FClothingSimulationInstance>& GetClothingSimulationInstances() const
+	{
+		return ClothingSimulationInstances;
+	}
+
+	/** Get the current clothing simulation instances (read/write) */
+	TArray<FClothingSimulationInstance>& GetClothingSimulationInstances()
+	{
+		return ClothingSimulationInstances;
+	}
+
+	/**
+	 * Get the interactor for a clothing simulation, if that simulation supports runtime interaction.
+	 * @param ClothingDataName the name of a clothing data that determines the simulation, or NAME_None to get the interactor of the first registered simulation if any.
+	 * @note Using NAME_None when there are clothing data of different types on the Skeletal Mesh asset will cause a warning.
+	 */
+	UFUNCTION(BlueprintCallable, Category = Clothing)
+	ENGINE_API UClothingSimulationInteractor* GetClothingSimulationInteractor(const FName ClothingDataName = NAME_None) const;
 
 	/** Callback when the parallel clothing task finishes, copies needed data back to component for gamethread */
 	ENGINE_API void CompleteParallelClothSimulation();
@@ -1584,21 +1608,9 @@ private:
 	/** previous root body scale to compare the difference when the actor scale is updated. */
 	FVector3f PrevRootRigidBodyScale;
 
-	/** 
-	 * Clothing simulation objects.
-	 * ClothingSimulation is responsible for maintaining and siulating clothing actors
-	 * ClothingSimulationContext is a datastore for simulation data sent to the clothing thread
-	 */
-	IClothingSimulation* ClothingSimulation;
-	IClothingSimulationContext* ClothingSimulationContext;
-
-	/** 
-	 * Object responsible for interacting with the clothing simulation.
-	 * Blueprints and code can call/set data on this from the game thread and the next time
-	 * it is safe to do so the interactor will sync to the simulation context
-	 */
+	/** Clothing simulation instance. */
 	UPROPERTY(Transient)
-	TObjectPtr<UClothingSimulationInteractor> ClothingInteractor;
+	TArray<FClothingSimulationInstance> ClothingSimulationInstances;
 
 	/** Array of sources for cloth collision */
 	TArray<FClothCollisionSource> ClothCollisionSources;
@@ -1621,9 +1633,6 @@ private:
 	 */
 	ENGINE_API void WritebackClothingSimulationData();
 
-	/** Gets the factory responsible for building the clothing simulation and simulation contexts */
-	ENGINE_API UClothingSimulationFactory* GetClothingSimFactory() const;
-
 protected:
 
 	/** Simulation data written back to the component after the simulation has taken place. If this data is required
@@ -1686,7 +1695,7 @@ public:
 	/** Store cloth simulation data into OutClothSimulData. Override USkinnedMeshComponent. */
 	ENGINE_API virtual void GetUpdateClothSimulationData_AnyThread(TMap<int32, FClothSimulData>& OutClothSimulData, FMatrix& OutLocalToWorld, float& OutClothBlendWeight) const override;
 
-	/** Remove clothing actors from their simulation */
+	UE_DEPRECATED(5.7, "Use ReleaseAllClothingResources instead.")
 	ENGINE_API void RemoveAllClothingActors();
 
 	/** Remove all clothing actors from their simulation and clear any other necessary clothing data to leave the simulations in a clean state */
diff --git a/Engine/Source/Runtime/Engine/Private/Components/SkeletalMeshComponent.cpp b/Engine/Source/Runtime/Engine/Private/Components/SkeletalMeshComponent.cpp
--- a/Engine/Source/Runtime/Engine/Private/Components/SkeletalMeshComponent.cpp
+++ b/Engine/Source/Runtime/Engine/Private/Components/SkeletalMeshComponent.cpp
@@ -36,7 +36,8 @@
 #include "Animation/AnimNode_LinkedInputPose.h"
 
 #include "ClothCollisionSource.h"
-#include "ClothingSimulationInterface.h"
+#include "ClothingSimulationFactory.h"
+#include "ClothingSimulationInstance.h"
 #include "ClothingSimulationInteractor.h"
 #include "ClothingSimulationTeleportHelpers.h"
 #include "Features/IModularFeatures.h"
@@ -487,12 +488,6 @@ USkeletalMeshComponent::USkeletalMeshComponent(const FObjectInitializer& ObjectI
 	CachedAnimCurveUidVersion = 0;
 	ResetRootBodyIndex();
 
-	ClothingSimulationFactory = UClothingSimulationFactory::GetDefaultClothingSimulationFactoryClass();
-
-	ClothingSimulation = nullptr;
-	ClothingSimulationContext = nullptr;
-	ClothingInteractor = nullptr;
-
 	bAllowClothActors = true;
 	bPostEvaluatingAnimation = false;
 	bAllowAnimCurveEvaluation = true;
@@ -958,89 +953,7 @@ void USkeletalMeshComponent::OnRegister()
 		SetComponentTickEnabled(false);
 	}
 
-	// If no simulation factory is currently set, set it to the default factory
-	if (!ClothingSimulationFactory)
-	{
-		ClothingSimulationFactory = UClothingSimulationFactory::GetDefaultClothingSimulationFactoryClass();
-	}
-
-	USkeletalMesh* SkelMesh = GetSkeletalMeshAsset();
-	// Look up for the best simulation factory to support each asset
-	if (ClothingSimulationFactory && SkelMesh)
-	{
-		// Check whether all clothing assets are supported by the current factory
-		bool bSupportsAllAssets = true;
-
-		UClothingSimulationFactory* const DefaultObject = ClothingSimulationFactory->GetDefaultObject<UClothingSimulationFactory>();
-		for (UClothingAssetBase* const ClothingAsset : SkelMesh->GetMeshClothingAssets())
-		{
-			if (ClothingAsset && !DefaultObject->SupportsAsset(ClothingAsset))
-			{
-				bSupportsAllAssets = false;
-
-				UE_LOG(LogSkeletalMesh, Display,
-					TEXT("OnRegister[%s]: [%s] is currently unable to provide a fully functional simulation for each of this SkeletalMesh's clothing assets."),
-					*GetPathNameSafe(SkelMesh),
-					*ClothingSimulationFactory->GetName());
-				UE_LOG(LogSkeletalMesh, Display,
-					TEXT("OnRegister[%s]: The ClothingSimulationFactory property will now be automatically updated to use the most functional simulation that can be found."),
-					*GetPathNameSafe(SkelMesh));
-
-				break;
-			}
-		}
-
-		// Try to find a new clothing factory that matches most asset requirements
-		if (!bSupportsAllAssets)
-		{
-			int MostSupportedNumAssets = 0;
-
-			const TArray<IClothingSimulationFactoryClassProvider*> ClassProviders = IModularFeatures::Get().GetModularFeatureImplementations<IClothingSimulationFactoryClassProvider>(IClothingSimulationFactoryClassProvider::FeatureName);
-			for (const IClothingSimulationFactoryClassProvider* const ClassProvider : ClassProviders)
-			{
-				if (ClassProvider)
-				{
-					if (const TSubclassOf<UClothingSimulationFactory> NewClothingSimulationFactory = ClassProvider->GetClothingSimulationFactoryClass())
-					{
-						int NumAssets = 0;
-						int SupportedNumAssets = 0;
-						UClothingSimulationFactory* const NewDefaultObject = NewClothingSimulationFactory->GetDefaultObject<UClothingSimulationFactory>();
-						for (UClothingAssetBase* const ClothingAsset : SkelMesh->GetMeshClothingAssets())
-						{
-							if (ClothingAsset)
-							{
-								if (NewDefaultObject->SupportsAsset(ClothingAsset))
-								{
-									++SupportedNumAssets;
-								}
-								++NumAssets;
-							}
-						}
-
-						if (SupportedNumAssets > MostSupportedNumAssets)
-						{
-							ClothingSimulationFactory = NewClothingSimulationFactory;
-							MostSupportedNumAssets = SupportedNumAssets;
-							if (SupportedNumAssets == NumAssets)
-							{
-								bSupportsAllAssets = true;
-								break;  // Stop at the first factory that supports all assets
-							}
-						}
-					}
-				}
-			}
-
-			UE_CLOG(!MostSupportedNumAssets, LogSkeletalMesh, Warning,
-				TEXT("OnRegister[%s]: There is no clothing simulation factory available that supports any of this SkeletalMesh's clothing assets."),
-				*GetPathNameSafe(SkelMesh));
-
-			UE_CLOG(MostSupportedNumAssets && !bSupportsAllAssets, LogSkeletalMesh, Warning,
-				TEXT("OnRegister[%s]: The most suitable clothing simulation factory available only partially supports this SkeletalMesh's clothing assets."),
-				*GetPathNameSafe(SkelMesh));
-		}
-	}
-
+	// Create clothing simulations
 	RecreateClothingActors();
 }
 
@@ -1052,9 +965,6 @@ void USkeletalMeshComponent::OnUnregister()
 	// Wait for any in flight animation evaluation to complete
 	HandleExistingParallelEvaluationTask(bBlockOnTask, bPerformPostAnimEvaluation);
 
-	// Wait for any in flight clothing simulation to complete
-	HandleExistingParallelClothSimulation();
-
 	//clothing actors will be re-created in TickClothing
 	ReleaseAllClothingResources();
 
@@ -1078,18 +988,6 @@ void USkeletalMeshComponent::OnUnregister()
 		PostProcessAnimInstance->UninitializeAnimation();
 	}
 
-	UClothingSimulationFactory* SimFactory = GetClothingSimFactory();
-	if(ClothingSimulation && SimFactory)
-	{
-		ClothingSimulation->DestroyContext(ClothingSimulationContext);
-		ClothingSimulation->DestroyActors();
-		ClothingSimulation->Shutdown();
-
-		SimFactory->DestroySimulation(ClothingSimulation);
-		ClothingSimulation = nullptr;
-		ClothingSimulationContext = nullptr;
-	}
-
 	if (DeferredKinematicUpdateIndex != INDEX_NONE)
 	{
 		UWorld* World = GetWorld();
@@ -1964,9 +1862,9 @@ void USkeletalMeshComponent::TickComponent(float DeltaTime, enum ELevelTick Tick
 
 	if(!PrimaryComponentTick.bRunOnAnyThread)
 	{
-		if (ClothingSimulation)
+		for (FClothingSimulationInstance& ClothingSimulationInstance : ClothingSimulationInstances)
 		{
-			ClothingSimulation->UpdateWorldForces(this);
+			ClothingSimulationInstance.UpdateWorldForces(this);
 		}
 
 		UpdateEndPhysicsTickRegisteredState();
@@ -1974,16 +1872,19 @@ void USkeletalMeshComponent::TickComponent(float DeltaTime, enum ELevelTick Tick
 
 		// If we are suspended, we will not simulate clothing, but as clothing is simulated in local space
 		// relative to a root bone we need to extract simulation positions as this bone could be animated.
-		if((!CVarEnableClothPhysics.GetValueOnGameThread() || bClothingSimulationSuspended) && ClothingSimulation)
+		if (!CVarEnableClothPhysics.GetValueOnGameThread() || bClothingSimulationSuspended)
 		{
 			CSV_SCOPED_TIMING_STAT(Animation, Cloth);
 
-			// First update the simulation context, since the simulation isn't ticking
-			// and it is still required to get the correct simulation data and bounds.
-			constexpr bool bIsInitialization = false;
-			ClothingSimulation->FillContext(this, DeltaTime, ClothingSimulationContext, bIsInitialization);
+			for (FClothingSimulationInstance& ClothingSimulationInstance : ClothingSimulationInstances)
+			{
+				// First update the simulation context, since the simulation isn't ticking
+				// and it is still required to get the correct simulation data and bounds.
+				constexpr bool bIsInitialization = false;
+				ClothingSimulationInstance.FillContext(this, DeltaTime, bIsInitialization);
 
-			ClothingSimulation->GetSimulationData(CurrentSimulationData, this, Cast<USkeletalMeshComponent>(LeaderPoseComponent.Get()));
+				ClothingSimulationInstance.GetSimulationData(CurrentSimulationData, this, Cast<USkeletalMeshComponent>(LeaderPoseComponent.Get()));
+			}
 		}
 	}
 
@@ -2513,31 +2414,59 @@ void USkeletalMeshComponent::EvaluatePostProcessMeshInstance(TArray<FTransform>&
 	}
 }
 
+// UE_DEPRECATED(5.7, "Use GetClothingSimulationInstances instead.")
 const IClothingSimulation* USkeletalMeshComponent::GetClothingSimulation() const
 {
-	return ClothingSimulation;
+	return GetClothingSimulationInstances().Num() ? ClothingSimulationInstances[0].GetClothingSimulation() : nullptr;
 }
 
+// UE_DEPRECATED(5.7, "Use GetClothingSimulationInstances instead.")
 IClothingSimulation* USkeletalMeshComponent::GetClothingSimulation()
 {
-	return ClothingSimulation;
+	return GetClothingSimulationInstances().Num() ? ClothingSimulationInstances[0].GetClothingSimulation() : nullptr;
 }
 
+// UE_DEPRECATED(5.7, "Use GetClothingSimulationInstances instead.")
 const IClothingSimulationContext* USkeletalMeshComponent::GetClothingSimulationContext() const
 {
-	return ClothingSimulationContext;
+	return GetClothingSimulationInstances().Num() ? ClothingSimulationInstances[0].GetClothingSimulationContext() : nullptr;
 }
 
+// UE_DEPRECATED(5.7, "Use GetClothingSimulationInstances instead.")
 IClothingSimulationContext* USkeletalMeshComponent::GetClothingSimulationContext()
 {
 	check(IsInGameThread());
-	return ClothingSimulationContext;
+	return GetClothingSimulationInstances().Num() ? ClothingSimulationInstances[0].GetClothingSimulationContext() : nullptr;
 }
 
-UClothingSimulationInteractor* USkeletalMeshComponent::GetClothingSimulationInteractor() const
+UClothingSimulationInteractor* USkeletalMeshComponent::GetClothingSimulationInteractor(const FName ClothingDataName) const
 {
 	check(IsInGameThread());
-	return ClothingInteractor;
+	if (ClothingDataName != NAME_None)
+	{
+		for (const FClothingSimulationInstance& ClothingSimulationInstance : GetClothingSimulationInstances())
+		{
+			if (UClothingSimulationInteractor* const ClothingSimulationInteractor = ClothingSimulationInstance.GetClothingSimulationInteractor())
+			{
+				if (ClothingSimulationInteractor->ClothingInteractors.Contains(ClothingDataName))
+				{
+					return ClothingSimulationInteractor;
+				}
+			}
+		}
+	}
+	else if (ClothingSimulationInstances.Num())
+	{
+		// For legacy blueprint compatibility only, return the first simulation's interactor
+		UE_CLOG(
+			ClothingSimulationInstances.Num() > 1,
+			LogSkeletalMesh,
+			Warning,
+			TEXT("GetClothingSimulationInteractor must be passed a valid Clothing Data name to retrieve a Clothing Simulation Interactor for Skeletal Mesh [%s], because this asset requires multiple simulation types."),
+			GetSkeletalMeshAsset() ? *GetSkeletalMeshAsset()->GetName() : TEXT(""));
+		return ClothingSimulationInstances[0].GetClothingSimulationInteractor();
+	}
+	return nullptr;
 }
 
 void USkeletalMeshComponent::CompleteParallelClothSimulation()
@@ -2578,21 +2507,12 @@ void USkeletalMeshComponent::UpdateClothSimulationContext(float InDeltaTime)
 		UpdateClothTransformImp();
 	}
 
-	// Fill the context for the next simulation
-	if(ClothingSimulation)
+	// Fill the contexts for the next simulation
+	for (FClothingSimulationInstance& ClothingSimulationInstance : ClothingSimulationInstances)
 	{
 		const bool bIsInitialization = false;
-		ClothingSimulation->FillContext(this, InDeltaTime, ClothingSimulationContext, bIsInitialization);
-
-		if(ClothingInteractor)
-		{
-			if (bNeedsHardReset)
-			{
-				ClothingInteractor->ClothConfigUpdated();
-			}
-
-			ClothingInteractor->Sync(ClothingSimulation, ClothingSimulationContext);
-		}
+		ClothingSimulationInstance.FillContext(this, InDeltaTime, bIsInitialization);
+		ClothingSimulationInstance.SyncClothingInteractor(bNeedsHardReset);
 	}
 
 	PendingTeleportType = ETeleportType::None;
@@ -2623,37 +2543,28 @@ void USkeletalMeshComponent::HandleExistingParallelClothSimulation()
 
 void USkeletalMeshComponent::WritebackClothingSimulationData()
 {
-	if(ClothingSimulation)
+	if (ClothingSimulationInstances.Num())
 	{
 		CSV_SCOPED_TIMING_STAT(Animation, Cloth);
 
-		USkinnedMeshComponent* OverrideComponent = nullptr;
-		if(LeaderPoseComponent.IsValid())
+		const USkinnedMeshComponent* OverrideComponent = nullptr;
+		if (LeaderPoseComponent.IsValid())
 		{
 			OverrideComponent = LeaderPoseComponent.Get();
 
 			// Check if our bone map is actually valid, if not there is no clothing data to build
-			if(LeaderBoneMap.Num() == 0)
+			if (LeaderBoneMap.Num() == 0)
 			{
 				CurrentSimulationData.Reset();
 				return;
 			}
 		}
 
-		ClothingSimulation->GetSimulationData(CurrentSimulationData, this, OverrideComponent);
-	}
-}
-
-UClothingSimulationFactory* USkeletalMeshComponent::GetClothingSimFactory() const
-{
-	UClass* SimFactoryClass = *ClothingSimulationFactory;
-	if(SimFactoryClass)
-	{
-		return SimFactoryClass->GetDefaultObject<UClothingSimulationFactory>();
+		for (const FClothingSimulationInstance& ClothingSimulationInstance : ClothingSimulationInstances)
+		{
+			ClothingSimulationInstance.GetSimulationData(CurrentSimulationData, this, OverrideComponent);
+		}
 	}
-
-	// No simulation factory set
-	return nullptr;
 }
 
 void USkeletalMeshComponent::DoInstancePreEvaluation()
@@ -3336,7 +3247,7 @@ void USkeletalMeshComponent::SetSkeletalMesh(USkeletalMesh* InSkelMesh, bool bRe
 	}
 
 	// Stop any existing cloth simulation prior to replacing the asset
-	RemoveAllClothingActors();
+	ReleaseAllClothingResources();
 
 	// Update property alias
 #if WITH_EDITORONLY_DATA
@@ -4830,7 +4741,7 @@ void USkeletalMeshComponent::BindClothToLeaderPoseComponent()
 			return;
 		}
 
-		if(ClothingSimulation && LeaderComp->ClothingSimulation)
+		if(ClothingSimulationInstances.Num() && LeaderComp->ClothingSimulationInstances.Num())
 		{
 			bDisableClothSimulation = true;
 
@@ -4845,7 +4756,7 @@ void USkeletalMeshComponent::UnbindClothFromLeaderPoseComponent(bool bRestoreSim
 	USkeletalMeshComponent* LeaderComp = Cast<USkeletalMeshComponent>(LeaderPoseComponent.Get());
 	if(LeaderComp && bBindClothToLeaderComponent)
 	{
-		if(ClothingSimulation)
+		if(ClothingSimulationInstances.Num())
 		{
 			bDisableClothSimulation = false;
 		}
diff --git a/Engine/Source/Runtime/Engine/Private/PhysicsEngine/EnvironmentalCollisions.cpp b/Engine/Source/Runtime/Engine/Private/PhysicsEngine/EnvironmentalCollisions.cpp
--- a/Engine/Source/Runtime/Engine/Private/PhysicsEngine/EnvironmentalCollisions.cpp
+++ b/Engine/Source/Runtime/Engine/Private/PhysicsEngine/EnvironmentalCollisions.cpp
@@ -13,6 +13,7 @@
 #include "Engine/World.h"
 #include "Misc/ConfigCacheIni.h"
 #include "ClothCollisionData.h"
+#include "ClothingSimulationInstance.h"
 
 void FEnvironmentalCollisions::AppendCollisionDataFromEnvironment(const USceneComponent* SceneComponent, FClothCollisionData& CollisionData)
 {
@@ -341,11 +342,21 @@ void FEnvironmentalCollisions::AppendCollisionDataFromEnvironment(const USceneCo
 						continue;
 					}
 
-					if (SkelComp->GetClothingSimulation())
+					const TArray<FClothingSimulationInstance>& ClothingSimulationInstances = SkelComp->GetClothingSimulationInstances();
+					if (ClothingSimulationInstances.Num())
 					{
 						// append skeletal component collisions
 						FClothCollisionData SkelCollisionData;
-						SkelComp->GetClothingSimulation()->GetCollisions(SkelCollisionData, false);
+						constexpr bool bIncludeExternal = false;
+
+						ClothingSimulationInstances[0].GetCollisions(SkelCollisionData, bIncludeExternal);
+
+						// AppendUnique is quite expensive, so let's only call it for subsequent clothing simulation instances
+						for (int32 Index = 1 ; Index < ClothingSimulationInstances.Num(); ++Index)
+						{
+							ClothingSimulationInstances[Index].AppendUniqueCollisions(SkelCollisionData, bIncludeExternal);
+						}
+
 						CollisionData.Append(SkelCollisionData);
 					}
 				}
diff --git a/Engine/Source/Runtime/Engine/Private/SkeletalMeshComponentPhysics.cpp b/Engine/Source/Runtime/Engine/Private/SkeletalMeshComponentPhysics.cpp
--- a/Engine/Source/Runtime/Engine/Private/SkeletalMeshComponentPhysics.cpp
+++ b/Engine/Source/Runtime/Engine/Private/SkeletalMeshComponentPhysics.cpp
@@ -2,7 +2,9 @@
 
 #include "BodySetupEnums.h"
 #include "Misc/ConfigCacheIni.h"
-#include "ClothingSimulationInterface.h"
+#include "ClothingAssetBase.h"
+#include "ClothingSimulationFactory.h"
+#include "ClothingSimulationInstance.h"
 #include "Misc/ConfigCacheIni.h"
 #include "ClothCollisionSource.h"
 #include "Components/InstancedStaticMeshComponent.h"
@@ -2904,9 +2906,12 @@ bool USkeletalMeshComponent::ComponentOverlapMultiImpl(TArray<struct FOverlapRes
 
 void USkeletalMeshComponent::AddClothingBounds(FBoxSphereBounds& InOutBounds, const FTransform& LocalToWorld) const
 {
-	if(ClothingSimulation && ClothingSimulation->ShouldSimulate())
+	for (const FClothingSimulationInstance& ClothingSimulationInstance : ClothingSimulationInstances)
 	{
-		InOutBounds = InOutBounds + ClothingSimulation->GetBounds(this).TransformBy(LocalToWorld);
+		if (ClothingSimulationInstance.ShouldSimulate())
+		{
+			InOutBounds = InOutBounds + ClothingSimulationInstance.GetBounds(this).TransformBy(LocalToWorld);
+		}
 	}
 }
 
@@ -2916,118 +2921,97 @@ void USkeletalMeshComponent::RecreateClothingActors()
 
 	ReleaseAllClothingResources();
 
-	USkeletalMesh* SkelMesh = GetSkeletalMeshAsset();
+	const USkeletalMesh* const SkelMesh = GetSkeletalMeshAsset();
 	if (SkelMesh == nullptr || !IsRegistered())
 	{
 		return;
 	}
 
-	if (bAllowClothActors && CVarEnableClothPhysics.GetValueOnGameThread() && (SkelMesh->GetMeshClothingAssets().Num() > 0))
+	if (bAllowClothActors && CVarEnableClothPhysics.GetValueOnGameThread())
 	{
-		UClass* SimFactoryClass = *ClothingSimulationFactory;
-		if (SimFactoryClass)
+		const TArray<UClothingAssetBase*>& MeshClothingAssets = SkelMesh->GetMeshClothingAssets();
+		if (MeshClothingAssets.Num())
 		{
-			UClothingSimulationFactory* SimFactory = SimFactoryClass->GetDefaultObject<UClothingSimulationFactory>();
-			if(ClothingSimulation == nullptr)
-			{
-				ClothingSimulation = SimFactory->CreateSimulation();
-
-				if(ClothingSimulation)
-				{
-					ClothingSimulation->Initialize();
-					ClothingSimulationContext = ClothingSimulation->CreateContext();
+			// Group assets per factory, as to only create one simulation per asset of the same type
+			TArray<UClothingAssetBase*> ClothingAssetsInUse;
+			SkelMesh->GetClothingAssetsInUse(ClothingAssetsInUse);
 
-					if(SimFactory->SupportsRuntimeInteraction())
-					{
-						ClothingInteractor = SimFactory->CreateInteractor();
-					}
-				}
-			}
-		}
+			TArray<FClothingSimulationInstance::FFactoryAssetGroup, TInlineAllocator<8>> FactoryAssetGroups;
+			FactoryAssetGroups.Reserve(ClothingAssetsInUse.Num());
 
-		if (ClothingSimulation)
-		{
-			// Only create cloth sim actors when the world is ready for it
-			const UWorld* const World = GetWorld();
-#if WITH_EDITORONLY_DATA
-			if (World && (World->bShouldSimulatePhysics || bUpdateClothInEditor) && World->GetPhysicsScene())
-#else
-			if (World && World->bShouldSimulatePhysics && World->GetPhysicsScene())
-#endif
+			for (int32 ClothingAssetIndex = 0; ClothingAssetIndex < MeshClothingAssets.Num(); ++ClothingAssetIndex)
 			{
-				UE_LOG(LogSkeletalMesh, Log, TEXT("USkeletalMeshComponent: Recreating Clothing Actors for '%s' with '%s'"), *GetName(), *SkelMesh->GetName());
-				TArray<UClothingAssetBase*> AssetsInUse;
-				SkelMesh->GetClothingAssetsInUse(AssetsInUse);
-
-				const bool bIsInitialization = true;
-				ClothingSimulation->FillContext(this, 0.f, ClothingSimulationContext, bIsInitialization);
-
-				const int32 NumMeshAssets = SkelMesh->GetMeshClothingAssets().Num();
-				for (int32 BaseAssetIndex = 0; BaseAssetIndex < NumMeshAssets; ++BaseAssetIndex)
+				const UClothingAssetBase* const ClothingAsset = MeshClothingAssets[ClothingAssetIndex];
+				if (ClothingAsset && ClothingAsset->IsValid() && ClothingAssetsInUse.Contains(ClothingAsset))
 				{
-					UClothingAssetBase* const Asset = SkelMesh->GetMeshClothingAssets()[BaseAssetIndex];
-					if (Asset && AssetsInUse.Contains(Asset))
+					if (UClothingSimulationFactory* const ClothingAssetFactory = UClothingSimulationFactory::GetClothingSimulationFactory(ClothingAsset))
 					{
-						ClothingSimulation->CreateActor(this, Asset, BaseAssetIndex);
-						
-						if (ClothingInteractor)
+						FClothingSimulationInstance::FFactoryAssetGroup* FactoryAssetGroup = FactoryAssetGroups.FindByPredicate(
+							[&ClothingAssetFactory](const FClothingSimulationInstance::FFactoryAssetGroup& FactoryAssetGroup)
+							{
+								return FactoryAssetGroup.GetClothingSimulationFactory() == ClothingAssetFactory;
+							});
+						if (!FactoryAssetGroup)
 						{
-							ClothingInteractor->CreateClothingInteractor(Asset, BaseAssetIndex);
+							FactoryAssetGroup = &FactoryAssetGroups.Emplace_GetRef(ClothingAssetFactory, MeshClothingAssets.Num());
 						}
+						FactoryAssetGroup->AddClothingAsset(ClothingAsset, ClothingAssetIndex);
 					}
 				}
 			}
 
-			// Retrieve the cloth sim data, or clear the data if the world isn't ready to sim
-			WritebackClothingSimulationData();
-		}
-	}
-	else
-	{
-		// No clothing assets, so destroy any clothing sim we have
-		UClothingSimulationFactory* SimFactory = GetClothingSimFactory();
-		if(ClothingSimulation && SimFactory)
-		{
-			ClothingSimulation->DestroyContext(ClothingSimulationContext);
-			ClothingSimulation->DestroyActors();
-			ClothingSimulation->Shutdown();
+			// Create clothing simulation instances
+			checkf(!ClothingSimulationInstances.Num(), TEXT("ReleaseAllClothingResources() should have been called to clear all clothing simulation instances before this point."));
+			ClothingSimulationInstances.Reserve(FactoryAssetGroups.Num());
+			for (FClothingSimulationInstance::FFactoryAssetGroup& FactoryAssetGroup : FactoryAssetGroups)
+			{
+				// Initialize a new clothing simulation instance for this asset group, add actors, create interactors, ...etc
+				FClothingSimulationInstance& ClothingSimulationInstance = ClothingSimulationInstances.Emplace_GetRef(this, FactoryAssetGroup);
+			}
 
-			SimFactory->DestroySimulation(ClothingSimulation);
-			ClothingSimulation = nullptr;
-			ClothingSimulationContext = nullptr;
+			if (ClothingSimulationInstances.Num())
+			{
+				// Retrieve the cloth sim data, or clear the data if the world isn't ready to sim
+				WritebackClothingSimulationData();
+			}
 		}
 	}
 }
 
+// UE_DEPRECATED(5.7, "Use ReleaseAllClothingResources instead.")
 void USkeletalMeshComponent::RemoveAllClothingActors()
 {
-	if(ClothingSimulation)
+	if (ClothingSimulationInstances.Num())
 	{
 		// Can't destroy our actors if we're still simulating
 		HandleExistingParallelClothSimulation();
 
-		if(ClothingInteractor)
+		for (FClothingSimulationInstance& ClothingSimulationInstance : ClothingSimulationInstances)
 		{
-			ClothingInteractor->DestroyClothingInteractors();
+			ClothingSimulationInstance.RemoveAllClothingActors();
 		}
-
-		ClothingSimulation->DestroyActors();
 	}
 }
 
 void USkeletalMeshComponent::ReleaseAllClothingResources()
 {
-#if WITH_CLOTH_COLLISION_DETECTION
-	if(ClothingSimulation)
+	if (ClothingSimulationInstances.Num())
 	{
 		// Ensure no running simulation first
 		HandleExistingParallelClothSimulation();
 
-		ClothingSimulation->ClearExternalCollisions();
-	}
+		for (FClothingSimulationInstance& ClothingSimulationInstance : ClothingSimulationInstances)
+		{
+#if WITH_CLOTH_COLLISION_DETECTION
+			ClothingSimulationInstance.ClearExternalCollisions();
 #endif // #if WITH_CLOTH_COLLISION_DETECTION
+			
+			ClothingSimulationInstance.RemoveAllClothingActors();
+		}
 
-	RemoveAllClothingActors();
+		// Remove all clothing simulations
+		ClothingSimulationInstances.Reset();
+	}
 }
 
 void USkeletalMeshComponent::GetWindForCloth_GameThread(FVector& WindDirection, float& WindAdaption) const
@@ -3060,11 +3044,18 @@ void USkeletalMeshComponent::GetWindForCloth_GameThread(FVector& WindDirection,
 
 void USkeletalMeshComponent::FindClothCollisions(FClothCollisionData& OutCollisions) const
 {
-	if(ClothingSimulation)
+	if (ClothingSimulationInstances.Num())
 	{
-		// Get collisions for this simulation, ignoring any externally added collisions
-		// (i.e. on grab the asset collisions, not environment etc.)
-		ClothingSimulation->GetCollisions(OutCollisions, false);
+		// append skeletal component collisions
+		constexpr bool bIncludeExternal = false;
+
+		ClothingSimulationInstances[0].GetCollisions(OutCollisions, bIncludeExternal);
+
+		// AppendUnique is quite expensive, so let's only call it for subsequent clothing simulation instances
+		for (int32 Index = 1; Index < ClothingSimulationInstances.Num(); ++Index)
+		{
+			ClothingSimulationInstances[Index].AppendUniqueCollisions(OutCollisions, bIncludeExternal);
+		}
 	}
 }
 
@@ -3085,36 +3076,40 @@ void USkeletalMeshComponent::ExtractCollisionsForCloth(
 
 void USkeletalMeshComponent::CopyClothCollisionsToChildren()
 {
-	// 3 steps
-	// 1. release all previous parent collisions
-	// 2. find new collisions from parent(this class)
-	// 3. add new collisions to children
+	// Find all suitable children
+	const TArray<TObjectPtr<USceneComponent>>& AttachedChildren = USceneComponent::GetAttachChildren();
 
 	TArray<USkeletalMeshComponent*> ClothChildren;
+	ClothChildren.Reserve(AttachedChildren.Num());
 
-	for (USceneComponent* AttachedChild : GetAttachChildren())
+	for (const TObjectPtr<USceneComponent>& AttachedChild : AttachedChildren)
 	{
-		USkeletalMeshComponent* pChild = Cast<USkeletalMeshComponent>(AttachedChild);
-		if(pChild && pChild->ClothingSimulation && pChild->ClothingSimulation->ShouldSimulate())
+		if (USkeletalMeshComponent* const Child = Cast<USkeletalMeshComponent>(AttachedChild))
 		{
-			ClothChildren.Add(pChild);
+			for (const FClothingSimulationInstance& ClothingSimulationInstance : Child->ClothingSimulationInstances)
+			{
+				if (ClothingSimulationInstance.ShouldSimulate())
+				{
+					ClothChildren.Add(Child);
+					break;
+				}
+			}
 		}
 	}
 
-	const int32 NumClothChildren = ClothChildren.Num();
-
-	if(NumClothChildren == 0)
+	// Add new collision to children
+	if (ClothChildren.Num())
 	{
-		return;
-	}
-
-	FClothCollisionData NewCollisions;
-
-	FindClothCollisions(NewCollisions);
+		FClothCollisionData NewCollisions;
+		FindClothCollisions(NewCollisions);
 
-	for(USkeletalMeshComponent* Child : ClothChildren)
-	{
-		Child->ClothingSimulation->AddExternalCollisions(NewCollisions);
+		for (USkeletalMeshComponent* const Child : ClothChildren)
+		{
+			for (FClothingSimulationInstance& ClothingSimulationInstance : Child->ClothingSimulationInstances)
+			{
+				ClothingSimulationInstance.AddExternalCollisions(NewCollisions);
+			}
+		}
 	}
 }
 
@@ -3122,13 +3117,15 @@ void USkeletalMeshComponent::CopyClothCollisionSources()
 {
 	FClothCollisionData ExternalCollisions;
 
-	for(FClothCollisionSource& ClothCollisionSource : ClothCollisionSources)
+	for (FClothCollisionSource& ClothCollisionSource : ClothCollisionSources)
 	{
 		ClothCollisionSource.ExtractCollisions(this, ExternalCollisions);
 	}
 
-
-	ClothingSimulation->AddExternalCollisions(ExternalCollisions);
+	for (FClothingSimulationInstance& ClothingSimulationInstance : ClothingSimulationInstances)
+	{
+		ClothingSimulationInstance.AddExternalCollisions(ExternalCollisions);
+	}
 }
 
 void USkeletalMeshComponent::CopyChildrenClothCollisionsToParent()
@@ -3144,20 +3141,28 @@ void USkeletalMeshComponent::CopyChildrenClothCollisionsToParent()
 		}
 	}
 
-	ClothingSimulation->AddExternalCollisions(NewCollisions);
+	// Add new collisions to parent (this component)
+	for (FClothingSimulationInstance& ClothingSimulationInstance : ClothingSimulationInstances)
+	{
+		ClothingSimulationInstance.AddExternalCollisions(NewCollisions);
+	}
 }
 
 void USkeletalMeshComponent::ProcessClothCollisionWithEnvironment()
 {
 	// don't handle collision detection if this component is in editor
-	if(!GetWorld()->IsGameWorld() || !ClothingSimulation)
+	if (!GetWorld()->IsGameWorld() || !ClothingSimulationInstances.Num())
 	{
 		return;
 	}
 
 	FClothCollisionData NewCollisionData;
 	FEnvironmentalCollisions::AppendCollisionDataFromEnvironment(this, NewCollisionData);
-	ClothingSimulation->AddExternalCollisions(NewCollisionData);
+
+	for (FClothingSimulationInstance& ClothingSimulationInstance : ClothingSimulationInstances)
+	{
+		ClothingSimulationInstance.AddExternalCollisions(NewCollisionData);
+	}
 }
 
 #endif// #if WITH_CLOTH_COLLISION_DETECTION
@@ -3268,10 +3273,16 @@ void USkeletalMeshComponent::EndPhysicsTickComponent(FSkeletalMeshComponentEndPh
 void USkeletalMeshComponent::UpdateClothTransformImp()
 {
 #if WITH_CLOTH_COLLISION_DETECTION
-	if (ClothingSimulation && ClothingSimulation->ShouldSimulate())
+	if (ClothingSimulationInstances.Num())
 	{
-		// Component has moved, and there is something to simulate, update all external cloth collisions
-		ClothingSimulation->ClearExternalCollisions();
+		for (FClothingSimulationInstance& ClothingSimulationInstance : ClothingSimulationInstances)
+		{
+			if (ClothingSimulationInstance.ShouldSimulate())
+			{
+				// Component has moved, and there is something to simulate, update all external cloth collisions
+				ClothingSimulationInstance.ClearExternalCollisions();
+			}
+		}
 
 		if (bCollideWithAttachedChildren)
 		{
@@ -3382,11 +3393,11 @@ FAutoConsoleTaskPriority CPrio_FParallelClothTask(
 
 class FParallelClothTask
 {
-	const USkeletalMeshComponent& SkeletalMeshComponent;
+	USkeletalMeshComponent& SkeletalMeshComponent;
 	float DeltaTime;
 
 public:
-	FParallelClothTask(const USkeletalMeshComponent& InSkeletalMeshComponent, float InDeltaTime)
+	FParallelClothTask(USkeletalMeshComponent& InSkeletalMeshComponent, float InDeltaTime)
 		: SkeletalMeshComponent(InSkeletalMeshComponent)
 		, DeltaTime(InDeltaTime)
 	{
@@ -3415,9 +3426,9 @@ public:
 		SCOPE_CYCLE_COUNTER(STAT_ClothTotalTime);
 		CSV_SCOPED_TIMING_STAT(Animation, Cloth);
 
-		if(SkeletalMeshComponent.ClothingSimulation)
+		for (FClothingSimulationInstance& ClothingSimulationInstance : SkeletalMeshComponent.ClothingSimulationInstances)
 		{
-			SkeletalMeshComponent.ClothingSimulation->Simulate(SkeletalMeshComponent.ClothingSimulationContext);
+			ClothingSimulationInstance.Simulate();
 		}
 	}
 };
@@ -3462,7 +3473,7 @@ public:
 
 bool USkeletalMeshComponent::RequiresPreEndOfFrameSync() const
 {
-	if ((ClothingSimulation != nullptr) && (CVarEnableClothPhysics.GetValueOnGameThread() != 0))
+	if (ClothingSimulationInstances.Num() && CVarEnableClothPhysics.GetValueOnGameThread())
 	{
 		// By default we await the cloth task in the ClothTickFunction, but...
 		// If we have cloth and have no game-thread dependencies on the cloth output, 
@@ -3503,7 +3514,7 @@ const TMap<int32, FClothSimulData>& USkeletalMeshComponent::GetCurrentClothingDa
 		return SEmptyClothSimulationData;
 	}
 
-	return ClothingSimulation ? CurrentSimulationData : SEmptyClothSimulationData;
+	return ClothingSimulationInstances.Num() ? CurrentSimulationData : SEmptyClothSimulationData;
 }
 
 const TMap<int32, FClothSimulData>& USkeletalMeshComponent::GetCurrentClothingData_AnyThread() const
@@ -3511,7 +3522,7 @@ const TMap<int32, FClothSimulData>& USkeletalMeshComponent::GetCurrentClothingDa
 	// This is called during EndOfFrameUpdates, usually in a parallel-for loop. We need to be sure that
 	// the cloth task (if there is one) is complete, but it cannpt be waited for here. See OnPreEndOfFrameUpdateSync
 	// which is called just before EOF updates and is where we would have waited for the cloth task.
-	if (ClothingSimulation && (!IsValidRef(ParallelClothTask) || ParallelClothTask->IsComplete()))
+	if (ClothingSimulationInstances.Num() && (!IsValidRef(ParallelClothTask) || ParallelClothTask->IsComplete()))
 	{
 		return CurrentSimulationData;
 	}
@@ -3524,7 +3535,7 @@ void USkeletalMeshComponent::UpdateClothStateAndSimulate(float DeltaTime, FTickF
 	check(IsInGameThread());
 
 	// If disabled or no simulation
-	if (CVarEnableClothPhysics.GetValueOnGameThread() == 0 || !ClothingSimulation || bDisableClothSimulation)
+	if (CVarEnableClothPhysics.GetValueOnGameThread() == 0 || !ClothingSimulationInstances.Num() || bDisableClothSimulation)
 	{
 		return;
 	}
@@ -3543,7 +3554,10 @@ void USkeletalMeshComponent::UpdateClothStateAndSimulate(float DeltaTime, FTickF
 #if WITH_CLOTH_COLLISION_DETECTION
 	if (bCollideWithAttachedChildren)
 	{
-		ClothingSimulation->ClearExternalCollisions();
+		for (FClothingSimulationInstance& ClothingSimulationInstance : ClothingSimulationInstances)
+		{
+			ClothingSimulationInstance.ClearExternalCollisions();
+		}
 
 		CopyClothCollisionsToChildren();
 		CopyChildrenClothCollisionsToParent();
@@ -3598,7 +3612,7 @@ void USkeletalMeshComponent::TickClothing(float DeltaTime, FTickFunction& ThisTi
 #if WITH_EDITOR
 	bIsCompiling = (GetSkeletalMeshAsset() && GetSkeletalMeshAsset()->IsCompiling());
 #endif
-	if (GetSkeletalMeshAsset() == nullptr || !ClothingSimulation || CVarEnableClothPhysics.GetValueOnGameThread() == 0 || bIsCompiling)
+	if (GetSkeletalMeshAsset() == nullptr || !ClothingSimulationInstances.Num() || CVarEnableClothPhysics.GetValueOnGameThread() == 0 || bIsCompiling)
 	{
 		return;
 	}
@@ -3637,10 +3651,9 @@ void USkeletalMeshComponent::GetUpdateClothSimulationData(TMap<int32, FClothSimu
 		return;
 	}
 
-	if(ClothingSimulation)
+	for (const FClothingSimulationInstance& ClothingSimulationInstance : ClothingSimulationInstances)
 	{
-		ClothingSimulation->GetSimulationData(OutClothSimData, this, OverrideLocalRootComponent);
-		return;
+		ClothingSimulationInstance.GetSimulationData(OutClothSimData, this, OverrideLocalRootComponent);
 	}
 }
 
@@ -3682,16 +3695,18 @@ void USkeletalMeshComponent::DebugDrawClothing(FPrimitiveDrawInterface* PDI)
 {
 #if WITH_EDITOR && ENABLE_DRAW_DEBUG
 
-	if(ClothingSimulation)
+	if (ClothingSimulationInstances.Num())
 	{
 		FClothingSystemEditorInterfaceModule& ClothingEditorModule = FModuleManager::LoadModuleChecked<FClothingSystemEditorInterfaceModule>(TEXT("ClothingSystemEditorInterface"));
 
-		if(ISimulationEditorExtender* Extender = ClothingEditorModule.GetSimulationEditorExtender(ClothingSimulationFactory->GetFName()))
+		for (const FClothingSimulationInstance& ClothingSimulationInstance : ClothingSimulationInstances)
 		{
-			Extender->DebugDrawSimulation(ClothingSimulation, this, PDI);
+			const FName ClothingSimulationFactoryName = ClothingSimulationInstance.GetClothingSimulationFactory()->GetClass()->GetFName();
+			if (ISimulationEditorExtender* const Extender = ClothingEditorModule.GetSimulationEditorExtender(ClothingSimulationFactoryName))
+			{
+				Extender->DebugDrawSimulation(ClothingSimulationInstance.GetClothingSimulation(), this, PDI);
+			}
 		}
-
-		return;
 	}
 
 #endif
@@ -3701,16 +3716,18 @@ void USkeletalMeshComponent::DebugDrawClothingTexts(FCanvas* Canvas, const FScen
 {
 #if WITH_EDITOR && ENABLE_DRAW_DEBUG
 
-	if(ClothingSimulation)
+	if (ClothingSimulationInstances.Num())
 	{
 		FClothingSystemEditorInterfaceModule& ClothingEditorModule = FModuleManager::LoadModuleChecked<FClothingSystemEditorInterfaceModule>(TEXT("ClothingSystemEditorInterface"));
 
-		if(ISimulationEditorExtender* Extender = ClothingEditorModule.GetSimulationEditorExtender(ClothingSimulationFactory->GetFName()))
+		for (const FClothingSimulationInstance& ClothingSimulationInstance : ClothingSimulationInstances)
 		{
-			Extender->DebugDrawSimulationTexts(ClothingSimulation, this, Canvas, SceneView);
+			const FName ClothingSimulationFactoryName = ClothingSimulationInstance.GetClothingSimulationFactory()->GetClass()->GetFName();
+			if (ISimulationEditorExtender* const Extender = ClothingEditorModule.GetSimulationEditorExtender(ClothingSimulationFactoryName))
+			{
+				Extender->DebugDrawSimulationTexts(ClothingSimulationInstance.GetClothingSimulation(), this, Canvas, SceneView);
+			}
 		}
-
-		return;
 	}
 
 #endif
