CIS - Fixed Issue 1016138: Undefined symbol 'UChaosClothAssetSKMClothingAsset::GetClothSimulationModelIds'

#rnx

diff --git a/Engine/Plugins/ChaosClothAsset/Source/ChaosClothAssetEngine/Private/ChaosClothAsset/ClothAssetSKMClothingAsset.cpp b/Engine/Plugins/ChaosClothAsset/Source/ChaosClothAssetEngine/Private/ChaosClothAsset/ClothAssetSKMClothingAsset.cpp
--- a/Engine/Plugins/ChaosClothAsset/Source/ChaosClothAssetEngine/Private/ChaosClothAsset/ClothAssetSKMClothingAsset.cpp
+++ b/Engine/Plugins/ChaosClothAsset/Source/ChaosClothAssetEngine/Private/ChaosClothAsset/ClothAssetSKMClothingAsset.cpp
@@ -436,7 +436,6 @@ FName UChaosClothAssetSKMClothingAsset::GetClothSimulationModelName() const
 }
 #endif  // #if WITH_EDITOR
 
-#if WITH_EDITORONLY_DATA
 TArray<FString> UChaosClothAssetSKMClothingAsset::GetClothSimulationModelIds() const
 {
 	using namespace UE::Chaos::ClothAsset::Private;
@@ -458,6 +457,5 @@ TArray<FString> UChaosClothAssetSKMClothingAsset::GetClothSimulationModelIds() c
 	}
 	return ClothSimulationModelIds;
 }
-#endif  // #if WITH_EDITORONLY_DATA
 
 #undef LOCTEXT_NAMESPACE
