CIS - Fixed Issue 1016138: Undefined symbol 'UChaosClothAssetSKMClothingAsset::GetClothSimulationModelIds'

#was-backout 45031155
#rnx

diff --git a/Engine/Plugins/ChaosClothAsset/Source/ChaosClothAssetEngine/Private/ChaosClothAsset/ClothAssetSKMClothingAsset.cpp b/Engine/Plugins/ChaosClothAsset/Source/ChaosClothAssetEngine/Private/ChaosClothAsset/ClothAssetSKMClothingAsset.cpp
--- a/Engine/Plugins/ChaosClothAsset/Source/ChaosClothAssetEngine/Private/ChaosClothAsset/ClothAssetSKMClothingAsset.cpp
+++ b/Engine/Plugins/ChaosClothAsset/Source/ChaosClothAssetEngine/Private/ChaosClothAsset/ClothAssetSKMClothingAsset.cpp
@@ -3,11 +3,11 @@
 #include "ChaosClothAsset/ClothAssetSKMClothingAsset.h"
 #include "ChaosClothAsset/ClothAssetBase.h"
 #include "ChaosClothAsset/ClothAssetPrivate.h"
+#include "ChaosClothAsset/ClothSimulationModel.h"
 
 #if WITH_EDITOR
 #include "Chaos/CollectionPropertyFacade.h"
 #include "ChaosClothAsset/ClothEngineTools.h"
-#include "ChaosClothAsset/ClothSimulationModel.h"
 #include "ChaosClothAsset/CollectionClothFacade.h"
 #include "Components/SkeletalMeshComponent.h"
 #include "Engine/SkeletalMesh.h"
@@ -22,12 +22,12 @@
 
 namespace UE::Chaos::ClothAsset::Private
 {
-#if WITH_EDITOR
 	static FString MakeClothSimulationModelIdString(const FName Name, const FGuid Guid)
 	{
 		return FString::Format(TEXT("{0}-{1}"), { Name.ToString(), Guid.ToString() });
 	}
 
+#if WITH_EDITOR
 	// Modified version of FLODUtilities::UnbindClothingAndBackup to only unbind specified assets
 	static void UnbindClothingAndBackup(USkeletalMesh& SkeletalMesh, const UChaosClothAssetSKMClothingAsset& ClothingAsset, TArray<ClothingAssetUtils::FClothingAssetMeshBinding> ClothingBindings)
 	{
@@ -436,7 +436,6 @@ FName UChaosClothAssetSKMClothingAsset::GetClothSimulationModelName() const
 }
 #endif  // #if WITH_EDITOR
 
-#if WITH_EDITORONLY_DATA
 TArray<FString> UChaosClothAssetSKMClothingAsset::GetClothSimulationModelIds() const
 {
 	using namespace UE::Chaos::ClothAsset::Private;
@@ -458,6 +457,5 @@ TArray<FString> UChaosClothAssetSKMClothingAsset::GetClothSimulationModelIds() c
 	}
 	return ClothSimulationModelIds;
 }
-#endif  // #if WITH_EDITORONLY_DATA
 
 #undef LOCTEXT_NAMESPACE
