CIS - Fixed Issue 1016153: PVS Studio warnings in ClothSimulationProxy.cpp

#rnx

diff --git a/Engine/Plugins/ChaosClothAsset/Source/ChaosClothAssetEngine/Private/ChaosClothAsset/ClothSimulationProxy.cpp b/Engine/Plugins/ChaosClothAsset/Source/ChaosClothAssetEngine/Private/ChaosClothAsset/ClothSimulationProxy.cpp
--- a/Engine/Plugins/ChaosClothAsset/Source/ChaosClothAssetEngine/Private/ChaosClothAsset/ClothSimulationProxy.cpp
+++ b/Engine/Plugins/ChaosClothAsset/Source/ChaosClothAssetEngine/Private/ChaosClothAsset/ClothSimulationProxy.cpp
@@ -831,7 +831,7 @@ PRAGMA_ENABLE_DEPRECATION_WARNINGS
 
 			// Transform into the cloth reference simulation space used at the time of simulation
 			bool bSimDataIsFinite = true;
-			if (bHasCloth)
+			if (bHasCloth && Data)
 			{
 	#if INTEL_ISPC
 				if (bTransformClothSimulData_ISPC_Enabled)
