Add additional accessory meshes for chaos cloth. Update the backstop constraint to be able to use this accessory mesh.

Currently, accessory meshes can be created by converting an existing sim mesh collection into an accessory mesh. If ImportedVertexIDs exist on  both the source and target collections, these can be used to match up vertex indices that were reordered as part of a SKM import.

diff --git a/Engine/Plugins/ChaosCloth/Source/ChaosCloth/Private/ChaosCloth/ChaosClothConstraints.cpp b/Engine/Plugins/ChaosCloth/Source/ChaosCloth/Private/ChaosCloth/ChaosClothConstraints.cpp
--- a/Engine/Plugins/ChaosCloth/Source/ChaosCloth/Private/ChaosCloth/ChaosClothConstraints.cpp
+++ b/Engine/Plugins/ChaosCloth/Source/ChaosCloth/Private/ChaosCloth/ChaosClothConstraints.cpp
@@ -1,5 +1,6 @@
 // Copyright Epic Games, Inc. All Rights Reserved.
 #include "ChaosCloth/ChaosClothConstraints.h"
+#include "ChaosCloth/ChaosClothingAccessoryMeshData.h"
 #include "ChaosCloth/ChaosClothingSimulationConfig.h"
 #include "ChaosCloth/ChaosWeightMapTarget.h"
 #include "ChaosCloth/ChaosClothingPatternData.h"
@@ -788,7 +789,8 @@ void FClothConstraints::AddRules(
 	const int32 MultiResCoarseLODParticleRangeId,
 	const TSharedPtr<Softs::FMultiResConstraints>& FineLODMultiResConstraint,
 	const TArray<const FClothComplexColliders*>& ComplexColliders,
-	const TSharedPtr<const FManagedArrayCollection>& ManagedArrayCollection)
+	const TSharedPtr<const FManagedArrayCollection>& ManagedArrayCollection,
+	const TMap<FName, FClothingAccessoryMeshData>* AccessoryMeshes)
 {
 	// Self collisions
 	CreateSelfCollisionConstraints(ConfigProperties, WeightMaps, VertexSets, FaceSets, FaceIntMaps, TriangleMesh);
@@ -812,7 +814,7 @@ void FClothConstraints::AddRules(
 	CreateMaxDistanceConstraints(ConfigProperties, WeightMaps, MeshScale);
 
 	// Backstop Constraints
-	CreateBackstopConstraints(ConfigProperties, WeightMaps, MeshScale);
+	CreateBackstopConstraints(ConfigProperties, WeightMaps, MeshScale, AccessoryMeshes);
 
 	// Animation Drive Constraints
 	CreateAnimDriveConstraints(ConfigProperties, WeightMaps);
@@ -1382,7 +1384,8 @@ void FClothConstraints::CreateMaxDistanceConstraints(
 void FClothConstraints::CreateBackstopConstraints(
 	const Softs::FCollectionPropertyConstFacade& ConfigProperties,
 	const TMap<FString, TConstArrayView<FRealSingle>>& WeightMaps,
-	Softs::FSolverReal MeshScale)
+	Softs::FSolverReal MeshScale,
+	const TMap<FName, FClothingAccessoryMeshData>* AccessoryMeshes)
 {
 	if (ConfigProperties.GetValue(FClothingSimulationConfig::UseLegacyConfigName, false))
 	{
@@ -1401,14 +1404,35 @@ void FClothConstraints::CreateBackstopConstraints(
 
 	if (Softs::FPBDSphericalBackstopConstraint::IsEnabled(ConfigProperties))
 	{
+		// Get name of backstop mesh
+		const TArray<Softs::FSolverVec3>* BackstopPositions = AnimationPositions;
+		const TArray<Softs::FSolverVec3>* BackstopNormals = AnimationNormals;
+		bool bUseGlobalIndexation = true;
+
+		if (AccessoryMeshes)
+		{
+			const FName BackstopMeshName = FName(Softs::FPBDSphericalBackstopConstraint::GetBackstopMeshNameString(ConfigProperties, TEXT("")));
+			if (const FClothingAccessoryMeshData* const BackstopMeshData = AccessoryMeshes->Find(BackstopMeshName))
+			{
+				if (BackstopMeshData->GetInterpolatedAnimationPositions().Num() == NumParticles && BackstopMeshData->GetInterpolatedAnimationNormals().Num() == NumParticles)
+				{
+					BackstopPositions = &BackstopMeshData->GetInterpolatedAnimationPositions();
+					BackstopNormals = &BackstopMeshData->GetInterpolatedAnimationNormals();
+					bUseGlobalIndexation = false;
+				}
+			}
+		}
+
+
 		BackstopConstraints = MakeShared<Softs::FPBDSphericalBackstopConstraint>(
 			ParticleOffset,
 			NumParticles,
-			*AnimationPositions,
-			*AnimationNormals,
+			*BackstopPositions,
+			*BackstopNormals,
 			WeightMaps,
 			ConfigProperties,
-			MeshScale);
+			MeshScale,
+			bUseGlobalIndexation);
 
 		++NumConstraintRules;
 	}
diff --git a/Engine/Plugins/ChaosCloth/Source/ChaosCloth/Private/ChaosCloth/ChaosClothVisualization.cpp b/Engine/Plugins/ChaosCloth/Source/ChaosCloth/Private/ChaosCloth/ChaosClothVisualization.cpp
--- a/Engine/Plugins/ChaosCloth/Source/ChaosCloth/Private/ChaosCloth/ChaosClothVisualization.cpp
+++ b/Engine/Plugins/ChaosCloth/Source/ChaosCloth/Private/ChaosCloth/ChaosClothVisualization.cpp
@@ -1216,6 +1216,83 @@ namespace Chaos
 		}
 	}
 
+	void FClothVisualizationNoGC::DrawAccessoryMesh(FPrimitiveDrawInterface* PDI, const FName& AccessoryMeshName) const
+	{
+		if (!Solver)
+		{
+			return;
+		}
+
+		static const FLinearColor AccessoryColor = FColor::Turquoise;
+
+		for (const FClothingSimulationCloth* const Cloth : Solver->GetCloths())
+		{
+			const int32 ParticleRangeId = Cloth->GetParticleRangeId(Solver);
+			if (ParticleRangeId == INDEX_NONE)
+			{
+				continue;
+			}
+			if (!Cloth->AccessoryMeshExists(Solver, AccessoryMeshName))
+			{
+				continue;
+			}
+
+			// Elements are local indexed for new solver
+			const int32 Offset = Solver->IsLegacySolver() ? ParticleRangeId : 0;
+
+			const TConstArrayView<TVec3<int32>> Elements = Cloth->GetTriangleMesh(Solver).GetElements();
+			const TConstArrayView<Softs::FSolverVec3> Positions = Cloth->GetAccessoryMeshPositions(Solver, AccessoryMeshName);
+
+			for (int32 ElementIndex = 0; ElementIndex < Elements.Num(); ++ElementIndex)
+			{
+				const TVec3<int32>& Element = Elements[ElementIndex];
+
+				const FVector Pos0 = GetWorldPosition(*Solver, Positions[Element.X - Offset]); // TODO: Triangle Mesh shouldn't really be solver dependent (ie not use an offset)
+				const FVector Pos1 = GetWorldPosition(*Solver, Positions[Element.Y - Offset]);
+				const FVector Pos2 = GetWorldPosition(*Solver, Positions[Element.Z - Offset]);
+
+				DrawLine(PDI, Pos0, Pos1, AccessoryColor);
+				DrawLine(PDI, Pos1, Pos2, AccessoryColor);
+				DrawLine(PDI, Pos2, Pos0, AccessoryColor);
+			}
+		}
+	}
+
+
+	void FClothVisualizationNoGC::DrawAccessoryMeshNormals(FPrimitiveDrawInterface* PDI, const FName& AccessoryMeshName, const FReal NormalLength) const
+	{
+		if (!Solver)
+		{
+			return;
+		}
+
+		static const FLinearColor AccessoryColor = FColor::Turquoise;
+
+		for (const FClothingSimulationCloth* const Cloth : Solver->GetCloths())
+		{
+			if (Cloth->GetParticleRangeId(Solver) == INDEX_NONE)
+			{
+				continue;
+			}
+			if (!Cloth->AccessoryMeshExists(Solver, AccessoryMeshName))
+			{
+				continue;
+			}
+
+			const TConstArrayView<Softs::FSolverVec3> Positions = Cloth->GetAccessoryMeshPositions(Solver, AccessoryMeshName);
+			const TConstArrayView<Softs::FSolverVec3> Normals = Cloth->GetAccessoryMeshNormals(Solver, AccessoryMeshName);
+			check(Normals.Num() == Positions.Num());
+
+			for (int32 Index = 0; Index < Positions.Num(); ++Index)
+			{
+				const FVector Pos0 = GetWorldPosition(*Solver, Positions[Index]);
+				const FVector Pos1 = Pos0 + GetWorldVector(*Solver, Normals[Index]) * NormalLength;
+
+				DrawLine(PDI, Pos0, Pos1, AccessoryColor);
+			}
+		}
+	}
+
 	void FClothVisualizationNoGC::DrawOpenEdges(FPrimitiveDrawInterface* PDI) const
 	{
 		auto MakeSortedUintVector2 = [](uint32 Index0, uint32 Index1) -> FUintVector2
@@ -1721,9 +1798,11 @@ namespace Chaos
 			if (const Softs::FPBDSphericalBackstopConstraint* const BackstopConstraint = ClothConstraints.GetBackstopConstraints().Get())
 			{
 				const bool bUseLegacyBackstop = BackstopConstraint->UseLegacyBackstop();
+				const FName BackstopMeshName = FName(Cloth->GetPropertyString(Solver, Softs::FPBDSphericalBackstopConstraint::BackstopMeshNameName));
+				const bool AccessoryMeshExists = Cloth->AccessoryMeshExists(Solver, BackstopMeshName);
 
-				const TConstArrayView<Softs::FSolverVec3> AnimationPositions = Cloth->GetAnimationPositions(Solver);
-				const TConstArrayView<Softs::FSolverVec3> AnimationNormals = Cloth->GetAnimationNormals(Solver);
+				const TConstArrayView<Softs::FSolverVec3> AnimationPositions = AccessoryMeshExists ? Cloth->GetAccessoryMeshPositions(Solver, BackstopMeshName) : Cloth->GetAnimationPositions(Solver);
+				const TConstArrayView<Softs::FSolverVec3> AnimationNormals = AccessoryMeshExists ? Cloth->GetAccessoryMeshNormals(Solver, BackstopMeshName) : Cloth->GetAnimationNormals(Solver);
 				const TConstArrayView<Softs::FSolverVec3> ParticlePositions = Cloth->GetParticlePositions(Solver);
 
 				if (AnimationPositions.IsValidIndex(Private::BackstopVertexIndex))
@@ -1811,8 +1890,11 @@ namespace Chaos
 			if (const Softs::FPBDSphericalBackstopConstraint* const BackstopConstraint = ClothConstraints.GetBackstopConstraints().Get())
 			{
 				const bool bUseLegacyBackstop = BackstopConstraint->UseLegacyBackstop();
-				const TConstArrayView<Softs::FSolverVec3> AnimationPositions = Cloth->GetAnimationPositions(Solver);
-				const TConstArrayView<Softs::FSolverVec3> AnimationNormals = Cloth->GetAnimationNormals(Solver);
+				const FName BackstopMeshName = FName(Cloth->GetPropertyString(Solver, Softs::FPBDSphericalBackstopConstraint::BackstopMeshNameName));
+				const bool AccessoryMeshExists = Cloth->AccessoryMeshExists(Solver, BackstopMeshName);
+
+				const TConstArrayView<Softs::FSolverVec3> AnimationPositions = AccessoryMeshExists ? Cloth->GetAccessoryMeshPositions(Solver, BackstopMeshName) : Cloth->GetAnimationPositions(Solver);
+				const TConstArrayView<Softs::FSolverVec3> AnimationNormals = AccessoryMeshExists ? Cloth->GetAccessoryMeshNormals(Solver, BackstopMeshName) : Cloth->GetAnimationNormals(Solver);
 
 				for (int32 Index = 0; Index < AnimationPositions.Num(); ++Index)
 				{
@@ -3224,6 +3306,16 @@ namespace Chaos
 		return AllNames.Array();
 	}
 
+	TArray<FName> FClothVisualizationNoGC::GetAllAccessoryMeshNames() const
+	{
+		TSet<FName> AllNames;
+		for (const FClothingSimulationCloth* const Cloth : Solver->GetCloths())
+		{
+			AllNames.Append(Cloth->GetAllAccessoryMeshNames());
+		}
+		return AllNames.Array();
+	}
+
 	void FClothVisualizationNoGC::DrawKinematicColliderShaded(FPrimitiveDrawInterface* PDI) const
 	{
 		if (!Solver || !FMaterials::GetInstance().GetCollisionMaterial())
diff --git a/Engine/Plugins/ChaosCloth/Source/ChaosCloth/Private/ChaosCloth/ChaosClothingAccessoryMeshData.cpp b/Engine/Plugins/ChaosCloth/Source/ChaosCloth/Private/ChaosCloth/ChaosClothingAccessoryMeshData.cpp
new file mode 100644
--- /dev/null
+++ b/Engine/Plugins/ChaosCloth/Source/ChaosCloth/Private/ChaosCloth/ChaosClothingAccessoryMeshData.cpp
@@ -0,0 +1,108 @@
+// Copyright Epic Games, Inc. All Rights Reserved.
+#include "ChaosCloth/ChaosClothingAccessoryMeshData.h"
+#include "ChaosCloth/ChaosClothingSimulationAccessoryMesh.h"
+#if INTEL_ISPC
+#include "ChaosCloth/ChaosClothingSimulationSolver.h"
+#include "ChaosClothingSimulationSolver.ispc.generated.h"
+#endif
+
+namespace Chaos
+{
+	FClothingAccessoryMeshData::FClothingAccessoryMeshData(
+		int32 InNumParticles,
+		const FClothingSimulationAccessoryMesh& InAccessoryMesh)
+		: NumParticles(InNumParticles)
+		, AccessoryMesh(InAccessoryMesh)
+	{
+		check(NumParticles == AccessoryMesh.GetNumPoints());
+	}
+
+	void FClothingAccessoryMeshData::AllocateParticles()
+	{
+		OldAnimationPositions.SetNum(NumParticles);
+		AnimationPositions.SetNum(NumParticles);
+		InterpolatedAnimationPositions.SetNum(NumParticles);
+		OldAnimationNormals.SetNum(NumParticles);
+		AnimationNormals.SetNum(NumParticles);
+		InterpolatedAnimationNormals.SetNum(NumParticles);
+		AnimationVelocities.SetNum(NumParticles);
+	}
+
+	void FClothingAccessoryMeshData::DeallocateParticles()
+	{
+		OldAnimationPositions.Empty();
+		AnimationPositions.Empty();
+		InterpolatedAnimationPositions.Empty();
+		OldAnimationNormals.Empty();
+		AnimationNormals.Empty();
+		InterpolatedAnimationNormals.Empty();
+		AnimationVelocities.Empty();
+	}
+
+	void FClothingAccessoryMeshData::ResetStartPose()
+	{
+		OldAnimationPositions = InterpolatedAnimationPositions = AnimationPositions;
+		OldAnimationNormals = InterpolatedAnimationNormals = AnimationNormals;
+
+		AnimationVelocities.Reset();
+		AnimationVelocities.SetNumZeroed(NumParticles);
+	}
+
+	void FClothingAccessoryMeshData::Update(const FClothingSimulationSolver* Solver)
+	{
+		if (AnimationPositions.Num() == NumParticles)
+		{
+			AccessoryMesh.Update(Solver, TArrayView<Softs::FSolverVec3>(AnimationPositions), TArrayView<Softs::FSolverVec3>(AnimationNormals));
+		}
+	}
+
+	void FClothingAccessoryMeshData::ApplyPreSimulationTransforms(const Softs::FSolverVec3& DeltaLocalSpaceLocation, const Softs::FSolverRigidTransform3& GroupSpaceTransform, const Softs::FSolverReal DeltaTime)
+	{
+		TRACE_CPUPROFILER_EVENT_SCOPE(FClothingAccessoryMeshData_ApplyPreSimulationTransforms);
+		if (bRealTypeCompatibleWithISPC && bChaos_PreSimulationTransforms_ISPC_Enabled)
+		{
+			ispc::ApplyAnimationPreSimulationTransform(
+				(ispc::FVector3f*)OldAnimationPositions.GetData(),
+				(ispc::FVector3f*)AnimationVelocities.GetData(),
+				(const ispc::FVector3f*)AnimationPositions.GetData(),
+				(const ispc::FTransform3f&)GroupSpaceTransform,
+				(const ispc::FVector3f&)DeltaLocalSpaceLocation,
+				DeltaTime,
+				OldAnimationPositions.Num());
+		}
+		else
+		{
+			for (int32 Index = 0; Index < OldAnimationPositions.Num(); ++Index)
+			{
+				OldAnimationPositions[Index] = GroupSpaceTransform.TransformPositionNoScale(OldAnimationPositions[Index]) - DeltaLocalSpaceLocation;
+				AnimationVelocities[Index] = (AnimationPositions[Index] - OldAnimationPositions[Index]) / DeltaTime;
+			}
+		}
+	}
+
+	void FClothingAccessoryMeshData::PreSubstep(const Softs::FSolverReal InterpolationAlpha)
+	{
+		TRACE_CPUPROFILER_EVENT_SCOPE(FClothingAccessoryMeshData_PreSubstep);
+		if (bRealTypeCompatibleWithISPC && bChaos_PreSubstepInterpolation_ISPC_Enabled)
+		{
+			ispc::PreSubstepInterpolation(
+				(ispc::FVector3f*)InterpolatedAnimationPositions.GetData(),
+				(ispc::FVector3f*)InterpolatedAnimationNormals.GetData(),
+				(const ispc::FVector3f*)AnimationPositions.GetData(),
+				(const ispc::FVector3f*)OldAnimationPositions.GetData(),
+				(const ispc::FVector3f*)AnimationNormals.GetData(),
+				(const ispc::FVector3f*)OldAnimationNormals.GetData(),
+				InterpolationAlpha,
+				0,
+				InterpolatedAnimationPositions.Num());
+		}
+		else
+		{
+			for (int32 Index = 0; Index < InterpolatedAnimationPositions.Num(); ++Index)
+			{
+				InterpolatedAnimationPositions[Index] = InterpolationAlpha * AnimationPositions[Index] + ((Softs::FSolverReal)1. - InterpolationAlpha) * OldAnimationPositions[Index];
+				InterpolatedAnimationNormals[Index] = (InterpolationAlpha * AnimationNormals[Index] + ((Softs::FSolverReal)1. - InterpolationAlpha) * OldAnimationNormals[Index]).GetSafeNormal();
+			}
+		}
+	}
+} // namespace Chaos
diff --git a/Engine/Plugins/ChaosCloth/Source/ChaosCloth/Private/ChaosCloth/ChaosClothingAccessoryMeshData.h b/Engine/Plugins/ChaosCloth/Source/ChaosCloth/Private/ChaosCloth/ChaosClothingAccessoryMeshData.h
new file mode 100644
--- /dev/null
+++ b/Engine/Plugins/ChaosCloth/Source/ChaosCloth/Private/ChaosCloth/ChaosClothingAccessoryMeshData.h
@@ -0,0 +1,65 @@
+// Copyright Epic Games, Inc. All Rights Reserved.
+#pragma once
+
+#include "Chaos/Core.h"
+#include "Chaos/PBDSoftsEvolutionFwd.h"
+#include "Containers/ContainersFwd.h"
+namespace Chaos
+{
+class FClothingSimulationAccessoryMesh;
+class FClothingSimulationSolver;
+
+struct FClothingAccessoryMeshData
+{
+	const int32 NumParticles;
+	const FClothingSimulationAccessoryMesh& AccessoryMesh;
+
+	FClothingAccessoryMeshData(
+		int32 InNumParticles,
+		const FClothingSimulationAccessoryMesh& InAccessoryMesh);
+
+	void AllocateParticles();
+	void DeallocateParticles();
+	void ResetStartPose();
+	void Update(const FClothingSimulationSolver* Solver);
+	void ApplyPreSimulationTransforms(const Softs::FSolverVec3& DeltaLocalSpaceLocation, const Softs::FSolverRigidTransform3& GroupSpaceTransform, const Softs::FSolverReal DeltaTime);
+	void PreSubstep(const Softs::FSolverReal InterpolationAlpha);
+
+	const TArray<Softs::FSolverVec3>& GetOldAnimationPositions() const
+	{
+		return OldAnimationPositions;
+	}
+	const TArray<Softs::FSolverVec3>& GetAnimationPositions() const
+	{
+		return AnimationPositions;
+	}
+	const TArray<Softs::FSolverVec3>& GetInterpolatedAnimationPositions() const
+	{
+		return InterpolatedAnimationPositions;
+	}
+	const TArray<Softs::FSolverVec3>& GetOldAnimationNormals() const
+	{
+		return OldAnimationNormals;
+	}
+	const TArray<Softs::FSolverVec3>& GetAnimationNormals() const
+	{
+		return AnimationNormals;
+	}
+	const TArray<Softs::FSolverVec3>& GetInterpolatedAnimationNormals() const
+	{
+		return InterpolatedAnimationNormals;
+	}
+	const TArray<Softs::FSolverVec3>& GetAnimationVelocities() const
+	{
+		return AnimationVelocities;
+	}
+private:
+	TArray<Softs::FSolverVec3> OldAnimationPositions;
+	TArray<Softs::FSolverVec3> AnimationPositions;
+	TArray<Softs::FSolverVec3> InterpolatedAnimationPositions;
+	TArray<Softs::FSolverVec3> OldAnimationNormals;
+	TArray<Softs::FSolverVec3> AnimationNormals;
+	TArray<Softs::FSolverVec3> InterpolatedAnimationNormals;
+	TArray<Softs::FSolverVec3> AnimationVelocities;
+};
+} // namespace Chaos
diff --git a/Engine/Plugins/ChaosCloth/Source/ChaosCloth/Private/ChaosCloth/ChaosClothingSimulationAccessoryMesh.cpp b/Engine/Plugins/ChaosCloth/Source/ChaosCloth/Private/ChaosCloth/ChaosClothingSimulationAccessoryMesh.cpp
new file mode 100644
--- /dev/null
+++ b/Engine/Plugins/ChaosCloth/Source/ChaosCloth/Private/ChaosCloth/ChaosClothingSimulationAccessoryMesh.cpp
@@ -0,0 +1,205 @@
+// Copyright Epic Games, Inc. All Rights Reserved.
+#include "ChaosCloth/ChaosClothingSimulationAccessoryMesh.h"
+#include "ChaosCloth/ChaosClothingSimulationMesh.h"
+#include "ChaosCloth/ChaosClothingSimulationSolver.h"
+#include "ChaosCloth/ChaosClothPrivate.h"
+#include "ClothingSimulation.h"
+#include "ClothVertBoneData.h"
+#include "Components/SkinnedMeshComponent.h"
+#include "Containers/ArrayView.h"
+#include "Async/ParallelFor.h"
+#if INTEL_ISPC
+#include "ChaosClothingSimulationMesh.ispc.generated.h"
+#endif
+
+#if INTEL_ISPC
+#if !UE_BUILD_SHIPPING || USE_ISPC_KERNEL_CONSOLE_VARIABLES_IN_SHIPPING
+bool bChaos_SkinPhysicsMesh_ISPC_Enabled = CHAOS_SKIN_PHYSICS_MESH_ISPC_ENABLED_DEFAULT;
+FAutoConsoleVariableRef CVarChaosSkinPhysicsMeshISPCEnabled(TEXT("p.Chaos.SkinPhysicsMesh.ISPC"), bChaos_SkinPhysicsMesh_ISPC_Enabled, TEXT("Whether to use ISPC optimizations on skinned physics meshes"));
+#endif
+
+static_assert(sizeof(ispc::FVector3f) == sizeof(Chaos::Softs::FSolverVec3), "sizeof(ispc::FVector3f) != sizeof(Chaos::Softs::FSolverVec3)");
+static_assert(sizeof(ispc::FVector3f) == sizeof(FVector3f), "sizeof(ispc::FVector3f) != sizeof(FVector3f)");
+static_assert(sizeof(ispc::FMatrix44f) == sizeof(FMatrix44f), "sizeof(ispc::FMatrix44f) != sizeof(FMatrix44f)");
+static_assert(sizeof(ispc::FTransform3f) == sizeof(FTransform3f), "sizeof(ispc::FTransform3f) != sizeof(FTransform3f)");
+static_assert(sizeof(ispc::FClothVertBoneData) == sizeof(FClothVertBoneData), "sizeof(ispc::FClothVertBoneData) != sizeof(Chaos::FClothVertBoneData)");
+#endif
+
+DECLARE_CYCLE_STAT(TEXT("Chaos Cloth Skin Physics Mesh"), STAT_ChaosClothSkinPhysicsMesh, STATGROUP_ChaosCloth);
+
+namespace Chaos
+{
+
+namespace Private
+{
+	template<typename InOutVectorType>
+	void ApplyMorphTargetInternal(
+		float ActiveMorphTargetWeight,
+		const TConstArrayView<FVector3f>& MorphTargetPositionDeltas,
+		const TConstArrayView<FVector3f>& MorphTargetTangentZDeltas,
+		const TConstArrayView<int32>& MorphTargetIndices,
+		const TArrayView<InOutVectorType>& InOutPositions,
+		const TArrayView<InOutVectorType>& InOutNormals)
+	{
+		TRACE_CPUPROFILER_EVENT_SCOPE(FClothingSimulationMesh_ApplyMorphTargetInternal);
+
+		for (int32 Index = 0; Index < MorphTargetIndices.Num(); ++Index)
+		{
+			const int32 VertexIndex = MorphTargetIndices[Index];
+			InOutPositions[VertexIndex] += ActiveMorphTargetWeight * MorphTargetPositionDeltas[Index];
+			InOutNormals[VertexIndex] = (InOutNormals[VertexIndex] + ActiveMorphTargetWeight * MorphTargetTangentZDeltas[Index]).GetSafeNormal();
+		}
+	}
+
+	// Inline function used to force the unrolling of the skinning loop, LWC: note skinning is all done in float to match the asset data type
+	FORCEINLINE static void AddInfluence(FVector3f& OutPosition, FVector3f& OutNormal, const FVector3f& RefParticle, const FVector3f& RefNormal, const FMatrix44f& BoneMatrix, const float Weight)
+	{
+		OutPosition += BoneMatrix.TransformPosition(RefParticle) * Weight;
+		OutNormal += BoneMatrix.TransformVector(RefNormal) * Weight;
+	}
+}
+
+FClothingSimulationAccessoryMesh::FClothingSimulationAccessoryMesh(const FClothingSimulationMesh& InMesh, const FName& InAccessoryMeshName)
+	: Mesh(InMesh)
+	, AccessoryMeshName(InAccessoryMeshName)
+{
+}
+
+FClothingSimulationAccessoryMesh::~FClothingSimulationAccessoryMesh() = default;
+
+void FClothingSimulationAccessoryMesh::ApplyMorphTarget(
+	int32 ActiveMorphTargetIndex,
+	float ActiveMorphTargetWeight,
+	const TArrayView<Softs::FSolverVec3>& InOutPositions,
+	const TArrayView<Softs::FSolverVec3>& InOutNormals) const
+{
+	const TConstArrayView<FVector3f> MorphTargetPositionDeltas = GetMorphTargetPositionDeltas(ActiveMorphTargetIndex);
+	const TConstArrayView<FVector3f> MorphTargetTangentZDeltas = GetMorphTargetTangentZDeltas(ActiveMorphTargetIndex);
+	const TConstArrayView<int32> MorphTargetIndices = GetMorphTargetIndices(ActiveMorphTargetIndex);
+
+	if (!FMath::IsNearlyZero(ActiveMorphTargetWeight) && !MorphTargetPositionDeltas.IsEmpty() &&
+		MorphTargetPositionDeltas.Num() == MorphTargetTangentZDeltas.Num() &&
+		MorphTargetPositionDeltas.Num() == MorphTargetIndices.Num())
+	{
+		Private::ApplyMorphTargetInternal(ActiveMorphTargetWeight, MorphTargetPositionDeltas, MorphTargetTangentZDeltas, MorphTargetIndices, InOutPositions, InOutNormals);
+	}
+}
+
+void FClothingSimulationAccessoryMesh::SkinPhysicsMesh(int32 ActiveMorphTargetIndex, float ActiveMorphTargetWeight, const FReal LocalSpaceScale, const FVec3& LocalSpaceLocation, const TArrayView<Softs::FSolverVec3>& OutPositions, const TArrayView<Softs::FSolverVec3>& OutNormals) const
+{
+	TRACE_CPUPROFILER_EVENT_SCOPE(FClothingSimulationMesh_SkinPhysicsMesh);
+	SCOPE_CYCLE_COUNTER(STAT_ChaosClothSkinPhysicsMesh);
+	SCOPE_CYCLE_COUNTER(STAT_ClothSkinPhysMesh);
+
+	FTransform ComponentToLocalSpaceReal = Mesh.GetComponentToWorldTransform();
+	ComponentToLocalSpaceReal.AddToTranslation(-LocalSpaceLocation);
+	check(LocalSpaceScale > UE_SMALL_NUMBER);
+	const FReal LocalSpaceScaleInv = 1. / LocalSpaceScale;
+	ComponentToLocalSpaceReal.MultiplyScale3D(FVec3(LocalSpaceScaleInv));
+	ComponentToLocalSpaceReal.ScaleTranslation(LocalSpaceScaleInv);
+	const FTransform3f ComponentToLocalSpace(ComponentToLocalSpaceReal);  // LWC: Now in local space, therefore it is safe to use single precision which is the asset data format
+
+	const int32* const RESTRICT BoneMap = Mesh.GetBoneMap().GetData();
+	const FMatrix44f* const RESTRICT RefToLocalMatrices = Mesh.GetRefToLocalMatrices().GetData();
+
+	const uint32 NumPoints = GetNumPoints();
+	check(NumPoints == OutPositions.Num());
+	check(NumPoints == OutNormals.Num());
+	const TConstArrayView<FClothVertBoneData> BoneData = GetBoneData();
+	TConstArrayView<FVector3f> Positions = GetPositions();
+	TConstArrayView<FVector3f> Normals = GetNormals();
+	TArray<FVector3f> WritablePositions;
+	TArray<FVector3f> WritableNormals;
+
+	const TConstArrayView<FVector3f> MorphTargetPositionDeltas = GetMorphTargetPositionDeltas(ActiveMorphTargetIndex);
+	const TConstArrayView<FVector3f> MorphTargetTangentZDeltas = GetMorphTargetTangentZDeltas(ActiveMorphTargetIndex);
+	const TConstArrayView<int32> MorphTargetIndices = GetMorphTargetIndices(ActiveMorphTargetIndex);
+
+	if (!FMath::IsNearlyZero(ActiveMorphTargetWeight) && !MorphTargetPositionDeltas.IsEmpty() &&
+		MorphTargetPositionDeltas.Num() == MorphTargetTangentZDeltas.Num() &&
+		MorphTargetPositionDeltas.Num() == MorphTargetIndices.Num())
+	{
+		TRACE_CPUPROFILER_EVENT_SCOPE(FClothingSimulationMesh_SkinPhysicsMesh_MorphTargets);
+		// TODO optimize this
+		WritablePositions = Positions;
+		WritableNormals = Normals;
+		Private::ApplyMorphTargetInternal(ActiveMorphTargetWeight, MorphTargetPositionDeltas, MorphTargetTangentZDeltas, MorphTargetIndices, TArrayView<FVector3f>(WritablePositions), TArrayView<FVector3f>(WritableNormals));
+		Positions = WritablePositions;
+		Normals = WritableNormals;
+	}
+
+#if INTEL_ISPC
+	if (bChaos_SkinPhysicsMesh_ISPC_Enabled)
+	{
+		ispc::SkinPhysicsMesh(
+			(ispc::FVector3f*)OutPositions.GetData(),
+			(ispc::FVector3f*)OutNormals.GetData(),
+			(ispc::FVector3f*)Positions.GetData(),
+			(ispc::FVector3f*)Normals.GetData(),
+			(ispc::FClothVertBoneData*)BoneData.GetData(),
+			BoneMap,
+			(ispc::FMatrix44f*)RefToLocalMatrices,
+			(ispc::FTransform3f&)ComponentToLocalSpace,
+			NumPoints);
+	}
+	else
+#endif
+	{
+		static const uint32 MinParallelVertices = 500;  // 500 seems to be the lowest threshold still giving gains even on profiled assets that are only using a small number of influences
+
+		ParallelFor(NumPoints, [&BoneData, &Positions, &Normals, &ComponentToLocalSpace, BoneMap, RefToLocalMatrices, &OutPositions, &OutNormals](uint32 VertIndex)
+		{
+			const uint16* const RESTRICT BoneIndices = BoneData[VertIndex].BoneIndices;
+			const float* const RESTRICT BoneWeights = BoneData[VertIndex].BoneWeights;
+	
+			// WARNING - HORRIBLE UNROLLED LOOP + JUMP TABLE BELOW
+			// done this way because this is a pretty tight and performance critical loop. essentially
+			// rather than checking each influence we can just jump into this switch and fall through
+			// everything to compose the final skinned data
+			const FVector3f& RefParticle = Positions[VertIndex];
+			const FVector3f& RefNormal = Normals[VertIndex];
+
+			FVector3f Position(ForceInitToZero);
+			FVector3f Normal(ForceInitToZero);
+			switch (BoneData[VertIndex].NumInfluences)
+			{
+			default:  // Intentional fall through
+			case 12: Private::AddInfluence(Position, Normal, RefParticle, RefNormal, RefToLocalMatrices[BoneMap[BoneIndices[11]]], BoneWeights[11]);  // Intentional fall through
+			case 11: Private::AddInfluence(Position, Normal, RefParticle, RefNormal, RefToLocalMatrices[BoneMap[BoneIndices[10]]], BoneWeights[10]);  // Intentional fall through
+			case 10: Private::AddInfluence(Position, Normal, RefParticle, RefNormal, RefToLocalMatrices[BoneMap[BoneIndices[ 9]]], BoneWeights[ 9]);  // Intentional fall through
+			case  9: Private::AddInfluence(Position, Normal, RefParticle, RefNormal, RefToLocalMatrices[BoneMap[BoneIndices[ 8]]], BoneWeights[ 8]);  // Intentional fall through
+			case  8: Private::AddInfluence(Position, Normal, RefParticle, RefNormal, RefToLocalMatrices[BoneMap[BoneIndices[ 7]]], BoneWeights[ 7]);  // Intentional fall through
+			case  7: Private::AddInfluence(Position, Normal, RefParticle, RefNormal, RefToLocalMatrices[BoneMap[BoneIndices[ 6]]], BoneWeights[ 6]);  // Intentional fall through
+			case  6: Private::AddInfluence(Position, Normal, RefParticle, RefNormal, RefToLocalMatrices[BoneMap[BoneIndices[ 5]]], BoneWeights[ 5]);  // Intentional fall through
+			case  5: Private::AddInfluence(Position, Normal, RefParticle, RefNormal, RefToLocalMatrices[BoneMap[BoneIndices[ 4]]], BoneWeights[ 4]);  // Intentional fall through
+			case  4: Private::AddInfluence(Position, Normal, RefParticle, RefNormal, RefToLocalMatrices[BoneMap[BoneIndices[ 3]]], BoneWeights[ 3]);  // Intentional fall through
+			case  3: Private::AddInfluence(Position, Normal, RefParticle, RefNormal, RefToLocalMatrices[BoneMap[BoneIndices[ 2]]], BoneWeights[ 2]);  // Intentional fall through
+			case  2: Private::AddInfluence(Position, Normal, RefParticle, RefNormal, RefToLocalMatrices[BoneMap[BoneIndices[ 1]]], BoneWeights[ 1]);  // Intentional fall through
+			case  1: Private::AddInfluence(Position, Normal, RefParticle, RefNormal, RefToLocalMatrices[BoneMap[BoneIndices[ 0]]], BoneWeights[ 0]);  // Intentional fall through
+			case  0: break;
+			}
+
+			OutPositions[VertIndex] = FSolverVec3(ComponentToLocalSpace.TransformPosition(Position));
+			OutNormals[VertIndex] = FSolverVec3(ComponentToLocalSpace.TransformVector(Normal).GetSafeNormal());
+
+		}, NumPoints > MinParallelVertices ? EParallelForFlags::None : EParallelForFlags::ForceSingleThread);
+	}
+}
+
+
+void FClothingSimulationAccessoryMesh::Update(
+	const FClothingSimulationSolver* Solver,
+	const TArrayView<Softs::FSolverVec3>& OutPositions,
+	const TArrayView<Softs::FSolverVec3>& OutNormals,
+	int32 ActiveMorphTargetIndex,
+	float ActiveMorphTargetWeight) const
+{
+	check(Solver);
+
+	// Skin current LOD positions
+	const FReal LocalSpaceScale = Solver->GetLocalSpaceScale();
+	const FVec3& LocalSpaceLocation = Solver->GetLocalSpaceLocation();
+
+	SkinPhysicsMesh(ActiveMorphTargetIndex, ActiveMorphTargetWeight, LocalSpaceScale, LocalSpaceLocation, OutPositions, OutNormals);
+}
+}  // End namespace Chaos
diff --git a/Engine/Plugins/ChaosCloth/Source/ChaosCloth/Private/ChaosCloth/ChaosClothingSimulationCloth.cpp b/Engine/Plugins/ChaosCloth/Source/ChaosCloth/Private/ChaosCloth/ChaosClothingSimulationCloth.cpp
--- a/Engine/Plugins/ChaosCloth/Source/ChaosCloth/Private/ChaosCloth/ChaosClothingSimulationCloth.cpp
+++ b/Engine/Plugins/ChaosCloth/Source/ChaosCloth/Private/ChaosCloth/ChaosClothingSimulationCloth.cpp
@@ -1,6 +1,7 @@
 // Copyright Epic Games, Inc. All Rights Reserved.
 
 #include "ChaosCloth/ChaosClothingSimulationCloth.h"
+#include "ChaosCloth/ChaosClothingAccessoryMeshData.h"
 #include "ChaosCloth/ChaosClothingSimulationSolver.h"
 #include "ChaosCloth/ChaosClothingSimulationMesh.h"
 #include "ChaosCloth/ChaosClothingSimulationCollider.h"
@@ -83,6 +84,8 @@ struct FClothingSimulationCloth::FLODData
 		int32 ParticleRangeId;
 		FTriangleMesh OffsetTriangleMesh; // Only used if using PBD solver
 		int32 MultiResCoarseLODIndex = INDEX_NONE;
+
+		TMap<FName, FClothingAccessoryMeshData> AccessoryMeshData;
 	};
 	TMap<FClothingSimulationSolver*, FSolverData> SolverData;
 
@@ -117,7 +120,11 @@ struct FClothingSimulationCloth::FLODData
 
 	void Enable(FClothingSimulationSolver* Solver, bool bEnable) const;
 
-	void ResetStartPose(FClothingSimulationSolver* Solver) const;
+	void UpdateAccessoryMeshes(FClothingSimulationSolver* Solver);
+	void ResetAccessoryMeshes(FClothingSimulationSolver* Solver);
+	void ResetStartPose(FClothingSimulationSolver* Solver);
+	void ApplyPreSimulationTransformsToAccessoryMeshes(FClothingSimulationSolver* Solver, const Softs::FSolverVec3& DeltaLocalSpaceLocation, const Softs::FSolverRigidTransform3& GroupSpaceTransform, const Softs::FSolverReal DeltaTime);
+	void PreSubstepAccessoryMeshes(FClothingSimulationSolver* Solver, const Softs::FSolverReal InterpolationAlpha);
 
 	void UpdateNormals(FClothingSimulationSolver* Solver) const;
 
@@ -224,10 +231,22 @@ void FClothingSimulationCloth::FLODData::AddParticles(FClothingSimulationSolver*
 		SolverDatum.OffsetTriangleMesh.GetPointToTriangleMap(); // Builds map for later use by GetPointNormals(), and the velocity fields
 	}
 
+	TConstArrayView<FName> AccessoryMeshes = Cloth->Mesh->GetAllAccessoryMeshNames(InLODIndex);
+	SolverDatum.AccessoryMeshData.Reserve(AccessoryMeshes.Num());
+	for (const FName& AccessoryMeshName : AccessoryMeshes)
+	{
+		if (const FClothingSimulationAccessoryMesh* const AccessoryMesh = Cloth->Mesh->GetAccessoryMesh(InLODIndex, AccessoryMeshName))
+		{
+			FClothingAccessoryMeshData& MeshData = SolverDatum.AccessoryMeshData.Add(AccessoryMeshName, FClothingAccessoryMeshData(NumParticles, *AccessoryMesh));
+			MeshData.AllocateParticles();
+		}
+	}
+
 	const FTriangleMesh& TriangleMesh = Solver->IsLegacySolver() ? SolverDatum.OffsetTriangleMesh : NoOffsetTriangleMesh;
 
 	// Update source mesh for this LOD, this is required prior to reset the start pose
 	Cloth->Mesh->Update(Solver, INDEX_NONE, InLODIndex, 0, SolverDatum.ParticleRangeId);
+	UpdateAccessoryMeshes(Solver);
 
 	// Reset the particles start pose before setting up mass and constraints
 	ResetStartPose(Solver);
@@ -373,7 +392,7 @@ void FClothingSimulationCloth::FLODData::Add(FClothingSimulationSolver* Solver,
 
 	// Create constraints
 	const bool bEnabled = false;  // Set constraint disabled by default
-	ClothConstraints.AddRules(ConfigProperties, TriangleMesh, &PatternData, WeightMaps, VertexSets, FaceSets, FaceIntMaps, Tethers, MeshScale, bEnabled, CoarseLODTriangleMesh, CoarseLODParticleRangeId, FineLODMultiResConstraint, ComplexColliders, ManagedArrayCollection);
+	ClothConstraints.AddRules(ConfigProperties, TriangleMesh, &PatternData, WeightMaps, VertexSets, FaceSets, FaceIntMaps, Tethers, MeshScale, bEnabled, CoarseLODTriangleMesh, CoarseLODParticleRangeId, FineLODMultiResConstraint, ComplexColliders, ManagedArrayCollection, &SolverDatum.AccessoryMeshData);
 
 	// Update LOD stats
 	const TConstArrayView<Softs::FSolverReal> InvMasses(Solver->GetParticleInvMasses(ParticleRangeId), NumParticles);
@@ -428,13 +447,60 @@ void FClothingSimulationCloth::FLODData::Enable(FClothingSimulationSolver* Solve
 	Solver->EnableParticles(ParticleRangeId, bEnable);
 }
 
-void FClothingSimulationCloth::FLODData::ResetStartPose(FClothingSimulationSolver* Solver) const
+void FClothingSimulationCloth::FLODData::ResetStartPose(FClothingSimulationSolver* Solver)
 {
 	check(Solver);
-	const int32 ParticleRangeId = SolverData.FindChecked(Solver).ParticleRangeId;
+
+	FSolverData& SolverDatum = SolverData.FindChecked(Solver);
+	const int32 ParticleRangeId = SolverDatum.ParticleRangeId;
 	check(ParticleRangeId != INDEX_NONE);
 
 	Solver->ResetStartPose(ParticleRangeId, NumParticles);
+	ResetAccessoryMeshes(Solver);
+}
+
+void FClothingSimulationCloth::FLODData::UpdateAccessoryMeshes(FClothingSimulationSolver* Solver)
+{
+	check(Solver);
+
+	FSolverData& SolverDatum = SolverData.FindChecked(Solver);
+	for (TPair<FName, FClothingAccessoryMeshData>& AccessoryMesh : SolverDatum.AccessoryMeshData)
+	{
+		AccessoryMesh.Get<1>().Update(Solver);
+	}
+}
+
+void FClothingSimulationCloth::FLODData::ResetAccessoryMeshes(FClothingSimulationSolver* Solver)
+{
+	check(Solver);
+
+	FSolverData& SolverDatum = SolverData.FindChecked(Solver);
+	for (TPair<FName, FClothingAccessoryMeshData>& AccessoryMesh : SolverDatum.AccessoryMeshData)
+	{
+		AccessoryMesh.Get<1>().ResetStartPose();
+	}
+}
+
+void FClothingSimulationCloth::FLODData::ApplyPreSimulationTransformsToAccessoryMeshes(FClothingSimulationSolver* Solver, const Softs::FSolverVec3& DeltaLocalSpaceLocation, const Softs::FSolverRigidTransform3& GroupSpaceTransform, const Softs::FSolverReal DeltaTime)
+{
+	check(Solver);
+
+	FSolverData& SolverDatum = SolverData.FindChecked(Solver);
+	for (TPair<FName, FClothingAccessoryMeshData>& AccessoryMesh : SolverDatum.AccessoryMeshData)
+	{
+		AccessoryMesh.Get<1>().ApplyPreSimulationTransforms(DeltaLocalSpaceLocation, GroupSpaceTransform, DeltaTime);
+	}
+}
+
+void FClothingSimulationCloth::FLODData::PreSubstepAccessoryMeshes(FClothingSimulationSolver* Solver, const Softs::FSolverReal InterpolationAlpha)
+{
+	check(Solver);
+
+	FSolverData& SolverDatum = SolverData.FindChecked(Solver);
+	for (TPair<FName, FClothingAccessoryMeshData>& AccessoryMesh : SolverDatum.AccessoryMeshData)
+	{
+		AccessoryMesh.Get<1>().PreSubstep(InterpolationAlpha);
+	}
 }
 
 void FClothingSimulationCloth::FLODData::UpdateNormals(FClothingSimulationSolver* Solver) const
@@ -770,6 +836,12 @@ TConstArrayView<FRealSingle> FClothingSimulationCloth::GetWeightMapByName(const
 	return LODData.IsValidIndex(LODIndex) ? LODData[LODIndex]->WeightMaps.FindRef(Name) : TConstArrayView<FRealSingle>();
 }
 
+FString FClothingSimulationCloth::GetPropertyString(const FClothingSimulationSolver* Solver, const FName& Property) const
+{
+	check(Config);
+	return Config->GetProperties(GetLODIndex(Solver)).GetStringValue(Property);
+}
+
 TConstArrayView<FRealSingle> FClothingSimulationCloth::GetWeightMapByProperty(const FClothingSimulationSolver* Solver, const FName& Property) const
 {
 	check(Config);
@@ -836,6 +908,16 @@ TSet<FString> FClothingSimulationCloth::GetAllMorphTargetNames() const
 	return Names;
 }
 
+TSet<FName> FClothingSimulationCloth::GetAllAccessoryMeshNames() const
+{
+	TSet<FName> Names;
+	for (int32 LODIndex = 0; LODIndex < LODData.Num(); ++LODIndex)
+	{
+		Names.Append(Mesh->GetAllAccessoryMeshNames(LODIndex));
+	}
+	return Names;
+}
+
 void FClothingSimulationCloth::PreUpdate(FClothingSimulationSolver* Solver)
 {
 	check(Solver);
@@ -906,6 +988,7 @@ void FClothingSimulationCloth::Update(FClothingSimulationSolver* Solver)
 		LODData[LODIndex]->UpdateCachedProperties(*Mesh, LODIndex, ConfigProperties);
 		CurrentLODMorphTargetIndex = LODData[LODIndex]->ActiveMorphTarget;
 		CurrentLODMorphTargetWeight = LODData[LODIndex]->ActiveMorphTargetWeight;
+		LODData[LODIndex]->UpdateAccessoryMeshes(Solver);
 	}
 
 	Mesh->Update(Solver, PrevLODIndex, LODIndex, PrevParticleRangeId, ParticleRangeId, CurrentLODMorphTargetIndex, CurrentLODMorphTargetWeight);
@@ -918,6 +1001,7 @@ void FClothingSimulationCloth::Update(FClothingSimulationSolver* Solver)
 		LODData[CoarseLODIndex]->Enable(Solver, true);
 		// TODO: interpolate/ reset when LOD switching to enable multires
 		Mesh->Update(Solver, CoarseLODIndex, CoarseLODIndex, CoarseParticleRangeId, CoarseParticleRangeId, LODData[CoarseLODIndex]->ActiveMorphTarget, LODData[CoarseLODIndex]->ActiveMorphTargetWeight);
+		LODData[CoarseLODIndex]->UpdateAccessoryMeshes(Solver);
 	}
 
 	// LOD Switching
@@ -955,6 +1039,11 @@ void FClothingSimulationCloth::Update(FClothingSimulationSolver* Solver)
 				Solver->GetParticleXs(ParticleRangeId),
 				Solver->GetParticleVs(ParticleRangeId));
 
+			if (!bNeedsReset)
+			{
+				LODData[LODIndex]->ResetAccessoryMeshes(Solver);
+			}
+
 			if (Solver->IsLegacySolver())
 			{
 				// Update the wind velocity field for the new LOD mesh
@@ -1184,6 +1273,24 @@ void FClothingSimulationCloth::UpdateFromCache(const FClothingSimulationCacheDat
 	}
 }
 
+void FClothingSimulationCloth::ApplyPreSimulationTransforms(FClothingSimulationSolver* Solver, const Softs::FSolverVec3& DeltaLocalSpaceLocation, const Softs::FSolverRigidTransform3& GroupSpaceTransform, const Softs::FSolverReal DeltaTime)
+{
+	const int32 LODIndex = LODIndices.FindChecked(Solver);
+	if (LODIndex != INDEX_NONE)
+	{
+		LODData[LODIndex]->ApplyPreSimulationTransformsToAccessoryMeshes(Solver, DeltaLocalSpaceLocation, GroupSpaceTransform, DeltaTime);
+	}
+}
+
+void FClothingSimulationCloth::PreSubstep(FClothingSimulationSolver* Solver, const Softs::FSolverReal InterpolationAlpha)
+{
+	const int32 LODIndex = LODIndices.FindChecked(Solver);
+	if (LODIndex != INDEX_NONE)
+	{
+		LODData[LODIndex]->PreSubstepAccessoryMeshes(Solver, InterpolationAlpha);
+	}
+}
+
 TConstArrayView<Softs::FSolverVec3> FClothingSimulationCloth::GetAnimationPositions(const FClothingSimulationSolver* Solver) const
 {
 	check(Solver);
@@ -1216,6 +1323,78 @@ TConstArrayView<Softs::FSolverVec3> FClothingSimulationCloth::GetAnimationNormal
 	return TConstArrayView<Softs::FSolverVec3>(Solver->GetAnimationNormals(GetParticleRangeId(Solver, LODIndex)), GetNumParticles(LODIndex));
 }
 
+bool FClothingSimulationCloth::AccessoryMeshExists(const FClothingSimulationSolver* Solver, const FName& AccessoryMeshName) const
+{
+	check(Solver);
+	const int32 LODIndex = LODIndices.FindChecked(Solver);
+	if (LODData.IsValidIndex(LODIndex))
+	{
+		const FLODData::FSolverData& SolverDatum = LODData[LODIndex]->SolverData.FindChecked(Solver);
+		return SolverDatum.AccessoryMeshData.Find(AccessoryMeshName) != nullptr;
+	}
+	return false;
+}
+
+TConstArrayView<Softs::FSolverVec3> FClothingSimulationCloth::GetAccessoryMeshPositions(const FClothingSimulationSolver* Solver, const FName& AccessoryMeshName) const
+{
+	check(Solver);
+	const int32 LODIndex = LODIndices.FindChecked(Solver);
+	if (LODData.IsValidIndex(LODIndex))
+	{
+		const FLODData::FSolverData& SolverDatum = LODData[LODIndex]->SolverData.FindChecked(Solver);
+		if (const FClothingAccessoryMeshData* const AccessoryMesh = SolverDatum.AccessoryMeshData.Find(AccessoryMeshName))
+		{
+			return TConstArrayView<Softs::FSolverVec3>(AccessoryMesh->GetAnimationPositions());
+		}
+	}
+	return TConstArrayView<Softs::FSolverVec3>();
+}
+
+TConstArrayView<Softs::FSolverVec3> FClothingSimulationCloth::GetOldAccessoryMeshPositions(const FClothingSimulationSolver* Solver, const FName& AccessoryMeshName) const
+{
+	check(Solver);
+	const int32 LODIndex = LODIndices.FindChecked(Solver);
+	if (LODData.IsValidIndex(LODIndex))
+	{
+		const FLODData::FSolverData& SolverDatum = LODData[LODIndex]->SolverData.FindChecked(Solver);
+		if (const FClothingAccessoryMeshData* const AccessoryMesh = SolverDatum.AccessoryMeshData.Find(AccessoryMeshName))
+		{
+			return TConstArrayView<Softs::FSolverVec3>(AccessoryMesh->GetOldAnimationPositions());
+		}
+	}
+	return TConstArrayView<Softs::FSolverVec3>();
+}
+
+TConstArrayView<Softs::FSolverVec3> FClothingSimulationCloth::GetAccessoryMeshNormals(const FClothingSimulationSolver* Solver, const FName& AccessoryMeshName) const
+{
+	check(Solver);
+	const int32 LODIndex = LODIndices.FindChecked(Solver);
+	if (LODData.IsValidIndex(LODIndex))
+	{
+		const FLODData::FSolverData& SolverDatum = LODData[LODIndex]->SolverData.FindChecked(Solver);
+		if (const FClothingAccessoryMeshData* const AccessoryMesh = SolverDatum.AccessoryMeshData.Find(AccessoryMeshName))
+		{
+			return TConstArrayView<Softs::FSolverVec3>(AccessoryMesh->GetAnimationNormals());
+		}
+	}
+	return TConstArrayView<Softs::FSolverVec3>();
+}
+
+TConstArrayView<Softs::FSolverVec3> FClothingSimulationCloth::GetAccessoryMeshVelocities(const FClothingSimulationSolver* Solver, const FName& AccessoryMeshName) const
+{
+	check(Solver);
+	const int32 LODIndex = LODIndices.FindChecked(Solver);
+	if (LODData.IsValidIndex(LODIndex))
+	{
+		const FLODData::FSolverData& SolverDatum = LODData[LODIndex]->SolverData.FindChecked(Solver);
+		if (const FClothingAccessoryMeshData* const AccessoryMesh = SolverDatum.AccessoryMeshData.Find(AccessoryMeshName))
+		{
+			return TConstArrayView<Softs::FSolverVec3>(AccessoryMesh->GetAnimationVelocities());
+		}
+	}
+	return TConstArrayView<Softs::FSolverVec3>();
+}
+
 TConstArrayView<Softs::FSolverVec3> FClothingSimulationCloth::GetParticlePositions(const FClothingSimulationSolver* Solver) const
 {
 	check(Solver);
diff --git a/Engine/Plugins/ChaosCloth/Source/ChaosCloth/Private/ChaosCloth/ChaosClothingSimulationDefaultAccessoryMesh.cpp b/Engine/Plugins/ChaosCloth/Source/ChaosCloth/Private/ChaosCloth/ChaosClothingSimulationDefaultAccessoryMesh.cpp
new file mode 100644
--- /dev/null
+++ b/Engine/Plugins/ChaosCloth/Source/ChaosCloth/Private/ChaosCloth/ChaosClothingSimulationDefaultAccessoryMesh.cpp
@@ -0,0 +1,71 @@
+// Copyright Epic Games, Inc. All Rights Reserved.
+#include "ChaosCloth/ChaosClothingSimulationDefaultAccessoryMesh.h"
+#include "ChaosCloth/ChaosClothingSimulationMesh.h"
+
+namespace Chaos
+{
+	const FName FClothingSimulationDefaultAccessoryMesh::DefaultAccessoryMeshName(TEXT("DefaultAccessoryMesh"));
+
+	FClothingSimulationDefaultAccessoryMesh::FClothingSimulationDefaultAccessoryMesh(const FClothingSimulationMesh& InMesh, const int32 InLODIndex)
+		: FClothingSimulationAccessoryMesh(InMesh, DefaultAccessoryMeshName)
+		, LODIndex(InLODIndex)
+	{
+	}
+
+	FClothingSimulationDefaultAccessoryMesh::~FClothingSimulationDefaultAccessoryMesh() = default;
+
+	int32 FClothingSimulationDefaultAccessoryMesh::GetNumPoints() const
+	{
+		return Mesh.GetNumPoints(LODIndex);
+	}
+
+	int32 FClothingSimulationDefaultAccessoryMesh::GetNumPatternPoints() const
+	{
+		return Mesh.GetNumPatternPoints(LODIndex);
+	}
+
+	TConstArrayView<FVector3f> FClothingSimulationDefaultAccessoryMesh::GetPositions() const
+	{
+		return Mesh.GetPositions(LODIndex);
+	}
+
+	TConstArrayView<FVector2f> FClothingSimulationDefaultAccessoryMesh::GetPatternPositions() const
+	{
+		return Mesh.GetPatternPositions(LODIndex);
+	}
+
+	TConstArrayView<FVector3f> FClothingSimulationDefaultAccessoryMesh::GetNormals() const
+	{
+		return Mesh.GetNormals(LODIndex);
+	}
+
+	TConstArrayView<FClothVertBoneData> FClothingSimulationDefaultAccessoryMesh::GetBoneData() const
+	{
+		return Mesh.GetBoneData(LODIndex);
+	}
+
+	int32 FClothingSimulationDefaultAccessoryMesh::FindMorphTargetByName(const FString& Name) const
+	{
+		return Mesh.FindMorphTargetByName(LODIndex, Name);
+	}
+
+	TConstArrayView<FString> FClothingSimulationDefaultAccessoryMesh::GetAllMorphTargetNames() const
+	{
+		return Mesh.GetAllMorphTargetNames(LODIndex);
+	}
+
+	TConstArrayView<FVector3f> FClothingSimulationDefaultAccessoryMesh::GetMorphTargetPositionDeltas(int32 MorphTargetIndex) const
+	{
+		return Mesh.GetMorphTargetPositionDeltas(LODIndex, MorphTargetIndex);
+	}
+
+	TConstArrayView<FVector3f> FClothingSimulationDefaultAccessoryMesh::GetMorphTargetTangentZDeltas(int32 MorphTargetIndex) const
+	{
+		return Mesh.GetMorphTargetTangentZDeltas(LODIndex, MorphTargetIndex);
+	}
+
+	TConstArrayView<int32> FClothingSimulationDefaultAccessoryMesh::GetMorphTargetIndices(int32 MorphTargetIndex) const
+	{
+		return Mesh.GetMorphTargetIndices(LODIndex, MorphTargetIndex);
+	}
+}  // End namespace Chaos
diff --git a/Engine/Plugins/ChaosCloth/Source/ChaosCloth/Private/ChaosCloth/ChaosClothingSimulationMesh.cpp b/Engine/Plugins/ChaosCloth/Source/ChaosCloth/Private/ChaosCloth/ChaosClothingSimulationMesh.cpp
--- a/Engine/Plugins/ChaosCloth/Source/ChaosCloth/Private/ChaosCloth/ChaosClothingSimulationMesh.cpp
+++ b/Engine/Plugins/ChaosCloth/Source/ChaosCloth/Private/ChaosCloth/ChaosClothingSimulationMesh.cpp
@@ -12,24 +12,8 @@
 #include "Containers/ArrayView.h"
 #include "Async/ParallelFor.h"
 #include "SkeletalMeshTypes.h"  // For FMeshToMeshVertData
-#if INTEL_ISPC
-#include "ChaosClothingSimulationMesh.ispc.generated.h"
-#endif
 
-#if INTEL_ISPC
-#if !UE_BUILD_SHIPPING || USE_ISPC_KERNEL_CONSOLE_VARIABLES_IN_SHIPPING
-bool bChaos_SkinPhysicsMesh_ISPC_Enabled = CHAOS_SKIN_PHYSICS_MESH_ISPC_ENABLED_DEFAULT;
-FAutoConsoleVariableRef CVarChaosSkinPhysicsMeshISPCEnabled(TEXT("p.Chaos.SkinPhysicsMesh.ISPC"), bChaos_SkinPhysicsMesh_ISPC_Enabled, TEXT("Whether to use ISPC optimizations on skinned physics meshes"));
-#endif
 
-static_assert(sizeof(ispc::FVector3f) == sizeof(Chaos::Softs::FSolverVec3), "sizeof(ispc::FVector3f) != sizeof(Chaos::Softs::FSolverVec3)");
-static_assert(sizeof(ispc::FVector3f) == sizeof(FVector3f), "sizeof(ispc::FVector3f) != sizeof(FVector3f)");
-static_assert(sizeof(ispc::FMatrix44f) == sizeof(FMatrix44f), "sizeof(ispc::FMatrix44f) != sizeof(FMatrix44f)");
-static_assert(sizeof(ispc::FTransform3f) == sizeof(FTransform3f), "sizeof(ispc::FTransform3f) != sizeof(FTransform3f)");
-static_assert(sizeof(ispc::FClothVertBoneData) == sizeof(FClothVertBoneData), "sizeof(ispc::FClothVertBoneData) != sizeof(Chaos::FClothVertBoneData)");
-#endif
-
-DECLARE_CYCLE_STAT(TEXT("Chaos Cloth Skin Physics Mesh"), STAT_ChaosClothSkinPhysicsMesh, STATGROUP_ChaosCloth);
 DECLARE_CYCLE_STAT(TEXT("Chaos Cloth Wrap Deform Mesh"), STAT_ChaosClothWrapDeformMesh, STATGROUP_ChaosCloth);
 DECLARE_CYCLE_STAT(TEXT("Chaos Cloth Wrap Deform Cloth LOD"), STAT_ChaosClothWrapDeformClothLOD, STATGROUP_ChaosCloth);
 
@@ -180,25 +164,6 @@ bool FClothingSimulationMesh::WrapDeformLOD(
 	return true;
 }
 
-template<typename InOutVectorType>
-void ApplyMorphTargetInternal(
-	float ActiveMorphTargetWeight,
-	const TConstArrayView<FVector3f>& MorphTargetPositionDeltas,
-	const TConstArrayView<FVector3f>& MorphTargetTangentZDeltas,
-	const TConstArrayView<int32>& MorphTargetIndices,
-	const TArrayView<InOutVectorType>& InOutPositions,
-	const TArrayView<InOutVectorType>& InOutNormals)
-{
-	TRACE_CPUPROFILER_EVENT_SCOPE(FClothingSimulationMesh_ApplyMorphTargetInternal);
-
-	for (int32 Index = 0; Index < MorphTargetIndices.Num(); ++Index)
-	{
-		const int32 VertexIndex = MorphTargetIndices[Index];
-		InOutPositions[VertexIndex] += ActiveMorphTargetWeight * MorphTargetPositionDeltas[Index];
-		InOutNormals[VertexIndex] = (InOutNormals[VertexIndex] + ActiveMorphTargetWeight * MorphTargetTangentZDeltas[Index]).GetSafeNormal();
-	}
-}
-
 void FClothingSimulationMesh::ApplyMorphTarget(
 	int32 LODIndex,
 	int32 ActiveMorphTargetIndex,
@@ -206,126 +171,16 @@ void FClothingSimulationMesh::ApplyMorphTarget(
 	const TArrayView<Softs::FSolverVec3>& InOutPositions,
 	const TArrayView<Softs::FSolverVec3>& InOutNormals) const
 {
-	const TConstArrayView<FVector3f> MorphTargetPositionDeltas = GetMorphTargetPositionDeltas(LODIndex, ActiveMorphTargetIndex);
-	const TConstArrayView<FVector3f> MorphTargetTangentZDeltas = GetMorphTargetTangentZDeltas(LODIndex, ActiveMorphTargetIndex);
-	const TConstArrayView<int32> MorphTargetIndices = GetMorphTargetIndices(LODIndex, ActiveMorphTargetIndex);
-
-	if (!FMath::IsNearlyZero(ActiveMorphTargetWeight) && !MorphTargetPositionDeltas.IsEmpty() &&
-		MorphTargetPositionDeltas.Num() == MorphTargetTangentZDeltas.Num() &&
-		MorphTargetPositionDeltas.Num() == MorphTargetIndices.Num())
+	// Exit if any inputs are missing or not ready, and if the LOD is invalid
+	if (!IsValidLODIndex(LODIndex))
 	{
-		ApplyMorphTargetInternal(ActiveMorphTargetWeight, MorphTargetPositionDeltas, MorphTargetTangentZDeltas, MorphTargetIndices, InOutPositions, InOutNormals);
+		return;
 	}
-}
 
-// Inline function used to force the unrolling of the skinning loop, LWC: note skinning is all done in float to match the asset data type
-FORCEINLINE static void AddInfluence(FVector3f& OutPosition, FVector3f& OutNormal, const FVector3f& RefParticle, const FVector3f& RefNormal, const FMatrix44f& BoneMatrix, const float Weight)
-{
-	OutPosition += BoneMatrix.TransformPosition(RefParticle) * Weight;
-	OutNormal += BoneMatrix.TransformVector(RefNormal) * Weight;
+	const FClothingSimulationDefaultAccessoryMesh DefaultMesh = GetDefaultAccessoryMesh(LODIndex);
+	DefaultMesh.ApplyMorphTarget(ActiveMorphTargetIndex, ActiveMorphTargetWeight, InOutPositions, InOutNormals);
 }
 
-void FClothingSimulationMesh::SkinPhysicsMesh(int32 LODIndex, int32 ActiveMorphTargetIndex, float ActiveMorphTargetWeight, const FReal LocalSpaceScale, const FVec3& LocalSpaceLocation, TArrayView<Softs::FSolverVec3>& OutPositions, TArrayView<Softs::FSolverVec3>& OutNormals) const
-{
-	TRACE_CPUPROFILER_EVENT_SCOPE(FClothingSimulationMesh_SkinPhysicsMesh);
-	SCOPE_CYCLE_COUNTER(STAT_ChaosClothSkinPhysicsMesh);
-	SCOPE_CYCLE_COUNTER(STAT_ClothSkinPhysMesh);
-
-	FTransform ComponentToLocalSpaceReal = GetComponentToWorldTransform();
-	ComponentToLocalSpaceReal.AddToTranslation(-LocalSpaceLocation);
-	check(LocalSpaceScale > UE_SMALL_NUMBER);
-	const FReal LocalSpaceScaleInv = 1. / LocalSpaceScale;
-	ComponentToLocalSpaceReal.MultiplyScale3D(FVec3(LocalSpaceScaleInv));
-	ComponentToLocalSpaceReal.ScaleTranslation(LocalSpaceScaleInv);
-	const FTransform3f ComponentToLocalSpace(ComponentToLocalSpaceReal);  // LWC: Now in local space, therefore it is safe to use single precision which is the asset data format
-
-	const int32* const RESTRICT BoneMap = GetBoneMap().GetData();
-	const FMatrix44f* const RESTRICT RefToLocalMatrices = GetRefToLocalMatrices().GetData();
-
-	check(IsValidLODIndex(LODIndex));
-	const uint32 NumPoints = GetNumPoints(LODIndex);
-	check(NumPoints == OutPositions.Num());
-	check(NumPoints == OutNormals.Num());
-	const TConstArrayView<FClothVertBoneData> BoneData = GetBoneData(LODIndex);
-	TConstArrayView<FVector3f> Positions = GetPositions(LODIndex);
-	TConstArrayView<FVector3f> Normals = GetNormals(LODIndex);
-	TArray<FVector3f> WritablePositions;
-	TArray<FVector3f> WritableNormals;
-
-	const TConstArrayView<FVector3f> MorphTargetPositionDeltas = GetMorphTargetPositionDeltas(LODIndex, ActiveMorphTargetIndex);
-	const TConstArrayView<FVector3f> MorphTargetTangentZDeltas = GetMorphTargetTangentZDeltas(LODIndex, ActiveMorphTargetIndex);
-	const TConstArrayView<int32> MorphTargetIndices = GetMorphTargetIndices(LODIndex, ActiveMorphTargetIndex);
-
-	if (!FMath::IsNearlyZero(ActiveMorphTargetWeight) && !MorphTargetPositionDeltas.IsEmpty() &&
-		MorphTargetPositionDeltas.Num() == MorphTargetTangentZDeltas.Num() &&
-		MorphTargetPositionDeltas.Num() == MorphTargetIndices.Num())
-	{
-		TRACE_CPUPROFILER_EVENT_SCOPE(FClothingSimulationMesh_SkinPhysicsMesh_MorphTargets);
-		// TODO optimize this
-		WritablePositions = Positions;
-		WritableNormals = Normals;
-		ApplyMorphTargetInternal(ActiveMorphTargetWeight, MorphTargetPositionDeltas, MorphTargetTangentZDeltas, MorphTargetIndices, TArrayView<FVector3f>(WritablePositions), TArrayView<FVector3f>(WritableNormals));
-		Positions = WritablePositions;
-		Normals = WritableNormals;
-	}
-
-#if INTEL_ISPC
-	if (bChaos_SkinPhysicsMesh_ISPC_Enabled)
-	{
-		ispc::SkinPhysicsMesh(
-			(ispc::FVector3f*)OutPositions.GetData(),
-			(ispc::FVector3f*)OutNormals.GetData(),
-			(ispc::FVector3f*)Positions.GetData(),
-			(ispc::FVector3f*)Normals.GetData(),
-			(ispc::FClothVertBoneData*)BoneData.GetData(),
-			BoneMap,
-			(ispc::FMatrix44f*)RefToLocalMatrices,
-			(ispc::FTransform3f&)ComponentToLocalSpace,
-			NumPoints);
-	}
-	else
-#endif
-	{
-		static const uint32 MinParallelVertices = 500;  // 500 seems to be the lowest threshold still giving gains even on profiled assets that are only using a small number of influences
-
-		ParallelFor(NumPoints, [&BoneData, &Positions, &Normals, &ComponentToLocalSpace, BoneMap, RefToLocalMatrices, &OutPositions, &OutNormals](uint32 VertIndex)
-		{
-			const uint16* const RESTRICT BoneIndices = BoneData[VertIndex].BoneIndices;
-			const float* const RESTRICT BoneWeights = BoneData[VertIndex].BoneWeights;
-	
-			// WARNING - HORRIBLE UNROLLED LOOP + JUMP TABLE BELOW
-			// done this way because this is a pretty tight and performance critical loop. essentially
-			// rather than checking each influence we can just jump into this switch and fall through
-			// everything to compose the final skinned data
-			const FVector3f& RefParticle = Positions[VertIndex];
-			const FVector3f& RefNormal = Normals[VertIndex];
-
-			FVector3f Position(ForceInitToZero);
-			FVector3f Normal(ForceInitToZero);
-			switch (BoneData[VertIndex].NumInfluences)
-			{
-			default:  // Intentional fall through
-			case 12: AddInfluence(Position, Normal, RefParticle, RefNormal, RefToLocalMatrices[BoneMap[BoneIndices[11]]], BoneWeights[11]);  // Intentional fall through
-			case 11: AddInfluence(Position, Normal, RefParticle, RefNormal, RefToLocalMatrices[BoneMap[BoneIndices[10]]], BoneWeights[10]);  // Intentional fall through
-			case 10: AddInfluence(Position, Normal, RefParticle, RefNormal, RefToLocalMatrices[BoneMap[BoneIndices[ 9]]], BoneWeights[ 9]);  // Intentional fall through
-			case  9: AddInfluence(Position, Normal, RefParticle, RefNormal, RefToLocalMatrices[BoneMap[BoneIndices[ 8]]], BoneWeights[ 8]);  // Intentional fall through
-			case  8: AddInfluence(Position, Normal, RefParticle, RefNormal, RefToLocalMatrices[BoneMap[BoneIndices[ 7]]], BoneWeights[ 7]);  // Intentional fall through
-			case  7: AddInfluence(Position, Normal, RefParticle, RefNormal, RefToLocalMatrices[BoneMap[BoneIndices[ 6]]], BoneWeights[ 6]);  // Intentional fall through
-			case  6: AddInfluence(Position, Normal, RefParticle, RefNormal, RefToLocalMatrices[BoneMap[BoneIndices[ 5]]], BoneWeights[ 5]);  // Intentional fall through
-			case  5: AddInfluence(Position, Normal, RefParticle, RefNormal, RefToLocalMatrices[BoneMap[BoneIndices[ 4]]], BoneWeights[ 4]);  // Intentional fall through
-			case  4: AddInfluence(Position, Normal, RefParticle, RefNormal, RefToLocalMatrices[BoneMap[BoneIndices[ 3]]], BoneWeights[ 3]);  // Intentional fall through
-			case  3: AddInfluence(Position, Normal, RefParticle, RefNormal, RefToLocalMatrices[BoneMap[BoneIndices[ 2]]], BoneWeights[ 2]);  // Intentional fall through
-			case  2: AddInfluence(Position, Normal, RefParticle, RefNormal, RefToLocalMatrices[BoneMap[BoneIndices[ 1]]], BoneWeights[ 1]);  // Intentional fall through
-			case  1: AddInfluence(Position, Normal, RefParticle, RefNormal, RefToLocalMatrices[BoneMap[BoneIndices[ 0]]], BoneWeights[ 0]);  // Intentional fall through
-			case  0: break;
-			}
-
-			OutPositions[VertIndex] = FSolverVec3(ComponentToLocalSpace.TransformPosition(Position));
-			OutNormals[VertIndex] = FSolverVec3(ComponentToLocalSpace.TransformVector(Normal).GetSafeNormal());
-
-		}, NumPoints > MinParallelVertices ? EParallelForFlags::None : EParallelForFlags::ForceSingleThread);
-	}
-}
 
 void FClothingSimulationMesh::Update(
 	FClothingSimulationSolver* Solver,
@@ -343,14 +198,12 @@ void FClothingSimulationMesh::Update(
 	{
 		return;
 	}
-
+	
 	// Skin current LOD positions
-	const FReal LocalSpaceScale = Solver->GetLocalSpaceScale();
-	const FVec3& LocalSpaceLocation = Solver->GetLocalSpaceLocation();
 	TArrayView<FSolverVec3> OutPositions = Solver->GetAnimationPositionsView(ParticleRangeId);
 	TArrayView<FSolverVec3> OutNormals = Solver->GetAnimationNormalsView(ParticleRangeId);
-	
-	SkinPhysicsMesh(LODIndex, ActiveMorphTargetIndex, ActiveMorphTargetWeight, LocalSpaceScale, LocalSpaceLocation, OutPositions, OutNormals);
+	const FClothingSimulationDefaultAccessoryMesh DefaultMesh = GetDefaultAccessoryMesh(LODIndex);
+	DefaultMesh.Update(Solver, OutPositions, OutNormals, ActiveMorphTargetIndex, ActiveMorphTargetWeight);
 
 	// Update old positions after LOD Switching
 	if (LODIndex != PrevLODIndex)
diff --git a/Engine/Plugins/ChaosCloth/Source/ChaosCloth/Private/ChaosCloth/ChaosClothingSimulationSolver.cpp b/Engine/Plugins/ChaosCloth/Source/ChaosCloth/Private/ChaosCloth/ChaosClothingSimulationSolver.cpp
--- a/Engine/Plugins/ChaosCloth/Source/ChaosCloth/Private/ChaosCloth/ChaosClothingSimulationSolver.cpp
+++ b/Engine/Plugins/ChaosCloth/Source/ChaosCloth/Private/ChaosCloth/ChaosClothingSimulationSolver.cpp
@@ -1666,106 +1666,106 @@ void FClothingSimulationSolver::ApplyPreSimulationTransforms()
 	{
 		const TArray<uint32> ActiveGroups = Evolution->GetActiveGroupsArray();
 		PhysicsParallelFor(ActiveGroups.Num(), [this, &ActiveGroups, &DeltaLocalSpaceLocation, MaxVelocitySquared](int32 ActiveGroupIndex)
-		{
-			const uint32 GroupId = ActiveGroups[ActiveGroupIndex];
-
-		// Update particles
-		const TSet<int32>& ActiveSoftBodies = Evolution->GetGroupActiveSoftBodies(GroupId);
-		const Softs::FSolverRigidTransform3& GroupSpaceTransform = PreSimulationTransforms[GroupId];
-		for (const int32 SoftBodyId : ActiveSoftBodies)
-		{
-			Softs::FSolverParticlesRange& Particles = Evolution->GetSoftBodyParticles(SoftBodyId);
-#if INTEL_ISPC
-			if (bRealTypeCompatibleWithISPC && bChaos_PreSimulationTransforms_ISPC_Enabled)  // TODO: Make the ISPC works with both Single and Double depending on the FSolverReal type
 			{
-				if (MaxVelocitySquared == TNumericLimits<FSolverReal>::Max())
-				{
-					ispc::ApplyPreSimulationTransform(
-						(ispc::FVector4f*)Particles.GetPAndInvM().GetData(),
-						(ispc::FVector3f*)Particles.GetV().GetData(),
-						(ispc::FVector3f*)Particles.XArray().GetData(),
-						(ispc::FVector3f*)Particles.GetArrayView(OldAnimationPositions).GetData(),
-						(ispc::FVector3f*)Particles.GetArrayView(AnimationVelocities).GetData(),
-						Particles.GetInvM().GetData(),
-						(const ispc::FVector3f*)Particles.GetConstArrayView(AnimationPositions).GetData(),
-						(const ispc::FTransform3f&)GroupSpaceTransform,
-						(const ispc::FVector3f&)DeltaLocalSpaceLocation,
-						DeltaTime,
-						Particles.GetRangeSize());
-				}
-				else
+				const uint32 GroupId = ActiveGroups[ActiveGroupIndex];
+
+				// Update particles
+				const TSet<int32>& ActiveSoftBodies = Evolution->GetGroupActiveSoftBodies(GroupId);
+				const Softs::FSolverRigidTransform3& GroupSpaceTransform = PreSimulationTransforms[GroupId];
+				for (const int32 SoftBodyId : ActiveSoftBodies)
 				{
-					ispc::ApplyPreSimulationTransformAndClampVelocity(
-						(ispc::FVector4f*)Particles.GetPAndInvM().GetData(),
-						(ispc::FVector3f*)Particles.GetV().GetData(),
-						(ispc::FVector3f*)Particles.XArray().GetData(),
-						(ispc::FVector3f*)Particles.GetArrayView(OldAnimationPositions).GetData(),
-						(ispc::FVector3f*)Particles.GetArrayView(AnimationVelocities).GetData(),
-						Particles.GetInvM().GetData(),
-						(const ispc::FVector3f*)Particles.GetConstArrayView(AnimationPositions).GetData(),
-						(const ispc::FTransform3f&)GroupSpaceTransform,
-						(ispc::FVector3f&)DeltaLocalSpaceLocation,
-						DeltaTime,
-						Particles.GetRangeSize(),
-						MaxVelocitySquared);
-				}
-			}
-			else
+					Softs::FSolverParticlesRange& Particles = Evolution->GetSoftBodyParticles(SoftBodyId);
+#if INTEL_ISPC
+					if (bRealTypeCompatibleWithISPC && bChaos_PreSimulationTransforms_ISPC_Enabled)  // TODO: Make the ISPC works with both Single and Double depending on the FSolverReal type
+					{
+						if (MaxVelocitySquared == TNumericLimits<FSolverReal>::Max())
+						{
+							ispc::ApplyPreSimulationTransform(
+								(ispc::FVector4f*)Particles.GetPAndInvM().GetData(),
+								(ispc::FVector3f*)Particles.GetV().GetData(),
+								(ispc::FVector3f*)Particles.XArray().GetData(),
+								(ispc::FVector3f*)Particles.GetArrayView(OldAnimationPositions).GetData(),
+								(ispc::FVector3f*)Particles.GetArrayView(AnimationVelocities).GetData(),
+								Particles.GetInvM().GetData(),
+								(const ispc::FVector3f*)Particles.GetConstArrayView(AnimationPositions).GetData(),
+								(const ispc::FTransform3f&)GroupSpaceTransform,
+								(const ispc::FVector3f&)DeltaLocalSpaceLocation,
+								DeltaTime,
+								Particles.GetRangeSize());
+						}
+						else
+						{
+							ispc::ApplyPreSimulationTransformAndClampVelocity(
+								(ispc::FVector4f*)Particles.GetPAndInvM().GetData(),
+								(ispc::FVector3f*)Particles.GetV().GetData(),
+								(ispc::FVector3f*)Particles.XArray().GetData(),
+								(ispc::FVector3f*)Particles.GetArrayView(OldAnimationPositions).GetData(),
+								(ispc::FVector3f*)Particles.GetArrayView(AnimationVelocities).GetData(),
+								Particles.GetInvM().GetData(),
+								(const ispc::FVector3f*)Particles.GetConstArrayView(AnimationPositions).GetData(),
+								(const ispc::FTransform3f&)GroupSpaceTransform,
+								(ispc::FVector3f&)DeltaLocalSpaceLocation,
+								DeltaTime,
+								Particles.GetRangeSize(),
+								MaxVelocitySquared);
+						}
+					}
+					else
 #endif
-			{
-				TArrayView<Softs::FSolverVec3> OldAnimationPositionsView = Particles.GetArrayView(OldAnimationPositions);
-				TArrayView<Softs::FSolverVec3> AnimationVelocitiesView = Particles.GetArrayView(AnimationVelocities);
-				TConstArrayView<Softs::FSolverVec3> AnimationPositionsView = Particles.GetConstArrayView(AnimationPositions);
+					{
+						TArrayView<Softs::FSolverVec3> OldAnimationPositionsView = Particles.GetArrayView(OldAnimationPositions);
+						TArrayView<Softs::FSolverVec3> AnimationVelocitiesView = Particles.GetArrayView(AnimationVelocities);
+						TConstArrayView<Softs::FSolverVec3> AnimationPositionsView = Particles.GetConstArrayView(AnimationPositions);
 
-				PhysicsParallelFor(Particles.GetRangeSize(),
-					[this, &GroupSpaceTransform, &OldAnimationPositionsView, &AnimationVelocitiesView, &AnimationPositionsView, &DeltaLocalSpaceLocation, &Particles, MaxVelocitySquared](int32 Index)
-				{
+						PhysicsParallelFor(Particles.GetRangeSize(),
+							[this, &GroupSpaceTransform, &OldAnimationPositionsView, &AnimationVelocitiesView, &AnimationPositionsView, &DeltaLocalSpaceLocation, &Particles, MaxVelocitySquared](int32 Index)
+							{
 
-					// Update initial state for particles
-					Particles.P(Index) = Particles.X(Index) = GroupSpaceTransform.TransformPositionNoScale(Particles.X(Index)) - DeltaLocalSpaceLocation;
-					Particles.V(Index) = GroupSpaceTransform.TransformVector(Particles.V(Index));
+								// Update initial state for particles
+								Particles.P(Index) = Particles.X(Index) = GroupSpaceTransform.TransformPositionNoScale(Particles.X(Index)) - DeltaLocalSpaceLocation;
+								Particles.V(Index) = GroupSpaceTransform.TransformVector(Particles.V(Index));
 
-					// Copy InvM over to PAndInvM
-					Particles.PAndInvM(Index).InvM = Particles.InvM(Index);
+								// Copy InvM over to PAndInvM
+								Particles.PAndInvM(Index).InvM = Particles.InvM(Index);
 
-					// Update anim initial state (target updated by skinning)
-					OldAnimationPositionsView[Index] = GroupSpaceTransform.TransformPositionNoScale(OldAnimationPositionsView[Index]) - DeltaLocalSpaceLocation;
+								// Update anim initial state (target updated by skinning)
+								OldAnimationPositionsView[Index] = GroupSpaceTransform.TransformPositionNoScale(OldAnimationPositionsView[Index]) - DeltaLocalSpaceLocation;
 
-					// Update Animation velocities
-					AnimationVelocitiesView[Index] = (AnimationPositionsView[Index] - OldAnimationPositionsView[Index]) / DeltaTime;
+								// Update Animation velocities
+								AnimationVelocitiesView[Index] = (AnimationPositionsView[Index] - OldAnimationPositionsView[Index]) / DeltaTime;
 
-					// Clamp relative velocity
-					const FSolverVec3 RelVelocity = Particles.V(Index) - AnimationVelocitiesView[Index];
-					const FSolverReal RelVelocitySquaredLength = RelVelocity.SquaredLength();
-					if (RelVelocitySquaredLength > MaxVelocitySquared)
-					{
-						Particles.V(Index) = AnimationVelocitiesView[Index] + RelVelocity * FMath::Sqrt(MaxVelocitySquared / RelVelocitySquaredLength);
+								// Clamp relative velocity
+								const FSolverVec3 RelVelocity = Particles.V(Index) - AnimationVelocitiesView[Index];
+								const FSolverReal RelVelocitySquaredLength = RelVelocity.SquaredLength();
+								if (RelVelocitySquaredLength > MaxVelocitySquared)
+								{
+									Particles.V(Index) = AnimationVelocitiesView[Index] + RelVelocity * FMath::Sqrt(MaxVelocitySquared / RelVelocitySquaredLength);
+								}
+							}, Particles.GetRangeSize() < ClothSolverMinParallelBatchSize);
 					}
-					}, Particles.GetRangeSize() < ClothSolverMinParallelBatchSize);
 				}
-			}
 
-			// Update collision particles
-			const TSet<int32>& ActiveCollisionRanges = Evolution->GetGroupActiveCollisionParticleRanges(GroupId);
-			for (const int32 CollisionRangeId : ActiveCollisionRanges)
-			{
-				Softs::FSolverCollisionParticlesRange& CollisionParticles = Evolution->GetCollisionParticleRange(CollisionRangeId);
-				TArrayView<Softs::FSolverRigidTransform3> OldCollisionTransformsView = CollisionParticles.GetArrayView(OldCollisionTransforms);
-				for (int32 Index = 0; Index < CollisionParticles.GetRangeSize(); ++Index)
+				// Update collision particles
+				const TSet<int32>& ActiveCollisionRanges = Evolution->GetGroupActiveCollisionParticleRanges(GroupId);
+				for (const int32 CollisionRangeId : ActiveCollisionRanges)
 				{
-					// Update initial state for collisions
-					OldCollisionTransformsView[Index] = OldCollisionTransformsView[Index] * GroupSpaceTransform;
-					OldCollisionTransformsView[Index].AddToTranslation(-DeltaLocalSpaceLocation);
-					CollisionParticles.X(Index) = OldCollisionTransformsView[Index].GetTranslation();
-					CollisionParticles.SetR(Index, OldCollisionTransformsView[Index].GetRotation());
-				}
-				if (FClothComplexColliders* const Colliders = GetComplexColliders(CollisionRangeId))
-				{
-					Colliders->ApplyPreSimulationTransforms(GroupSpaceTransform, DeltaLocalSpaceLocation, CollisionParticles.GetConstArrayView(OldCollisionTransforms),
-						CollisionParticles.GetConstArrayView(CollisionTransforms), DeltaTime);
+					Softs::FSolverCollisionParticlesRange& CollisionParticles = Evolution->GetCollisionParticleRange(CollisionRangeId);
+					TArrayView<Softs::FSolverRigidTransform3> OldCollisionTransformsView = CollisionParticles.GetArrayView(OldCollisionTransforms);
+					for (int32 Index = 0; Index < CollisionParticles.GetRangeSize(); ++Index)
+					{
+						// Update initial state for collisions
+						OldCollisionTransformsView[Index] = OldCollisionTransformsView[Index] * GroupSpaceTransform;
+						OldCollisionTransformsView[Index].AddToTranslation(-DeltaLocalSpaceLocation);
+						CollisionParticles.X(Index) = OldCollisionTransformsView[Index].GetTranslation();
+						CollisionParticles.SetR(Index, OldCollisionTransformsView[Index].GetRotation());
+					}
+					if (FClothComplexColliders* const Colliders = GetComplexColliders(CollisionRangeId))
+					{
+						Colliders->ApplyPreSimulationTransforms(GroupSpaceTransform, DeltaLocalSpaceLocation, CollisionParticles.GetConstArrayView(OldCollisionTransforms),
+							CollisionParticles.GetConstArrayView(CollisionTransforms), DeltaTime);
+					}
 				}
-			}
-		}, /*bForceSingleThreaded =*/ !bClothSolverParallelClothPreUpdate);
+			}, /*bForceSingleThreaded =*/ !bClothSolverParallelClothPreUpdate);
 	}
 	else
 	{
@@ -1877,6 +1877,12 @@ void FClothingSimulationSolver::ApplyPreSimulationTransforms()
 			});
 		}
 	}
+
+	for (FClothingSimulationCloth* const Cloth : Cloths)
+	{
+		check(Cloth);
+		Cloth->ApplyPreSimulationTransforms(this, DeltaLocalSpaceLocation, PreSimulationTransforms[Cloth->GetGroupId()], DeltaTime);
+	}
 }
 
 void FClothingSimulationSolver::PreSubstep( const Softs::FSolverReal InterpolationAlpha, bool bDetectSelfCollisions)
@@ -2000,6 +2006,12 @@ void FClothingSimulationSolver::PreSubstep( const Softs::FSolverReal Interpolati
 			}
 		}, /*bForceSingleThreaded =*/ !bClothSolverParallelClothPreUpdate);
 	}
+
+	for (FClothingSimulationCloth* const Cloth : Cloths)
+	{
+		check(Cloth);
+		Cloth->PreSubstep(this, InterpolationAlpha);
+	}
 }
 
 void FClothingSimulationSolver::UpdateSolverField()
diff --git a/Engine/Plugins/ChaosCloth/Source/ChaosCloth/Private/ChaosCloth/ChaosClothingSimulationSolver.ispc b/Engine/Plugins/ChaosCloth/Source/ChaosCloth/Private/ChaosCloth/ChaosClothingSimulationSolver.ispc
--- a/Engine/Plugins/ChaosCloth/Source/ChaosCloth/Private/ChaosCloth/ChaosClothingSimulationSolver.ispc
+++ b/Engine/Plugins/ChaosCloth/Source/ChaosCloth/Private/ChaosCloth/ChaosClothingSimulationSolver.ispc
@@ -87,6 +87,30 @@ export void ApplyPreSimulationTransform(uniform FVector4f ParticlesPandInvM[],
 	}
 }
 
+export void ApplyAnimationPreSimulationTransform(uniform FVector3f OldAnimationPositions[],
+										uniform FVector3f AnimationVelocities[],
+										const uniform FVector3f AnimationPositions[],
+										const uniform FTransform3f& GroupSpaceTransform,
+										const uniform FVector3f &DeltaLocalSpaceLocation,
+										const uniform float DeltaTime,
+										const uniform int32 NumParticles)
+{
+	foreach(i = 0 ... NumParticles)
+	{
+		const uniform int32 Index = extract(i, 0);
+
+		const varying FVector3f OldAnimationPosition = VectorLoad(&OldAnimationPositions[Index]);
+		const varying FVector3f AP = TransformPosition(GroupSpaceTransform, OldAnimationPosition) - DeltaLocalSpaceLocation;
+
+		VectorStore(&OldAnimationPositions[Index], AP);
+
+		const varying FVector3f AnimationPosition = VectorLoad(&AnimationPositions[Index]);
+		const varying FVector3f AV = (AnimationPosition - AP) / DeltaTime;
+
+		VectorStore(&AnimationVelocities[Index], AV);
+	}
+}
+
 export void ApplyPreSimulationTransformAndClampVelocity(uniform FVector4f ParticlesPandInvM[],
 										uniform FVector3f ParticlesV[],
 										uniform FVector3f ParticlesX[],
diff --git a/Engine/Plugins/ChaosCloth/Source/ChaosCloth/Public/ChaosCloth/ChaosClothConstraints.h b/Engine/Plugins/ChaosCloth/Source/ChaosCloth/Public/ChaosCloth/ChaosClothConstraints.h
--- a/Engine/Plugins/ChaosCloth/Source/ChaosCloth/Public/ChaosCloth/ChaosClothConstraints.h
+++ b/Engine/Plugins/ChaosCloth/Source/ChaosCloth/Public/ChaosCloth/ChaosClothConstraints.h
@@ -18,6 +18,7 @@ namespace Chaos
 {
 
 	struct FClothingPatternData;
+	struct FClothingAccessoryMeshData;
 	class FClothComplexColliders;
 
 	class FClothConstraints final
@@ -78,7 +79,8 @@ namespace Chaos
 			const int32 MultiResCoarseLODParticleRangeId = INDEX_NONE,
 			const TSharedPtr<Softs::FMultiResConstraints>& FineLODMultiResConstraint = TSharedPtr<Softs::FMultiResConstraints>(nullptr),
 			const TArray<const FClothComplexColliders*>& ComplexColliders = TArray<const FClothComplexColliders*>(),
-			const TSharedPtr<const FManagedArrayCollection>& ManagedArrayCollection = TSharedPtr<const FManagedArrayCollection>());
+			const TSharedPtr<const FManagedArrayCollection>& ManagedArrayCollection = TSharedPtr<const FManagedArrayCollection>(),
+			const TMap<FName, FClothingAccessoryMeshData>* AccessoryMeshes = nullptr);
 
 		void Update(
 			const Softs::FCollectionPropertyConstFacade& ConfigProperties,
@@ -177,7 +179,8 @@ namespace Chaos
 		void CreateBackstopConstraints(
 			const Softs::FCollectionPropertyConstFacade& ConfigProperties,
 			const TMap<FString, TConstArrayView<FRealSingle>>& WeightMaps,
-			Softs::FSolverReal MeshScale);
+			Softs::FSolverReal MeshScale,
+			const TMap<FName, FClothingAccessoryMeshData>* AccessoryMeshes);
 		void CreateAnimDriveConstraints(
 			const Softs::FCollectionPropertyConstFacade& ConfigProperties,
 			const TMap<FString, TConstArrayView<FRealSingle>>& WeightMaps);
diff --git a/Engine/Plugins/ChaosCloth/Source/ChaosCloth/Public/ChaosCloth/ChaosClothVisualization.h b/Engine/Plugins/ChaosCloth/Source/ChaosCloth/Public/ChaosCloth/ChaosClothVisualization.h
--- a/Engine/Plugins/ChaosCloth/Source/ChaosCloth/Public/ChaosCloth/ChaosClothVisualization.h
+++ b/Engine/Plugins/ChaosCloth/Source/ChaosCloth/Public/ChaosCloth/ChaosClothVisualization.h
@@ -4,6 +4,7 @@
 #include "Misc/Build.h"      // For CHAOS_DEBUG_DRAW
 #include "Chaos/Declares.h"  //
 #include "Misc/CoreMiscDefines.h"
+#include "UObject/NameTypes.h"
 
 #if WITH_EDITOR
 #include "UObject/GCObject.h"
@@ -61,6 +62,8 @@ namespace Chaos
 		CHAOSCLOTH_API void DrawClothClothConstraints(FPrimitiveDrawInterface* PDI = nullptr) const;
 		CHAOSCLOTH_API void DrawTeleportReset(FPrimitiveDrawInterface* PDI = nullptr) const;
 		CHAOSCLOTH_API void DrawExtremlyDeformedEdges(FPrimitiveDrawInterface* PDI = nullptr) const;
+		CHAOSCLOTH_API void DrawAccessoryMesh(FPrimitiveDrawInterface* PDI = nullptr, const FName& AccessoryMeshName = NAME_None) const;
+		CHAOSCLOTH_API void DrawAccessoryMeshNormals(FPrimitiveDrawInterface* PDI = nullptr, const FName& AccessoryMeshName = NAME_None, const FReal Length = 20.) const;
 	
 #else  // #if CHAOS_DEBUG_DRAW
 		void SetSolver(const ::Chaos::FClothingSimulationSolver* /*InSolver*/) {}
@@ -108,6 +111,7 @@ namespace Chaos
 		CHAOSCLOTH_API void DrawSimMorphTarget(FPrimitiveDrawInterface* PDI = nullptr, const FString& Name = "") const;
 		CHAOSCLOTH_API TArray<FString> GetAllWeightMapNames() const;
 		CHAOSCLOTH_API TArray<FString> GetAllMorphTargetNames() const;
+		CHAOSCLOTH_API TArray<FName> GetAllAccessoryMeshNames() const;
 #else  // #if WITH_EDITOR && CHAOS_DEBUG_DRAW
 		void DrawPhysMeshShaded(FPrimitiveDrawInterface* /*PDI*/) const {}
 		void DrawKinematicColliderShaded(FPrimitiveDrawInterface* /*PDI*/) const {}
diff --git a/Engine/Plugins/ChaosCloth/Source/ChaosCloth/Public/ChaosCloth/ChaosClothingSimulationAccessoryMesh.h b/Engine/Plugins/ChaosCloth/Source/ChaosCloth/Public/ChaosCloth/ChaosClothingSimulationAccessoryMesh.h
new file mode 100644
--- /dev/null
+++ b/Engine/Plugins/ChaosCloth/Source/ChaosCloth/Public/ChaosCloth/ChaosClothingSimulationAccessoryMesh.h
@@ -0,0 +1,126 @@
+// Copyright Epic Games, Inc. All Rights Reserved.
+#pragma once
+#include "Chaos/Core.h"
+#include "Chaos/PBDSoftsEvolutionFwd.h"
+#include "Containers/ContainersFwd.h"
+#include "EngineDefines.h"
+
+#define UE_API CHAOSCLOTH_API
+
+struct FClothVertBoneData;
+struct FManagedArrayCollection;
+
+namespace Chaos
+{
+	class FClothingSimulationSolver;
+	class FClothingSimulationMesh;
+
+	// Accessory mesh used as an input to a simulation. They are expected to share topology (e.g., Indices, PatternIndices) with their associated FClothingSimulationMesh
+	class FClothingSimulationAccessoryMesh
+	{
+	public:
+		UE_API FClothingSimulationAccessoryMesh(const FClothingSimulationMesh& InMesh, const FName& InAccessoryMeshName);
+
+		UE_API virtual ~FClothingSimulationAccessoryMesh();
+
+		FClothingSimulationAccessoryMesh(const FClothingSimulationAccessoryMesh&) = delete;
+		FClothingSimulationAccessoryMesh(FClothingSimulationAccessoryMesh&&) = delete;
+		FClothingSimulationAccessoryMesh& operator=(const FClothingSimulationAccessoryMesh&) = delete;
+		FClothingSimulationAccessoryMesh& operator=(FClothingSimulationAccessoryMesh&&) = delete;
+
+		const FName& GetAccessoryMeshName() const { return AccessoryMeshName; }
+
+		/* Return the number of points. */
+		virtual int32 GetNumPoints() const = 0;
+
+		/* Return the number of pattern points (2d, unwelded). */
+		virtual int32 GetNumPatternPoints() const = 0;
+
+		/* Return the source mesh positions (pre-skinning). */
+		virtual TConstArrayView<FVector3f> GetPositions() const = 0;
+
+		/* Return the source mesh 2d pattern positions. */
+		virtual TConstArrayView<FVector2f> GetPatternPositions() const = 0;
+
+		/* Return the source mesh normals (pre-skinning). */
+		virtual TConstArrayView<FVector3f> GetNormals() const = 0;
+
+		/* Return the bone data containing bone weights and influences. */
+		virtual TConstArrayView<FClothVertBoneData> GetBoneData() const = 0;
+
+		/** Return the MorphTargetIndex for a given Morph Target Name, if it exists. NOTE: currently only the DefaultAccessoryMesh supports Morph Targets.*/
+		virtual int32 FindMorphTargetByName(const FString& Name) const
+		{
+			return INDEX_NONE;
+		}
+
+		/** Get a list of all MorphTargets. (Index matches FindMorphTargetByName) */
+		virtual TConstArrayView<FString> GetAllMorphTargetNames() const
+		{
+			return TConstArrayView<FString>();
+		}
+
+		/** Get all Morph Target position deltas for a given MorphTargetIndex (e.g., index returned by FindMorphTargetByName). Deltas index back to Positions via MorphTargetIndices */
+		virtual TConstArrayView<FVector3f> GetMorphTargetPositionDeltas(int32 MorphTargetIndex) const
+		{
+			return TConstArrayView<FVector3f>();
+		}
+
+		/** Get all Morph Target tangent z (normal) deltas for a given MorphTargetIndex (e.g., index returned by FindMorphTargetByName). Deltas index back to Normals via MorphTargetIndices */
+		virtual TConstArrayView<FVector3f> GetMorphTargetTangentZDeltas(int32 MorphTargetIndex) const
+		{
+			return TConstArrayView<FVector3f>();
+		}
+
+		/** Get all Morph Target indices for a given MorphTargetIndex (e.g., index returned by FindMorphTargetByName). These indices can map MorphTargetDeltas back to Positions. */
+		virtual TConstArrayView<int32> GetMorphTargetIndices(int32 MorphTargetIndex) const
+		{
+			return TConstArrayView<int32>();
+		}
+
+		/* Update the mesh for the next solver step, doing skinning and matching the shapes during LOD changes. */
+		UE_API void Update(
+			const FClothingSimulationSolver* Solver,
+			const TArrayView<Softs::FSolverVec3>& OutPositions,
+			const TArrayView<Softs::FSolverVec3>& OutNormals,
+			int32 ActiveMorphTargetIndex = INDEX_NONE,
+			float ActiveMorphTargetWeight = 0.f) const;
+
+		UE_API void ApplyMorphTarget(
+			int32 ActiveMorphTargetIndex,
+			float ActiveMorphTargetWeight,
+			const TArrayView<Softs::FSolverVec3>& InOutPositions,
+			const TArrayView<Softs::FSolverVec3>& InOutNormals) const;
+		// ---- End of the Cloth interface ----
+
+	private:
+		void SkinPhysicsMesh(
+			int32 ActiveMorphTargetIndex,
+			float ActiveMorphTargetWeight,
+			const FReal LocalSpaceScale,
+			const FVec3& LocalSpaceLocation,
+			const TArrayView<Softs::FSolverVec3>& OutPositions,
+			const TArrayView<Softs::FSolverVec3>& OutNormals) const;
+
+	protected:
+		const FClothingSimulationMesh& Mesh;
+		const FName AccessoryMeshName;
+	};
+} // namespace Chaos
+
+#if !defined(CHAOS_SKIN_PHYSICS_MESH_ISPC_ENABLED_DEFAULT)
+#define CHAOS_SKIN_PHYSICS_MESH_ISPC_ENABLED_DEFAULT 1
+#endif
+
+#if !defined(USE_ISPC_KERNEL_CONSOLE_VARIABLES_IN_SHIPPING)
+#define USE_ISPC_KERNEL_CONSOLE_VARIABLES_IN_SHIPPING 0
+#endif
+
+// Support run-time toggling on supported platforms in non-shipping configurations
+#if !INTEL_ISPC || (UE_BUILD_SHIPPING && !USE_ISPC_KERNEL_CONSOLE_VARIABLES_IN_SHIPPING)
+static constexpr bool bChaos_SkinPhysicsMesh_ISPC_Enabled = INTEL_ISPC && CHAOS_SKIN_PHYSICS_MESH_ISPC_ENABLED_DEFAULT;
+#else
+extern bool bChaos_SkinPhysicsMesh_ISPC_Enabled;
+#endif
+
+#undef UE_API
diff --git a/Engine/Plugins/ChaosCloth/Source/ChaosCloth/Public/ChaosCloth/ChaosClothingSimulationCloth.h b/Engine/Plugins/ChaosCloth/Source/ChaosCloth/Public/ChaosCloth/ChaosClothingSimulationCloth.h
--- a/Engine/Plugins/ChaosCloth/Source/ChaosCloth/Public/ChaosCloth/ChaosClothingSimulationCloth.h
+++ b/Engine/Plugins/ChaosCloth/Source/ChaosCloth/Public/ChaosCloth/ChaosClothingSimulationCloth.h
@@ -81,8 +81,17 @@ namespace Chaos
 		CHAOSCLOTH_API TConstArrayView<Softs::FSolverVec3> GetOldAnimationPositions(const FClothingSimulationSolver* Solver) const;
 		// Return the solver's input normals for this cloth source current LOD, not thread safe, call must be done right after the solver update.
 		CHAOSCLOTH_API TConstArrayView<Softs::FSolverVec3> GetAnimationNormals(const FClothingSimulationSolver* Solver) const;
-		// Return the solver's input normals for this cloth source current LOD, not thread safe, call must be done right after the solver update.
+		// Return the solver's input velocities for this cloth source current LOD, not thread safe, call must be done right after the solver update.
 		CHAOSCLOTH_API TConstArrayView<Softs::FSolverVec3> GetAnimationVelocities(const FClothingSimulationSolver* Solver) const;
+		CHAOSCLOTH_API bool AccessoryMeshExists(const FClothingSimulationSolver* Solver, const FName& AccessoryMeshName) const;
+		// Return the solver's input positions for this cloth source current LOD, not thread safe, call must be done right after the solver update.
+		CHAOSCLOTH_API TConstArrayView<Softs::FSolverVec3> GetAccessoryMeshPositions(const FClothingSimulationSolver* Solver, const FName& AccessoryMeshName) const;
+		// Return the solver's input positions for this cloth source current LOD, not thread safe, call must be done right after the solver update.
+		CHAOSCLOTH_API TConstArrayView<Softs::FSolverVec3> GetOldAccessoryMeshPositions(const FClothingSimulationSolver* Solver, const FName& AccessoryMeshName) const;
+		// Return the solver's input normals for this cloth source current LOD, not thread safe, call must be done right after the solver update.
+		CHAOSCLOTH_API TConstArrayView<Softs::FSolverVec3> GetAccessoryMeshNormals(const FClothingSimulationSolver* Solver, const FName& AccessoryMeshName) const;
+		// Return the solver's input velocities for this cloth source current LOD, not thread safe, call must be done right after the solver update.
+		CHAOSCLOTH_API TConstArrayView<Softs::FSolverVec3> GetAccessoryMeshVelocities(const FClothingSimulationSolver* Solver, const FName& AccessoryMeshName) const;
 		// Return the solver's positions for this cloth current LOD, not thread safe, call must be done right after the solver update.
 		CHAOSCLOTH_API TConstArrayView<Softs::FSolverVec3> GetParticlePositions(const FClothingSimulationSolver* Solver) const;
 		// Return the solver's velocities for this cloth current LOD, not thread safe, call must be done right after the solver update.
@@ -126,6 +135,7 @@ namespace Chaos
 		{
 			return GetFaceIntMapByProperty(Solver, FName(Property));
 		}
+		CHAOSCLOTH_API FString GetPropertyString(const FClothingSimulationSolver* Solver, const FName& Name) const;
 		// Return the current LOD tethers.
 		CHAOSCLOTH_API const TArray<TConstArrayView<TTuple<int32, int32, float>>>& GetTethers(const FClothingSimulationSolver* Solver) const;
 		// Return the reference bone index for this cloth.
@@ -135,6 +145,7 @@ namespace Chaos
 		CHAOSCLOTH_API int32 GetCurrentMorphTargetIndex(const FClothingSimulationSolver* Solver) const;
 		CHAOSCLOTH_API FRealSingle GetCurrentMorphTargetWeight(const FClothingSimulationSolver* Solver) const;
 		CHAOSCLOTH_API TSet<FString> GetAllMorphTargetNames() const;
+		CHAOSCLOTH_API TSet<FName> GetAllAccessoryMeshNames() const;
 #if CHAOS_DEBUG_DRAW
 		FRealSingle GetTimeSinceLastTeleport() const { return TimeSinceLastTeleport; }
 		FRealSingle GetTimeSinceLastReset() const { return TimeSinceLastReset; }
@@ -150,6 +161,9 @@ namespace Chaos
 		CHAOSCLOTH_API void PostUpdate(FClothingSimulationSolver* Solver);
 
 		CHAOSCLOTH_API void UpdateFromCache(const FClothingSimulationCacheData& CacheData);
+
+		CHAOSCLOTH_API void ApplyPreSimulationTransforms(FClothingSimulationSolver* Solver, const Softs::FSolverVec3& DeltaLocalSpaceLocation, const Softs::FSolverRigidTransform3& GroupSpaceTransform, const Softs::FSolverReal DeltaTime);
+		CHAOSCLOTH_API void PreSubstep(FClothingSimulationSolver* Solver, const Softs::FSolverReal InterpolationAlpha);
 		// ---- End of the Solver interface ----
 
 
diff --git a/Engine/Plugins/ChaosCloth/Source/ChaosCloth/Public/ChaosCloth/ChaosClothingSimulationDefaultAccessoryMesh.h b/Engine/Plugins/ChaosCloth/Source/ChaosCloth/Public/ChaosCloth/ChaosClothingSimulationDefaultAccessoryMesh.h
new file mode 100644
--- /dev/null
+++ b/Engine/Plugins/ChaosCloth/Source/ChaosCloth/Public/ChaosCloth/ChaosClothingSimulationDefaultAccessoryMesh.h
@@ -0,0 +1,69 @@
+// Copyright Epic Games, Inc. All Rights Reserved.
+#pragma once
+
+#include "ChaosCloth/ChaosClothingSimulationAccessoryMesh.h"
+
+#define UE_API CHAOSCLOTH_API
+
+namespace Chaos
+{
+	class FClothingSimulationMesh;
+
+	// Wrapper around a FClothingSimulationMesh's LOD Data so it can share methods with accessory meshes.
+	class FClothingSimulationDefaultAccessoryMesh final : public FClothingSimulationAccessoryMesh
+	{
+	public:
+		UE_API FClothingSimulationDefaultAccessoryMesh(const FClothingSimulationMesh& InMesh, const int32 InLODIndex);
+
+		UE_API virtual ~FClothingSimulationDefaultAccessoryMesh();
+
+		const FClothingSimulationMesh& GetMesh() const
+		{
+			return Mesh;
+		}
+
+		int32 GetLODIndex() const
+		{ 
+			return LODIndex;
+		}
+
+		/* Return the number of points. */
+		UE_API virtual int32 GetNumPoints() const override;
+
+		/* Return the number of pattern points (2d, unwelded). */
+		UE_API virtual int32 GetNumPatternPoints() const override;
+
+		/* Return the source mesh positions (pre-skinning). */
+		UE_API virtual TConstArrayView<FVector3f> GetPositions() const override;
+
+		/* Return the source mesh 2d pattern positions. */
+		UE_API virtual TConstArrayView<FVector2f> GetPatternPositions() const override;
+
+		/* Return the source mesh normals (pre-skinning). */
+		UE_API virtual TConstArrayView<FVector3f> GetNormals() const override;
+
+		/* Return the bone data containing bone weights and influences. */
+		UE_API virtual TConstArrayView<FClothVertBoneData> GetBoneData() const override;
+
+		/** Return the MorphTargetIndex for a given Morph Target Name, if it exists. */
+		UE_API virtual int32 FindMorphTargetByName(const FString& Name) const override;
+
+		/** Get a list of all MorphTargets. (Index matches FindMorphTargetByName) */
+		UE_API virtual TConstArrayView<FString> GetAllMorphTargetNames() const override;
+
+		/** Get all Morph Target position deltas for a given MorphTargetIndex (e.g., index returned by FindMorphTargetByName). Deltas index back to Positions via MorphTargetIndices */
+		UE_API virtual TConstArrayView<FVector3f> GetMorphTargetPositionDeltas(int32 MorphTargetIndex) const override;
+
+		/** Get all Morph Target tangent z (normal) deltas for a given MorphTargetIndex (e.g., index returned by FindMorphTargetByName). Deltas index back to Normals via MorphTargetIndices */
+		UE_API virtual TConstArrayView<FVector3f> GetMorphTargetTangentZDeltas(int32 MorphTargetIndex) const override;
+
+		/** Get all Morph Target indices for a given MorphTargetIndex (e.g., index returned by FindMorphTargetByName). These indices can map MorphTargetDeltas back to Positions. */
+		UE_API virtual TConstArrayView<int32> GetMorphTargetIndices(int32 MorphTargetIndex) const override;
+
+	private:
+		static const FName DefaultAccessoryMeshName;
+		const int32 LODIndex;
+	};
+} // namespace Chaos
+
+#undef UE_API
diff --git a/Engine/Plugins/ChaosCloth/Source/ChaosCloth/Public/ChaosCloth/ChaosClothingSimulationMesh.h b/Engine/Plugins/ChaosCloth/Source/ChaosCloth/Public/ChaosCloth/ChaosClothingSimulationMesh.h
--- a/Engine/Plugins/ChaosCloth/Source/ChaosCloth/Public/ChaosCloth/ChaosClothingSimulationMesh.h
+++ b/Engine/Plugins/ChaosCloth/Source/ChaosCloth/Public/ChaosCloth/ChaosClothingSimulationMesh.h
@@ -1,6 +1,7 @@
 // Copyright Epic Games, Inc. All Rights Reserved.
 #pragma once
 
+#include "ChaosCloth/ChaosClothingSimulationDefaultAccessoryMesh.h"
 #include "Chaos/Core.h"
 #include "Chaos/PBDSoftsEvolutionFwd.h"
 #include "Containers/ContainersFwd.h"
@@ -13,15 +14,11 @@ struct FMeshToMeshVertData;
 struct FClothVertBoneData;
 struct FManagedArrayCollection;
 
-#if UE_ENABLE_INCLUDE_ORDER_DEPRECATED_IN_5_5
-class USkeletalMeshComponent;
-class UClothingAssetCommon;
-class FClothingSimulationContextCommon;
-#endif
 
 namespace Chaos
 {
 	class FClothingSimulationSolver;
+	class FClothingSimulationAccessoryMesh;
 
 	// Mesh simulation node
 	class FClothingSimulationMesh
@@ -173,6 +170,22 @@ namespace Chaos
 			return TConstArrayView<int32>();
 		}
 
+		virtual TConstArrayView<FName> GetAllAccessoryMeshNames(int32 LODIndex) const
+		{
+			return TConstArrayView<FName>();
+		}
+
+		/** Return the Accessory Mesh if it exists.*/
+		virtual const FClothingSimulationAccessoryMesh* GetAccessoryMesh(int32 LODIndex, const FName& AccessoryMeshName) const
+		{
+			return nullptr;
+		}
+
+		FClothingSimulationDefaultAccessoryMesh GetDefaultAccessoryMesh(int32 LODIndex) const
+		{
+			return FClothingSimulationDefaultAccessoryMesh(*this, LODIndex);
+		}
+
 		/** Deform the specified positions to match the shape of the previous LOD. */
 		UE_API bool WrapDeformLOD(
 			int32 PrevLODIndex,
@@ -212,16 +225,8 @@ namespace Chaos
 		// ---- End of the Cloth interface ----
 
 	private:
-		UE_API void SkinPhysicsMesh(
-			int32 LODIndex,
-			int32 ActiveMorphTargetIndex,
-			float ActiveMorphTargetWeight,
-			const FReal LocalSpaceScale,
-			const FVec3& LocalSpaceLocation,
-			TArrayView<Softs::FSolverVec3>& OutPositions,
-			TArrayView<Softs::FSolverVec3>& OutNormals) const;
 
-		UE_API bool WrapDeformLOD(
+		bool WrapDeformLOD(
 			int32 PrevLODIndex,
 			int32 LODIndex,
 			const TConstArrayView<Softs::FSolverVec3>& Positions,
@@ -233,6 +238,7 @@ namespace Chaos
 		/** Debug name of the source component. */
 		FString DebugName;
 #endif
+
 	protected:
 #if UE_ENABLE_DEBUG_DRAWING
 		FName ReferenceBoneName = NAME_None;
@@ -240,19 +246,4 @@ namespace Chaos
 	};
 } // namespace Chaos
 
-#if !defined(CHAOS_SKIN_PHYSICS_MESH_ISPC_ENABLED_DEFAULT)
-#define CHAOS_SKIN_PHYSICS_MESH_ISPC_ENABLED_DEFAULT 1
-#endif
-
-#if !defined(USE_ISPC_KERNEL_CONSOLE_VARIABLES_IN_SHIPPING)
-#define USE_ISPC_KERNEL_CONSOLE_VARIABLES_IN_SHIPPING 0
-#endif
-
-// Support run-time toggling on supported platforms in non-shipping configurations
-#if !INTEL_ISPC || (UE_BUILD_SHIPPING && !USE_ISPC_KERNEL_CONSOLE_VARIABLES_IN_SHIPPING)
-static constexpr bool bChaos_SkinPhysicsMesh_ISPC_Enabled = INTEL_ISPC && CHAOS_SKIN_PHYSICS_MESH_ISPC_ENABLED_DEFAULT;
-#else
-extern bool bChaos_SkinPhysicsMesh_ISPC_Enabled;
-#endif
-
 #undef UE_API
diff --git a/Engine/Plugins/ChaosClothAsset/Source/ChaosClothAsset/Private/ChaosClothAsset/ClothCollection.cpp b/Engine/Plugins/ChaosClothAsset/Source/ChaosClothAsset/Private/ChaosClothAsset/ClothCollection.cpp
--- a/Engine/Plugins/ChaosClothAsset/Source/ChaosClothAsset/Private/ChaosClothAsset/ClothCollection.cpp
+++ b/Engine/Plugins/ChaosClothAsset/Source/ChaosClothAsset/Private/ChaosClothAsset/ClothCollection.cpp
@@ -87,6 +87,7 @@ namespace UE::Chaos::ClothAsset::Private
 	{
 		ClothCollectionAttribute::SimPosition2D,
 		ClothCollectionAttribute::SimVertex3DLookup,
+		ClothCollectionAttribute::SimImportVertexID,
 	};
 
 	// Sim Vertices 3D Group
@@ -104,6 +105,13 @@ namespace UE::Chaos::ClothAsset::Private
 		ClothCollectionAttribute::SeamStitchLookup,
 		ClothCollectionAttribute::SimCustomResizingBlend,
 	};
+	static const TArray<FName> SimAccessoryMeshVertices3DGroupAttributePrefixes =
+	{
+		ClothCollectionAttribute::SimAccessoryMeshPosition3DPrefix,
+		ClothCollectionAttribute::SimAccessoryMeshNormalPrefix,
+		ClothCollectionAttribute::SimAccessoryMeshBoneIndicesPrefix,
+		ClothCollectionAttribute::SimAccessoryMeshBoneWeightsPrefix,
+	};
 
 	//~ Sim Morph Targets Group
 	static const TArray<FName> SimMorphTargetsGroupAttributes =
@@ -154,6 +162,16 @@ namespace UE::Chaos::ClothAsset::Private
 		ClothCollectionAttribute::CustomResizingRegionType
 	};
 
+	//~ Sim Accessory Meshes Group
+	static const TArray<FName> SimAccessoryMeshesGroupAttributes =
+	{
+		ClothCollectionAttribute::SimAccessoryMeshName,
+		ClothCollectionAttribute::SimAccessoryMeshPosition3DAttribute,
+		ClothCollectionAttribute::SimAccessoryMeshNormalAttribute,
+		ClothCollectionAttribute::SimAccessoryMeshBoneIndicesAttribute,
+		ClothCollectionAttribute::SimAccessoryMeshBoneWeightsAttribute,
+	};
+
 	// All fixed attributes for this collection
 	static const TMap<FName, TArray<FName>> FixedAttributeNamesMap =
 	{
@@ -172,6 +190,13 @@ namespace UE::Chaos::ClothAsset::Private
 		{ ClothCollectionGroup::Fabrics, SimFabricGroupAttributes },
 		{ ClothCollectionGroup::Solvers, SolverGroupAttributes },
 		{ ClothCollectionGroup::CustomResizingRegions, GroupResizingGroupAttributes },
+		{ ClothCollectionGroup::SimAccessoryMeshes, SimAccessoryMeshesGroupAttributes },
+	};
+
+	// All fixed attribute prefixes for this collection (see SimAccessoryMesh SimVertex3D)
+	static const TMap<FName, TArray<FName>> FixedAttributePrefixNamesMap =
+	{
+		{ ClothCollectionGroup::SimVertices3D, SimAccessoryMeshVertices3DGroupAttributePrefixes },
 	};
 
 	// All paintable attributes for this collection
@@ -180,6 +205,45 @@ namespace UE::Chaos::ClothAsset::Private
 		{ ClothCollectionGroup::SimVertices3D, { ClothCollectionAttribute::SimCustomResizingBlend } },  // Can be painted if required
 		{ ClothCollectionGroup::RenderVertices, { ClothCollectionAttribute::RenderDeformerSkinningBlend, ClothCollectionAttribute::RenderCustomResizingBlend } },  // Can be painted if required
 	};
+
+	static bool NamesHaveMatchingPrefix(const FName& Name, const FName& Prefix)
+	{
+#if UE_FNAME_OUTLINE_NUMBER
+		// Comparison index includes number. Create a new FName without a number and compare it to prefix.
+		FName NumberlessName(Name);
+		NumberlessName.SetNumber(NAME_NO_NUMBER_INTERNAL);
+		return NumberlessName == Prefix;
+#else
+		// Comparison index does not include number
+		return Name.GetComparisonIndex() == Prefix.GetComparisonIndex();
+#endif
+	}
+
+	template<typename T>
+	FName AddUniquePrefixedAttribute(const FName& Group, const FName& AttributePrefix, const TManagedArray<FName>* ExistingAttributes, TSharedRef<FManagedArrayCollection> Collection)
+	{
+		int32 MaxExistingNumber = 0;
+		if (ExistingAttributes)
+		{
+			for (const FName& ExistingAttr : *ExistingAttributes)
+			{
+				if (ExistingAttr == NAME_None)
+				{
+					continue;
+				}
+				if (ensure(NamesHaveMatchingPrefix(ExistingAttr, AttributePrefix)))
+				{
+					MaxExistingNumber = FMath::Max(MaxExistingNumber, ExistingAttr.GetNumber());
+				}
+			}
+		}
+
+		const FName NewAttributeName(AttributePrefix, MaxExistingNumber + 1);
+		check(!Collection->HasAttribute(NewAttributeName, Group));
+
+		Collection->AddAttribute<T>(NewAttributeName, Group);
+		return NewAttributeName;
+	}
 }  // End namespace UE::Chaos::ClothAsset::Private
 
 namespace UE::Chaos::ClothAsset
@@ -241,6 +305,7 @@ namespace UE::Chaos::ClothAsset
 		// Sim Vertices 2D Group
 		SimPosition2D = ManagedArrayCollection->FindAttribute<FVector2f>(ClothCollectionAttribute::SimPosition2D, ClothCollectionGroup::SimVertices2D);
 		SimVertex3DLookup = ManagedArrayCollection->FindAttribute<int32>(ClothCollectionAttribute::SimVertex3DLookup, ClothCollectionGroup::SimVertices2D);
+		SimImportVertexID = ManagedArrayCollection->FindAttribute<int32>(ClothCollectionAttribute::SimImportVertexID, ClothCollectionGroup::SimVertices2D);
 
 		// Sim Vertices 3D Group
 		SimPosition3D = ManagedArrayCollection->FindAttribute<FVector3f>(ClothCollectionAttribute::SimPosition3D, ClothCollectionGroup::SimVertices3D);
@@ -287,6 +352,13 @@ namespace UE::Chaos::ClothAsset
 		// Group Resizing Group
 		CustomResizingRegionSet = ManagedArrayCollection->FindAttribute<FString>(ClothCollectionAttribute::CustomResizingRegionSet, ClothCollectionGroup::CustomResizingRegions);
 		CustomResizingRegionType = ManagedArrayCollection->FindAttribute<int32>(ClothCollectionAttribute::CustomResizingRegionType, ClothCollectionGroup::CustomResizingRegions);
+
+		// Sim Accessory Meshes Group
+		SimAccessoryMeshName = ManagedArrayCollection->FindAttribute<FName>(ClothCollectionAttribute::SimAccessoryMeshName, ClothCollectionGroup::SimAccessoryMeshes);
+		SimAccessoryMeshPosition3DAttribute = ManagedArrayCollection->FindAttribute<FName>(ClothCollectionAttribute::SimAccessoryMeshPosition3DAttribute, ClothCollectionGroup::SimAccessoryMeshes);
+		SimAccessoryMeshNormalAttribute = ManagedArrayCollection->FindAttribute<FName>(ClothCollectionAttribute::SimAccessoryMeshNormalAttribute, ClothCollectionGroup::SimAccessoryMeshes);
+		SimAccessoryMeshBoneIndicesAttribute = ManagedArrayCollection->FindAttribute<FName>(ClothCollectionAttribute::SimAccessoryMeshBoneIndicesAttribute, ClothCollectionGroup::SimAccessoryMeshes);
+		SimAccessoryMeshBoneWeightsAttribute = ManagedArrayCollection->FindAttribute<FName>(ClothCollectionAttribute::SimAccessoryMeshBoneWeightsAttribute, ClothCollectionGroup::SimAccessoryMeshes);
 	}
 
 	FConstClothCollection::FConstClothCollection(const TSharedRef<const FManagedArrayCollection>& InManagedArrayCollection)
@@ -310,7 +382,14 @@ namespace UE::Chaos::ClothAsset
 				SimMorphTargetSimVertex3DIndex &&
 
 				// Sim Vertices 3D Group
-				PreResizedSimPosition3D;
+				PreResizedSimPosition3D &&
+
+				// Sim Accessory Meshes Group
+				SimAccessoryMeshName &&
+				SimAccessoryMeshPosition3DAttribute &&
+				SimAccessoryMeshNormalAttribute &&
+				SimAccessoryMeshBoneIndicesAttribute &&
+				SimAccessoryMeshBoneWeightsAttribute;
 		}
 
 		return
@@ -341,7 +420,6 @@ namespace UE::Chaos::ClothAsset
 			SimFacesEnd &&
 			SimPatternFabric &&
 
-
 			// Render Patterns Group
 			RenderVerticesStart &&
 			RenderVerticesEnd &&
@@ -371,6 +449,8 @@ namespace UE::Chaos::ClothAsset
 			// Sim Vertices 2D Group
 			SimPosition2D &&
 			SimVertex3DLookup &&
+			(!EnumHasAnyFlags(OptionalSchemas, EClothCollectionExtendedSchemas::Import) ||
+				(SimImportVertexID)) &&
 
 			// Sim Vertices 3D Group
 			SimPosition3D &&
@@ -421,7 +501,15 @@ namespace UE::Chaos::ClothAsset
 			// Group Resizing Group
 			(!EnumHasAnyFlags(OptionalSchemas, EClothCollectionExtendedSchemas::Resizing) ||
 				(CustomResizingRegionSet &&
-					CustomResizingRegionType));
+					CustomResizingRegionType)) &&
+
+			// Sim Accessory Meshes Group
+			(!EnumHasAnyFlags(OptionalSchemas, EClothCollectionExtendedSchemas::SimAccessoryMeshes) ||
+				(SimAccessoryMeshName &&
+					SimAccessoryMeshPosition3DAttribute &&
+					SimAccessoryMeshNormalAttribute &&
+					SimAccessoryMeshBoneIndicesAttribute &&
+					SimAccessoryMeshBoneWeightsAttribute));
 	}
 
 	FClothCollection::FClothCollection(const TSharedRef<FManagedArrayCollection>& InManagedArrayCollection)
@@ -523,6 +611,10 @@ namespace UE::Chaos::ClothAsset
 		// Sim Vertices 2D Group
 		SimPosition2D = &GetManagedArrayCollection()->AddAttribute<FVector2f>(ClothCollectionAttribute::SimPosition2D, ClothCollectionGroup::SimVertices2D);
 		SimVertex3DLookup = &GetManagedArrayCollection()->AddAttribute<int32>(ClothCollectionAttribute::SimVertex3DLookup, ClothCollectionGroup::SimVertices2D, SimVertices3DDependency);
+		if (EnumHasAnyFlags(OptionalSchemas, EClothCollectionExtendedSchemas::Import))
+		{
+			SimImportVertexID = &GetManagedArrayCollection()->AddAttribute<int32>(ClothCollectionAttribute::SimImportVertexID, ClothCollectionGroup::SimVertices2D);
+		}
 
 		// Sim Vertices 3D Group
 		SimPosition3D = &GetManagedArrayCollection()->AddAttribute<FVector3f>(ClothCollectionAttribute::SimPosition3D, ClothCollectionGroup::SimVertices3D);
@@ -587,6 +679,16 @@ namespace UE::Chaos::ClothAsset
 			CustomResizingRegionSet = &GetManagedArrayCollection()->AddAttribute<FString>(ClothCollectionAttribute::CustomResizingRegionSet, ClothCollectionGroup::CustomResizingRegions);
 			CustomResizingRegionType = &GetManagedArrayCollection()->AddAttribute<int32>(ClothCollectionAttribute::CustomResizingRegionType, ClothCollectionGroup::CustomResizingRegions);
 		}
+
+		if (EnumHasAnyFlags(OptionalSchemas, EClothCollectionExtendedSchemas::SimAccessoryMeshes))
+		{
+			// Sim Accessory Meshes Group
+			SimAccessoryMeshName = &GetManagedArrayCollection()->AddAttribute<FName>(ClothCollectionAttribute::SimAccessoryMeshName, ClothCollectionGroup::SimAccessoryMeshes);
+			SimAccessoryMeshPosition3DAttribute = &GetManagedArrayCollection()->AddAttribute<FName>(ClothCollectionAttribute::SimAccessoryMeshPosition3DAttribute, ClothCollectionGroup::SimAccessoryMeshes);
+			SimAccessoryMeshNormalAttribute = &GetManagedArrayCollection()->AddAttribute<FName>(ClothCollectionAttribute::SimAccessoryMeshNormalAttribute, ClothCollectionGroup::SimAccessoryMeshes);
+			SimAccessoryMeshBoneIndicesAttribute = &GetManagedArrayCollection()->AddAttribute<FName>(ClothCollectionAttribute::SimAccessoryMeshBoneIndicesAttribute, ClothCollectionGroup::SimAccessoryMeshes);
+			SimAccessoryMeshBoneWeightsAttribute = &GetManagedArrayCollection()->AddAttribute<FName>(ClothCollectionAttribute::SimAccessoryMeshBoneWeightsAttribute, ClothCollectionGroup::SimAccessoryMeshes);
+		}
 	}
 
 	void FClothCollection::PostSerialize(const FArchive& Ar)
@@ -881,13 +983,19 @@ namespace UE::Chaos::ClothAsset
 
 		if (const TArray<FName>* const FixedAttributeNames = FixedAttributeNamesMap.Find(GroupName))  // Also checks that the group name is a recognized group name
 		{
+			const TArray<FName>* const FixedAttributePrefixNames = FixedAttributePrefixNamesMap.Find(GroupName);
+
 			const TArray<FName> AttributeNames = ManagedArrayCollection->AttributeNames(GroupName);
 
 			for (const FName& AttributeName : AttributeNames)
 			{
-				if (!FixedAttributeNames->Contains(AttributeName) && ManagedArrayCollection->FindAttributeTyped<T>(AttributeName, GroupName))
+				if (!FixedAttributeNames->Contains(AttributeName) && ManagedArrayCollection->FindAttributeTyped<T>(AttributeName, GroupName)
+					&& (!FixedAttributePrefixNames || !FixedAttributePrefixNames->ContainsByPredicate([&AttributeName](const FName& Prefix)
+						{
+							return Private::NamesHaveMatchingPrefix(AttributeName, Prefix);
+						})))
 				{
-					UserDefinedAttributeNames.Add(AttributeName);
+						UserDefinedAttributeNames.Add(AttributeName);
 				}
 			}
 		}
@@ -904,6 +1012,17 @@ namespace UE::Chaos::ClothAsset
 	TManagedArray<T>* FClothCollection::FindOrAddUserDefinedAttribute(const FName& Name, const FName& GroupName, const FName& GroupDependency)
 	{
 		using namespace UE::Chaos::ClothAsset::Private;
+		if (const TArray<FName>* const FixedAttributePrefixes = FixedAttributePrefixNamesMap.Find(GroupName))
+		{
+			// Don't allow adding a name reserved as a fixed attribute prefix.
+			if (FixedAttributePrefixes->ContainsByPredicate([&Name](const FName& Prefix)
+				{
+					return NamesHaveMatchingPrefix(Name, Prefix);
+				}))
+			{
+				return nullptr;
+			}
+		}
 		if (const TArray<FName>* const FixedAttributeNames = FixedAttributeNamesMap.Find(GroupName))
 		{
 			// Only allow adding to known groups. Do not allow adding a name reserved as a fixed attribute.
@@ -933,6 +1052,17 @@ namespace UE::Chaos::ClothAsset
 		using namespace UE::Chaos::ClothAsset::Private;
 		if (const TArray<FName>* const FixedAttributeNames = FixedAttributeNamesMap.Find(GroupName))
 		{
+			if (const TArray<FName>* const FixedAttributePrefixes = FixedAttributePrefixNamesMap.Find(GroupName))
+			{
+				// Don't allow removing a reserved attribute
+				if (FixedAttributePrefixes->ContainsByPredicate([&Name](const FName& Prefix)
+					{
+						return NamesHaveMatchingPrefix(Name, Prefix);
+					}))
+				{
+					return;
+				}
+			}
 			// User defined attributes are only allowed in known group names. Do not allow removing fixed attributes through this method.
 			if (!FixedAttributeNames->Contains(Name))
 			{
@@ -945,6 +1075,17 @@ namespace UE::Chaos::ClothAsset
 	bool FConstClothCollection::HasUserDefinedAttribute(const FName& Name, const FName& GroupName) const
 	{
 		using namespace UE::Chaos::ClothAsset::Private;
+		if (const TArray<FName>* const FixedAttributePrefixes = FixedAttributePrefixNamesMap.Find(GroupName))
+		{
+			// Don't consider reserved attribute prefixes as user defined.
+			if (FixedAttributePrefixes->ContainsByPredicate([&Name](const FName& Prefix)
+				{
+					return NamesHaveMatchingPrefix(Name, Prefix);
+				}))
+			{
+				return false;
+			}
+		}
 		if (const TArray<FName>* const FixedAttributeNames = FixedAttributeNamesMap.Find(GroupName))
 		{
 			// User defined attributes are only allowed in known group names.
@@ -973,6 +1114,17 @@ namespace UE::Chaos::ClothAsset
 	const TManagedArray<T>* FConstClothCollection::GetUserDefinedAttribute(const FName& Name, const FName& GroupName) const
 	{
 		using namespace UE::Chaos::ClothAsset::Private;
+		if (const TArray<FName>* const FixedAttributePrefixes = FixedAttributePrefixNamesMap.Find(GroupName))
+		{
+			// Don't consider reserved attribute prefixes as user defined.
+			if (FixedAttributePrefixes->ContainsByPredicate([&Name](const FName& Prefix)
+				{
+					return NamesHaveMatchingPrefix(Name, Prefix);
+				}))
+			{
+				return nullptr;
+			}
+		}
 		if (const TArray<FName>* const FixedAttributeNames = FixedAttributeNamesMap.Find(GroupName))
 		{
 			// User defined attributes are only allowed in known group names.
@@ -1001,6 +1153,17 @@ namespace UE::Chaos::ClothAsset
 	TManagedArray<T>* FClothCollection::GetUserDefinedAttribute(const FName& Name, const FName& GroupName)
 	{
 		using namespace UE::Chaos::ClothAsset::Private;
+		if (const TArray<FName>* const FixedAttributePrefixes = FixedAttributePrefixNamesMap.Find(GroupName))
+		{
+			// Don't consider reserved attribute prefixes as user defined.
+			if (FixedAttributePrefixes->ContainsByPredicate([&Name](const FName& Prefix)
+				{
+					return NamesHaveMatchingPrefix(Name, Prefix);
+				}))
+			{
+				return nullptr;
+			}
+		}
 		if (const TArray<FName>* const FixedAttributeNames = FixedAttributeNamesMap.Find(GroupName))
 		{
 			// User defined attributes are only allowed in known group names.
@@ -1042,6 +1205,17 @@ namespace UE::Chaos::ClothAsset
 	bool FConstClothCollection::IsValidUserDefinedAttributeName(const FName& Name, const FName& GroupName) const
 	{
 		using namespace UE::Chaos::ClothAsset::Private;
+		if (const TArray<FName>* const FixedAttributePrefixes = FixedAttributePrefixNamesMap.Find(GroupName))
+		{
+			// Don't consider reserved attribute prefixes as user defined.
+			if (FixedAttributePrefixes->ContainsByPredicate([&Name](const FName& Prefix)
+				{
+					return NamesHaveMatchingPrefix(Name, Prefix);
+				}))
+			{
+				return false;
+			}
+		}
 		if (const TArray<FName>* const FixedAttributeNames = FixedAttributeNamesMap.Find(GroupName))
 		{
 			// User defined attributes are only allowed in known group names and cannot be reserved by FixedAttributeNames.
@@ -1060,4 +1234,90 @@ namespace UE::Chaos::ClothAsset
 		}
 		return false;
 	}
+
+	const TManagedArray<FVector3f>* FConstClothCollection::GetSimAccessoryMeshPosition3D(const FName& AttributeName) const
+	{
+		if (AttributeName == NAME_None)
+		{
+			return nullptr;
+		}
+		check(Private::NamesHaveMatchingPrefix(AttributeName, ClothCollectionAttribute::SimAccessoryMeshPosition3DPrefix));
+		return ManagedArrayCollection->FindAttribute<FVector3f>(AttributeName, ClothCollectionGroup::SimVertices3D);
+	}
+
+	const TManagedArray<FVector3f>* FConstClothCollection::GetSimAccessoryMeshNormal(const FName& AttributeName) const
+	{
+		if (AttributeName == NAME_None)
+		{
+			return nullptr;
+		}
+		check(Private::NamesHaveMatchingPrefix(AttributeName, ClothCollectionAttribute::SimAccessoryMeshNormalPrefix));
+		return ManagedArrayCollection->FindAttribute<FVector3f>(AttributeName, ClothCollectionGroup::SimVertices3D);
+	}
+
+	const TManagedArray<TArray<int32>>* FConstClothCollection::GetSimAccessoryMeshBoneIndices(const FName& AttributeName) const
+	{
+		if (AttributeName == NAME_None)
+		{
+			return nullptr;
+		}
+		check(Private::NamesHaveMatchingPrefix(AttributeName, ClothCollectionAttribute::SimAccessoryMeshBoneIndicesPrefix));
+		return ManagedArrayCollection->FindAttribute<TArray<int32>>(AttributeName, ClothCollectionGroup::SimVertices3D);
+	}
+
+	const TManagedArray<TArray<float>>* FConstClothCollection::GetSimAccessoryMeshBoneWeights(const FName& AttributeName) const
+	{
+		if (AttributeName == NAME_None)
+		{
+			return nullptr;
+		}
+		check(Private::NamesHaveMatchingPrefix(AttributeName, ClothCollectionAttribute::SimAccessoryMeshBoneWeightsPrefix));
+		return ManagedArrayCollection->FindAttribute<TArray<float>>(AttributeName, ClothCollectionGroup::SimVertices3D);
+	}
+
+	FName FClothCollection::AddSimAccessoryMeshPosition3D()
+	{
+		return Private::AddUniquePrefixedAttribute<FVector3f>(ClothCollectionGroup::SimVertices3D, ClothCollectionAttribute::SimAccessoryMeshPosition3DPrefix, GetSimAccessoryMeshPosition3DAttribute(), GetManagedArrayCollection());
+	}
+
+	FName FClothCollection::AddSimAccessoryMeshNormal()
+	{
+		return Private::AddUniquePrefixedAttribute<FVector3f>(ClothCollectionGroup::SimVertices3D, ClothCollectionAttribute::SimAccessoryMeshNormalPrefix, GetSimAccessoryMeshNormalAttribute(), GetManagedArrayCollection());
+	}
+
+	FName FClothCollection::AddSimAccessoryMeshBoneIndices()
+	{
+		return Private::AddUniquePrefixedAttribute<TArray<int32>>(ClothCollectionGroup::SimVertices3D, ClothCollectionAttribute::SimAccessoryMeshBoneIndicesPrefix, GetSimAccessoryMeshBoneIndicesAttribute(), GetManagedArrayCollection());
+	}
+
+	FName FClothCollection::AddSimAccessoryMeshBoneWeights()
+	{
+		return Private::AddUniquePrefixedAttribute<TArray<float>>(ClothCollectionGroup::SimVertices3D, ClothCollectionAttribute::SimAccessoryMeshBoneWeightsPrefix, GetSimAccessoryMeshBoneWeightsAttribute(), GetManagedArrayCollection());
+	}
+
+	void FClothCollection::RemoveSimAccessoryMeshPosition3D(const FName& AttributeName)
+	{
+		check(Private::NamesHaveMatchingPrefix(AttributeName, ClothCollectionAttribute::SimAccessoryMeshPosition3DPrefix));
+		GetManagedArrayCollection()->RemoveAttribute(AttributeName, ClothCollectionGroup::SimVertices3D);
+	}
+
+	void FClothCollection::RemoveSimAccessoryMeshNormal(const FName& AttributeName)
+	{
+		check(Private::NamesHaveMatchingPrefix(AttributeName, ClothCollectionAttribute::SimAccessoryMeshNormalPrefix));
+		GetManagedArrayCollection()->RemoveAttribute(AttributeName, ClothCollectionGroup::SimVertices3D);
+	}
+
+	void FClothCollection::RemoveSimAccessoryMeshBoneIndices(const FName& AttributeName)
+	{
+		check(Private::NamesHaveMatchingPrefix(AttributeName, ClothCollectionAttribute::SimAccessoryMeshBoneIndicesPrefix));
+		GetManagedArrayCollection()->RemoveAttribute(AttributeName, ClothCollectionGroup::SimVertices3D);
+	}
+
+	void FClothCollection::RemoveSimAccessoryMeshBoneWeights(const FName& AttributeName)
+	{
+		check(Private::NamesHaveMatchingPrefix(AttributeName, ClothCollectionAttribute::SimAccessoryMeshBoneWeightsPrefix));
+		GetManagedArrayCollection()->RemoveAttribute(AttributeName, ClothCollectionGroup::SimVertices3D);
+	}
+
+
 } // End namespace UE::Chaos::ClothAsset
diff --git a/Engine/Plugins/ChaosClothAsset/Source/ChaosClothAsset/Private/ChaosClothAsset/ClothCollection.h b/Engine/Plugins/ChaosClothAsset/Source/ChaosClothAsset/Private/ChaosClothAsset/ClothCollection.h
--- a/Engine/Plugins/ChaosClothAsset/Source/ChaosClothAsset/Private/ChaosClothAsset/ClothCollection.h
+++ b/Engine/Plugins/ChaosClothAsset/Source/ChaosClothAsset/Private/ChaosClothAsset/ClothCollection.h
@@ -188,6 +188,7 @@ namespace UE::Chaos::ClothAsset
 		//~ Sim Vertices 2D Group
 		const TManagedArray<FVector2f>* GetSimPosition2D() const { return SimPosition2D; }
 		const TManagedArray<int32>* GetSimVertex3DLookup() const { return SimVertex3DLookup; }
+		const TManagedArray<int32>* GetSimImportVertexID() const { return SimImportVertexID; }
 
 		//~ Sim Vertices 3D Group
 		const TManagedArray<FVector3f>* GetSimPosition3D() const { return SimPosition3D; }
@@ -201,6 +202,12 @@ namespace UE::Chaos::ClothAsset
 		const TManagedArray<TArray<int32>>* GetSeamStitchLookup() const { return SeamStitchLookup; }
 		const TManagedArray<float>* GetSimCustomResizingBlend() const { return SimCustomResizingBlend; }
 
+		//~ Sim Vertices 3D Group -- Sim Accessory Mesh
+		const TManagedArray<FVector3f>* GetSimAccessoryMeshPosition3D(const FName& AttributeName) const;
+		const TManagedArray<FVector3f>* GetSimAccessoryMeshNormal(const FName& AttributeName) const;
+		const TManagedArray<TArray<int32>>* GetSimAccessoryMeshBoneIndices(const FName& AttributeName) const;
+		const TManagedArray<TArray<float>>* GetSimAccessoryMeshBoneWeights(const FName& AttributeName) const;
+
 		//~ Sim Morph Targets Group
 		const TManagedArray<FString>* GetSimMorphTargetName() const { return SimMorphTargetName; }
 		const TManagedArray<int32>* GetSimMorphTargetVerticesStart() const { return SimMorphTargetVerticesStart; }
@@ -235,6 +242,13 @@ namespace UE::Chaos::ClothAsset
 		const TManagedArray<FString>* GetCustomResizingRegionSet() const { return CustomResizingRegionSet; }
 		const TManagedArray<int32>* GetCustomResizingRegionType() const { return CustomResizingRegionType; }
 
+		//~ Sim Accessory Meshes Group
+		const TManagedArray<FName>* GetSimAccessoryMeshName() const { return SimAccessoryMeshName; }
+		const TManagedArray<FName>* GetSimAccessoryMeshPosition3DAttribute() const { return SimAccessoryMeshPosition3DAttribute; }
+		const TManagedArray<FName>* GetSimAccessoryMeshNormalAttribute() const { return SimAccessoryMeshNormalAttribute; }
+		const TManagedArray<FName>* GetSimAccessoryMeshBoneIndicesAttribute() const { return SimAccessoryMeshBoneIndicesAttribute; }
+		const TManagedArray<FName>* GetSimAccessoryMeshBoneWeightsAttribute() const { return SimAccessoryMeshBoneWeightsAttribute; }
+
 	protected:
 		struct TProtectedInit {};
 
@@ -303,6 +317,7 @@ namespace UE::Chaos::ClothAsset
 		//~ Sim Vertices 2D Group
 		const TManagedArray<FVector2f>* SimPosition2D;
 		const TManagedArray<int32>* SimVertex3DLookup; // Lookup into corresponding 3D vertices
+		const TManagedArray<int32>* SimImportVertexID; // Can be populated by importers to match data 
 
 		//~ Sim Vertices 3D Group
 		const TManagedArray<FVector3f>* SimPosition3D;
@@ -349,6 +364,13 @@ namespace UE::Chaos::ClothAsset
 		//~ Group Resizing Group
 		const TManagedArray<FString>* CustomResizingRegionSet; // Name of the set associated with this binding type
 		const TManagedArray<int32>* CustomResizingRegionType; // Actually an Enum for binding type
+
+		//~ Sim Accessory Meshes Group
+		const TManagedArray<FName>* SimAccessoryMeshName;
+		const TManagedArray<FName>* SimAccessoryMeshPosition3DAttribute;
+		const TManagedArray<FName>* SimAccessoryMeshNormalAttribute;
+		const TManagedArray<FName>* SimAccessoryMeshBoneIndicesAttribute;
+		const TManagedArray<FName>* SimAccessoryMeshBoneWeightsAttribute;
 	};
 
 	/**
@@ -456,6 +478,7 @@ namespace UE::Chaos::ClothAsset
 		//~ Sim Vertices 2D Group
 		TManagedArray<FVector2f>* GetSimPosition2D() { return const_cast<TManagedArray<FVector2f>*>(SimPosition2D); }
 		TManagedArray<int32>* GetSimVertex3DLookup() { return const_cast<TManagedArray<int32>*>(SimVertex3DLookup); }
+		TManagedArray<int32>* GetSimImportVertexID() { return const_cast<TManagedArray<int32>*>(SimImportVertexID); }
 
 		//~ Sim Vertices 3D Group
 		TManagedArray<FVector3f>* GetSimPosition3D() { return const_cast<TManagedArray<FVector3f>*>(SimPosition3D); }
@@ -469,6 +492,25 @@ namespace UE::Chaos::ClothAsset
 		TManagedArray<TArray<int32>>* GetSeamStitchLookup() { return const_cast<TManagedArray<TArray<int32>>*>(SeamStitchLookup); }
 		TManagedArray<float>* GetSimCustomResizingBlend() { return const_cast<TManagedArray<float>*>(SimCustomResizingBlend); }
 
+		//~ Sim Vertices 3D Group -- Sim Accessory Mesh
+		TManagedArray<FVector3f>* GetSimAccessoryMeshPosition3D(const FName& AttributeName) { return const_cast<TManagedArray<FVector3f>*>(((FConstClothCollection*)this)->GetSimAccessoryMeshPosition3D(AttributeName)); }
+		TManagedArray<FVector3f>* GetSimAccessoryMeshNormal(const FName& AttributeName) const { return const_cast<TManagedArray<FVector3f>*>(((FConstClothCollection*)this)->GetSimAccessoryMeshNormal(AttributeName)); }
+		TManagedArray<TArray<int32>>* GetSimAccessoryMeshBoneIndices(const FName& AttributeName) const { return const_cast<TManagedArray<TArray<int32>>*>(((FConstClothCollection*)this)->GetSimAccessoryMeshBoneIndices(AttributeName)); }
+		TManagedArray<TArray<float>>* GetSimAccessoryMeshBoneWeights(const FName& AttributeName) const { return const_cast<TManagedArray<TArray<float>>*>(((FConstClothCollection*)this)->GetSimAccessoryMeshBoneWeights(AttributeName)); }
+
+		// Add accessory mesh attribute and return resulting attribute's name. Note: these currently do a linear search through existing attributes to generate a unique name.
+		// This does not record the new attribute in the associated SimAccessoryMesh group Attribute 
+		FName AddSimAccessoryMeshPosition3D();
+		FName AddSimAccessoryMeshNormal();
+		FName AddSimAccessoryMeshBoneIndices();
+		FName AddSimAccessoryMeshBoneWeights();
+		// Remove accessory mesh attribute by name.
+		// This does not reset the attribute in the associated SimAccessoryMesh group Attribute 
+		void RemoveSimAccessoryMeshPosition3D(const FName& AttributeName);
+		void RemoveSimAccessoryMeshNormal(const FName& AttributeName);
+		void RemoveSimAccessoryMeshBoneIndices(const FName& AttributeName);
+		void RemoveSimAccessoryMeshBoneWeights(const FName& AttributeName);
+
 		//~ Sim Morph Targets Group
 		TManagedArray<FString>* GetSimMorphTargetName() { return const_cast<TManagedArray<FString>*>(SimMorphTargetName); }
 		TManagedArray<int32>* GetSimMorphTargetVerticesStart() { return const_cast<TManagedArray<int32>*>(SimMorphTargetVerticesStart); }
@@ -503,6 +545,13 @@ namespace UE::Chaos::ClothAsset
 		TManagedArray<FString>* GetCustomResizingRegionSet() { return const_cast<TManagedArray<FString>*>(CustomResizingRegionSet); }
 		TManagedArray<int32>* GetCustomResizingRegionType() { return const_cast<TManagedArray<int32>*>(CustomResizingRegionType); }
 
+		//~ Sim Accessory Meshes Group
+		TManagedArray<FName>* GetSimAccessoryMeshName() const { return const_cast<TManagedArray<FName>*>(SimAccessoryMeshName); }
+		TManagedArray<FName>* GetSimAccessoryMeshPosition3DAttribute() const { return const_cast<TManagedArray<FName>*>(SimAccessoryMeshPosition3DAttribute); }
+		TManagedArray<FName>* GetSimAccessoryMeshNormalAttribute() const { return const_cast<TManagedArray<FName>*>(SimAccessoryMeshNormalAttribute); }
+		TManagedArray<FName>* GetSimAccessoryMeshBoneIndicesAttribute() const { return const_cast<TManagedArray<FName>*>(SimAccessoryMeshBoneIndicesAttribute); }
+		TManagedArray<FName>* GetSimAccessoryMeshBoneWeightsAttribute() const { return const_cast<TManagedArray<FName>*>(SimAccessoryMeshBoneWeightsAttribute); }
+
 	private:
 		TSharedRef<FManagedArrayCollection> GetManagedArrayCollection()
 		{
diff --git a/Engine/Plugins/ChaosClothAsset/Source/ChaosClothAsset/Private/ChaosClothAsset/ClothGeometryTools.cpp b/Engine/Plugins/ChaosClothAsset/Source/ChaosClothAsset/Private/ChaosClothAsset/ClothGeometryTools.cpp
--- a/Engine/Plugins/ChaosClothAsset/Source/ChaosClothAsset/Private/ChaosClothAsset/ClothGeometryTools.cpp
+++ b/Engine/Plugins/ChaosClothAsset/Source/ChaosClothAsset/Private/ChaosClothAsset/ClothGeometryTools.cpp
@@ -1888,4 +1888,137 @@ PRAGMA_ENABLE_DEPRECATION_WARNINGS
 			}
 		}
 	}
+
+	int32 FClothGeometryTools::CopySimMeshToSimAccessoryMesh(const FName& AccessoryMeshName, FCollectionClothFacade& ToCloth, const FCollectionClothConstFacade& FromCloth, bool bUseSimImportVertexID)
+	{
+		auto GenerateUniqueAccessoryMeshName = [&AccessoryMeshName, &ToCloth]()
+			{
+				int32 MaxExistingNumber = -1;
+				bool bFoundMatchWithIndex = false;
+				for (const FName& Name : ToCloth.GetSimAccessoryMeshName())
+				{
+					if (Name.GetComparisonIndex() == AccessoryMeshName.GetComparisonIndex())
+					{
+						MaxExistingNumber = FMath::Max(MaxExistingNumber, Name.GetNumber());
+					}
+					bFoundMatchWithIndex = bFoundMatchWithIndex || Name == AccessoryMeshName;
+				}
+
+				if (!bFoundMatchWithIndex)
+				{
+					return AccessoryMeshName;
+				}
+				FName UniqueName = AccessoryMeshName;
+				UniqueName.SetNumber(MaxExistingNumber + 1);
+				return UniqueName;
+			};
+
+		if (bUseSimImportVertexID) 
+		{
+			if (ToCloth.IsValid(EClothCollectionExtendedSchemas::Import) && FromCloth.IsValid(EClothCollectionExtendedSchemas::Import))
+			{
+				// Start with existing sim data
+				TArray<FVector3f> Positions(ToCloth.GetSimPosition3D());
+				TArray<FVector3f> Normals(ToCloth.GetSimNormal());
+				TArray<TArray<int32>> BoneIndices(ToCloth.GetSimBoneIndices());
+				TArray<TArray<float>> BoneWeights(ToCloth.GetSimBoneWeights());
+
+				TConstArrayView<int32> ToImportVertex = ToCloth.GetSimImportVertexID();
+				TConstArrayView<int32> FromImportVertex = FromCloth.GetSimImportVertexID();
+
+				TMap<int32, TArray<int32>> ToImportToSim2DLookup;
+				ToImportToSim2DLookup.Reserve(ToImportVertex.Num());
+				for (int32 ToSim2DIndex = 0; ToSim2DIndex < ToImportVertex.Num(); ++ToSim2DIndex)
+				{
+					ToImportToSim2DLookup.FindOrAdd(ToImportVertex[ToSim2DIndex]).Add(ToSim2DIndex);
+				}
+
+				TConstArrayView<int32> ToSim3DLookup = ToCloth.GetSimVertex3DLookup();
+				TConstArrayView<int32> FromSim3DLookup = FromCloth.GetSimVertex3DLookup();
+				TConstArrayView<FVector3f> FromPositions = FromCloth.GetSimPosition3D();
+				TConstArrayView<FVector3f> FromNormals = FromCloth.GetSimNormal();
+				TConstArrayView<TArray<int32>> FromBoneIndices = FromCloth.GetSimBoneIndices();
+				TConstArrayView<TArray<float>> FromBoneWeights = FromCloth.GetSimBoneWeights();
+				for (int32 FromSim2DIndex = 0; FromSim2DIndex < FromImportVertex.Num(); ++FromSim2DIndex)
+				{
+					const int32 FromSim3DIndex = FromSim3DLookup[FromSim2DIndex];
+					if (FromPositions.IsValidIndex(FromSim3DIndex))
+					{
+						if (const TArray<int32>* const ToSim2DIndices = ToImportToSim2DLookup.Find(FromImportVertex[FromSim2DIndex]))
+						{
+							// Found matching index on ToCloth.
+							// Copy data between associated sim3d vertices.
+							// NOTE: this will just last one wins for vertices that did not seam the same way between the two meshes. Hopefully that's good enough.
+							for (const int32 ToSim2DIndex : *ToSim2DIndices)
+							{
+								const int32 ToSim3DIndex = ToSim3DLookup[ToSim2DIndex];
+								if (Positions.IsValidIndex(ToSim3DIndex))
+								{
+									Positions[ToSim3DIndex] = FromPositions[FromSim3DIndex];
+									Normals[ToSim3DIndex] = FromNormals[FromSim3DIndex];
+									BoneIndices[ToSim3DIndex] = FromBoneIndices[FromSim3DIndex];
+									BoneWeights[ToSim3DIndex] = FromBoneWeights[FromSim3DIndex];
+								}
+							}
+						}
+					}
+				}
+
+				FCollectionClothSimAccessoryMeshFacade AccessoryFacade = ToCloth.AddGetSimAccessoryMesh();
+				AccessoryFacade.Initialize(GenerateUniqueAccessoryMeshName(), TConstArrayView<FVector3f>(Positions), TConstArrayView<FVector3f>(Normals), TConstArrayView<TArray<int32>>(BoneIndices), TConstArrayView<TArray<float>>(BoneWeights));
+				return AccessoryFacade.GetSimAccessoryMeshIndex();
+			}
+
+			// no valid import data
+			return INDEX_NONE;
+		}
+
+		if (FromCloth.GetNumSimVertices3D() >= ToCloth.GetNumSimVertices3D())
+		{
+			// Can just initialize directly 
+			const int32 NumToSimVertices3D = ToCloth.GetNumSimVertices3D();
+			FCollectionClothSimAccessoryMeshFacade AccessoryFacade = ToCloth.AddGetSimAccessoryMesh();
+			AccessoryFacade.Initialize(GenerateUniqueAccessoryMeshName(), ToCloth.GetSimPosition3D().Left(NumToSimVertices3D), ToCloth.GetSimNormal().Left(NumToSimVertices3D), ToCloth.GetSimBoneIndices().Left(NumToSimVertices3D), ToCloth.GetSimBoneWeights().Left(NumToSimVertices3D));
+			return AccessoryFacade.GetSimAccessoryMeshIndex();
+		}
+
+		// Fill in tail data with ToCloth's sim mesh data.
+		const int32 NumToSimVertices3D = ToCloth.GetNumSimVertices3D();
+		TArray<FVector3f> Positions;
+		TArray<FVector3f> Normals;
+		TArray<TArray<int32>> BoneIndices;
+		TArray<TArray<float>> BoneWeights;
+		Positions.Reserve(NumToSimVertices3D);
+		Normals.Reserve(NumToSimVertices3D);
+		BoneIndices.Reserve(NumToSimVertices3D);
+		BoneWeights.Reserve(NumToSimVertices3D);
+
+		TConstArrayView<FVector3f> FromPositions = FromCloth.GetSimPosition3D();
+		TConstArrayView<FVector3f> FromNormals = FromCloth.GetSimNormal();
+		TConstArrayView<TArray<int32>> FromBoneIndices = FromCloth.GetSimBoneIndices();
+		TConstArrayView<TArray<float>> FromBoneWeights = FromCloth.GetSimBoneWeights();
+		for (int32 Index = 0; Index < FromCloth.GetNumSimVertices3D(); ++Index)
+		{
+			Positions.Emplace(FromPositions[Index]);
+			Normals.Emplace(FromNormals[Index]);
+			BoneIndices.Emplace(FromBoneIndices[Index]);
+			BoneWeights.Emplace(FromBoneWeights[Index]);
+		}
+		TConstArrayView<FVector3f> ToPositions = ToCloth.GetSimPosition3D();
+		TConstArrayView<FVector3f> ToNormals = ToCloth.GetSimNormal();
+		TConstArrayView<TArray<int32>> ToBoneIndices = ToCloth.GetSimBoneIndices();
+		TConstArrayView<TArray<float>> ToBoneWeights = ToCloth.GetSimBoneWeights();
+		for (int32 Index = FromCloth.GetNumSimVertices3D(); Index < NumToSimVertices3D; ++Index)
+		{
+			Positions.Emplace(ToPositions[Index]);
+			Normals.Emplace(ToNormals[Index]);
+			BoneIndices.Emplace(ToBoneIndices[Index]);
+			BoneWeights.Emplace(ToBoneWeights[Index]);
+		}
+
+		FCollectionClothSimAccessoryMeshFacade AccessoryFacade = ToCloth.AddGetSimAccessoryMesh();
+		AccessoryFacade.Initialize(GenerateUniqueAccessoryMeshName(), TConstArrayView<FVector3f>(Positions), TConstArrayView<FVector3f>(Normals), TConstArrayView<TArray<int32>>(BoneIndices), TConstArrayView<TArray<float>>(BoneWeights));
+		return AccessoryFacade.GetSimAccessoryMeshIndex();
+
+	}
 }  // End namespace UE::Chaos::ClothAsset
diff --git a/Engine/Plugins/ChaosClothAsset/Source/ChaosClothAsset/Private/ChaosClothAsset/CollectionClothFacade.cpp b/Engine/Plugins/ChaosClothAsset/Source/ChaosClothAsset/Private/ChaosClothAsset/CollectionClothFacade.cpp
--- a/Engine/Plugins/ChaosClothAsset/Source/ChaosClothAsset/Private/ChaosClothAsset/CollectionClothFacade.cpp
+++ b/Engine/Plugins/ChaosClothAsset/Source/ChaosClothAsset/Private/ChaosClothAsset/CollectionClothFacade.cpp
@@ -252,6 +252,11 @@ namespace UE::Chaos::ClothAsset
 		return ClothCollection->GetElements(ClothCollection->GetSimVertex3DLookup());
 	}
 
+	TConstArrayView<int32> FCollectionClothConstFacade::GetSimImportVertexID() const
+	{
+		return ClothCollection->GetElements(ClothCollection->GetSimImportVertexID());
+	}
+
 	int32 FCollectionClothConstFacade::GetNumSimVertices3D() const
 	{
 		return ClothCollection->GetNumElements(ClothCollectionGroup::SimVertices3D);
@@ -390,6 +395,54 @@ namespace UE::Chaos::ClothAsset
 		return ClothCollection->GetElements(ClothCollection->GetSimMorphTargetSimVertex3DIndex());
 	}
 
+	int32 FCollectionClothConstFacade::GetNumSimAccessoryMeshes() const
+	{
+		return ClothCollection->GetNumElements(ClothCollectionGroup::SimAccessoryMeshes);
+	}
+
+	int32 FCollectionClothConstFacade::FindSimAccessoryMeshIndexByName(const FName& AccessoryMeshName) const
+	{
+		TConstArrayView<FName> SimAccessoryMeshNames = GetSimAccessoryMeshName();
+		for (int32 Index = 0; Index < SimAccessoryMeshNames.Num(); ++Index)
+		{
+			if (SimAccessoryMeshNames[Index] == AccessoryMeshName)
+			{
+				return Index;
+			}
+		}
+		return INDEX_NONE;
+	}
+
+	FCollectionClothSimAccessoryMeshConstFacade FCollectionClothConstFacade::GetSimAccessoryMesh(int32 MeshIndex) const
+	{
+		return FCollectionClothSimAccessoryMeshConstFacade(ClothCollection, MeshIndex);
+	}
+
+	TConstArrayView<FName> FCollectionClothConstFacade::GetSimAccessoryMeshName() const
+	{
+		return ClothCollection->GetElements(ClothCollection->GetSimAccessoryMeshName());
+	}
+
+	TConstArrayView<FName> FCollectionClothConstFacade::GetSimAccessoryMeshPosition3DAttribute() const
+	{
+		return ClothCollection->GetElements(ClothCollection->GetSimAccessoryMeshPosition3DAttribute());
+	}
+
+	TConstArrayView<FName> FCollectionClothConstFacade::GetSimAccessoryMeshNormalAttribute() const
+	{
+		return ClothCollection->GetElements(ClothCollection->GetSimAccessoryMeshNormalAttribute());
+	}
+
+	TConstArrayView<FName> FCollectionClothConstFacade::GetSimAccessoryMeshBoneIndicesAttribute() const
+	{
+		return ClothCollection->GetElements(ClothCollection->GetSimAccessoryMeshBoneIndicesAttribute());
+	}
+
+	TConstArrayView<FName> FCollectionClothConstFacade::GetSimAccessoryMeshBoneWeightsAttribute() const
+	{
+		return ClothCollection->GetElements(ClothCollection->GetSimAccessoryMeshBoneWeightsAttribute());
+	}
+
 	int32 FCollectionClothConstFacade::GetNumRenderPatterns() const
 	{
 		return ClothCollection->GetNumElements(ClothCollectionGroup::RenderPatterns);
@@ -678,6 +731,7 @@ namespace UE::Chaos::ClothAsset
 		}
 		SetNumSimMorphTargets(0);
 		SetNumCustomResizingRegions(0);
+		SetNumSimAccessoryMeshes(0);
 	}
 
 	void FCollectionClothFacade::Initialize(const FCollectionClothConstFacade& Other)
@@ -759,6 +813,7 @@ namespace UE::Chaos::ClothAsset
 		}
 
 		// Sim Patterns Group
+		const bool bHadImportData = IsValid(EClothCollectionExtendedSchemas::Import);
 		const int32 StartNumSimVertices2D = GetNumSimVertices2D();
 		const int32 StartNumSimPatterns = GetNumSimPatterns();
 		const int32 OtherNumSimPatterns = Other.GetNumSimPatterns();
@@ -767,6 +822,16 @@ namespace UE::Chaos::ClothAsset
 		{
 			GetSimPattern(StartNumSimPatterns + PatternIndex).Initialize(Other.GetSimPattern(PatternIndex), StartNumSimVertices3D, StartNumFabrics);
 		}
+		const bool bOtherHasImportData = Other.IsValid(EClothCollectionExtendedSchemas::Import);
+		if (bOtherHasImportData && !bHadImportData)
+		{
+			DefineSchema(EClothCollectionExtendedSchemas::Import);
+			TArrayView<int32> ImportVertexID = GetSimImportVertexID();
+			for (int32 Index = 0; Index < StartNumSimVertices2D; ++Index)
+			{
+				ImportVertexID[Index] = INDEX_NONE;
+			}
+		}
 
 		// Seams Group
 		const int32 StartNumSeamStitches = ClothCollection->GetNumElements(ClothCollectionGroup::SeamStitches);
@@ -822,6 +887,49 @@ namespace UE::Chaos::ClothAsset
 			FClothCollection::CopyArrayViewData(GetCustomResizingRegionType().Right(OtherNumResizingBindingGroups), Other.GetCustomResizingRegionType());
 		}
 
+		// Sim Accessory Meshes
+		// Merge accessory meshes with the same name.
+		// Fill in with existing sim mesh data when an accessory mesh exists in one end of the merge but not the other.
+		const int32 StartNumAccessoryMeshes = GetNumSimAccessoryMeshes();
+		const int32 OtherNumAccessoryMeshes = Other.GetNumSimAccessoryMeshes();
+		TConstArrayView<FName> OtherSimAccessoryMeshNames = Other.GetSimAccessoryMeshName();
+		for (int32 ExistingMeshIndex = 0; ExistingMeshIndex < StartNumAccessoryMeshes; ++ExistingMeshIndex)
+		{
+			FCollectionClothSimAccessoryMeshFacade MeshFacade = GetSimAccessoryMesh(ExistingMeshIndex);
+			const int32 OtherMatchingMeshIndex = OtherSimAccessoryMeshNames.IndexOfByKey(MeshFacade.GetSimAccessoryMeshName());
+			if (OtherMatchingMeshIndex != INDEX_NONE)
+			{
+				FCollectionClothSimAccessoryMeshConstFacade OtherMeshFacade = Other.GetSimAccessoryMesh(OtherMatchingMeshIndex);
+				FClothCollection::CopyArrayViewData(MeshFacade.GetSimAccessoryMeshPosition3D().Right(OtherNumSimVertices3D), OtherMeshFacade.GetSimAccessoryMeshPosition3D());
+				FClothCollection::CopyArrayViewData(MeshFacade.GetSimAccessoryMeshNormal().Right(OtherNumSimVertices3D), OtherMeshFacade.GetSimAccessoryMeshNormal());
+				FClothCollection::CopyArrayViewData(MeshFacade.GetSimAccessoryMeshBoneIndices().Right(OtherNumSimVertices3D), OtherMeshFacade.GetSimAccessoryMeshBoneIndices());
+				FClothCollection::CopyArrayViewData(MeshFacade.GetSimAccessoryMeshBoneWeights().Right(OtherNumSimVertices3D), OtherMeshFacade.GetSimAccessoryMeshBoneWeights());
+			}
+			else
+			{
+				FClothCollection::CopyArrayViewData(MeshFacade.GetSimAccessoryMeshPosition3D().Right(OtherNumSimVertices3D), Other.GetSimPosition3D());
+				FClothCollection::CopyArrayViewData(MeshFacade.GetSimAccessoryMeshNormal().Right(OtherNumSimVertices3D), Other.GetSimNormal());
+				FClothCollection::CopyArrayViewData(MeshFacade.GetSimAccessoryMeshBoneIndices().Right(OtherNumSimVertices3D), Other.GetSimBoneIndices());
+				FClothCollection::CopyArrayViewData(MeshFacade.GetSimAccessoryMeshBoneWeights().Right(OtherNumSimVertices3D), Other.GetSimBoneWeights());
+			}
+		}
+		for (int32 OtherMeshIndex = 0; OtherMeshIndex < OtherNumAccessoryMeshes; ++OtherMeshIndex)
+		{
+			FCollectionClothSimAccessoryMeshConstFacade OtherMeshFacade = Other.GetSimAccessoryMesh(OtherMeshIndex);
+			const int32 MatchingMeshIndex = GetSimAccessoryMeshName().IndexOfByKey(OtherMeshFacade.GetSimAccessoryMeshName());
+			if (MatchingMeshIndex == INDEX_NONE)
+			{
+				FCollectionClothSimAccessoryMeshFacade MeshFacade = AddGetSimAccessoryMesh();
+				// Initialize with sim mesh data and then copy over other accessory mesh data.
+				MeshFacade.Initialize(OtherMeshFacade.GetSimAccessoryMeshName(), GetSimPosition3D(), GetSimNormal(), GetSimBoneIndices(), GetSimBoneWeights());
+				FClothCollection::CopyArrayViewData(MeshFacade.GetSimAccessoryMeshPosition3D().Right(OtherNumSimVertices3D), OtherMeshFacade.GetSimAccessoryMeshPosition3D());
+				FClothCollection::CopyArrayViewData(MeshFacade.GetSimAccessoryMeshNormal().Right(OtherNumSimVertices3D), OtherMeshFacade.GetSimAccessoryMeshNormal());
+				FClothCollection::CopyArrayViewData(MeshFacade.GetSimAccessoryMeshBoneIndices().Right(OtherNumSimVertices3D), OtherMeshFacade.GetSimAccessoryMeshBoneIndices());
+				FClothCollection::CopyArrayViewData(MeshFacade.GetSimAccessoryMeshBoneWeights().Right(OtherNumSimVertices3D), OtherMeshFacade.GetSimAccessoryMeshBoneWeights());				
+
+			} // Else case handled above
+		}
+
 		// Weight maps
 		const TArray<FName> WeightMapNames = Other.GetWeightMapNames();
 		for (const FName& WeightMapName : WeightMapNames)
@@ -879,6 +987,15 @@ namespace UE::Chaos::ClothAsset
 		// Sim Vertices 3D Group
 		GetClothCollection()->SetNumElements(Other.GetNumSimVertices3D(), ClothCollectionGroup::SimVertices3D);
 		FClothCollection::CopyArrayViewData(GetPreResizedSimPosition3D().Right(Other.GetNumSimVertices3D()), Other.GetPreResizedSimPosition3D());
+
+		// Sim AccessoryMeshes
+		SetNumSimAccessoryMeshes(0);
+		const int32 OtherNumSimAccessoryMeshes = Other.GetNumSimAccessoryMeshes();
+		SetNumSimAccessoryMeshes(OtherNumSimAccessoryMeshes);
+		for (int32 MeshIndex = 0; MeshIndex < OtherNumSimAccessoryMeshes; ++MeshIndex)
+		{
+			GetSimAccessoryMesh(MeshIndex).Initialize(Other.GetSimAccessoryMesh(MeshIndex));
+		}
 	}
 
 	void FCollectionClothFacade::SetPhysicsAssetSoftObjectPathName(const FSoftObjectPath& PathName)
@@ -951,6 +1068,15 @@ namespace UE::Chaos::ClothAsset
 		return GetClothCollection()->GetElements(GetClothCollection()->GetSimVertex3DLookup());
 	}
 
+	TArrayView<int32> FCollectionClothFacade::GetSimImportVertexID()
+	{
+		if (!IsValid(EClothCollectionExtendedSchemas::Import))
+		{
+			DefineSchema(EClothCollectionExtendedSchemas::Import);
+		}
+		return GetClothCollection()->GetElements(GetClothCollection()->GetSimImportVertexID());
+	}
+
 	TArrayView<FVector3f> FCollectionClothFacade::GetSimPosition3D()
 	{
 		return GetClothCollection()->GetElements(GetClothCollection()->GetSimPosition3D());
@@ -1175,6 +1301,75 @@ namespace UE::Chaos::ClothAsset
 		GetClothCollection()->RemoveElements(ClothCollectionGroup::SimMorphTargetVertices, SortedDeletionList);
 	}
 
+	void FCollectionClothFacade::SetNumSimAccessoryMeshes(int32 InNumMeshes, bool bCookedOnly)
+	{
+		const EClothCollectionExtendedSchemas OptionalSchemas = bCookedOnly ? EClothCollectionExtendedSchemas::CookedOnly : EClothCollectionExtendedSchemas::SimAccessoryMeshes;
+		if (!IsValid(OptionalSchemas))
+		{
+			DefineSchema(OptionalSchemas);
+		}
+
+		const int32 NumMeshes = GetNumSimAccessoryMeshes();
+
+		for (int32 MeshIndex = InNumMeshes; MeshIndex < NumMeshes; ++MeshIndex)
+		{
+			GetSimAccessoryMesh(MeshIndex).Reset();
+		}
+
+		GetClothCollection()->SetNumElements(InNumMeshes, ClothCollectionGroup::SimAccessoryMeshes);
+
+		for (int32 MeshIndex = NumMeshes; MeshIndex < InNumMeshes; ++MeshIndex)
+		{
+			GetSimAccessoryMesh(MeshIndex).SetDefaults();
+		}
+	}
+
+	int32 FCollectionClothFacade::AddSimAccessoryMesh()
+	{
+		const int32 MeshIndex = GetNumSimAccessoryMeshes();
+		SetNumSimAccessoryMeshes(MeshIndex + 1);
+		return MeshIndex;
+	}
+
+	FCollectionClothSimAccessoryMeshFacade FCollectionClothFacade::GetSimAccessoryMesh(int32 MeshIndex)
+	{
+		return FCollectionClothSimAccessoryMeshFacade(GetClothCollection(), MeshIndex);
+	}
+
+	void FCollectionClothFacade::RemoveSimAccessoryMeshes(const TArray<int32>& SortedDeletionList)
+	{
+		for (const int32 MeshToRemove : SortedDeletionList)
+		{
+			GetSimAccessoryMesh(MeshToRemove).Reset();
+		}
+		GetClothCollection()->RemoveElements(ClothCollectionGroup::SimAccessoryMeshes, SortedDeletionList);
+	}
+
+	TArrayView<FName> FCollectionClothFacade::GetSimAccessoryMeshName()
+	{
+		return GetClothCollection()->GetElements(GetClothCollection()->GetSimAccessoryMeshName());
+	}
+
+	TArrayView<FName> FCollectionClothFacade::GetSimAccessoryMeshPosition3DAttribute()
+	{
+		return GetClothCollection()->GetElements(GetClothCollection()->GetSimAccessoryMeshPosition3DAttribute());
+	}
+
+	TArrayView<FName> FCollectionClothFacade::GetSimAccessoryMeshNormalAttribute()
+	{
+		return GetClothCollection()->GetElements(GetClothCollection()->GetSimAccessoryMeshNormalAttribute());
+	}
+
+	TArrayView<FName> FCollectionClothFacade::GetSimAccessoryMeshBoneIndicesAttribute()
+	{
+		return GetClothCollection()->GetElements(GetClothCollection()->GetSimAccessoryMeshBoneIndicesAttribute());
+	}
+
+	TArrayView<FName> FCollectionClothFacade::GetSimAccessoryMeshBoneWeightsAttribute()
+	{
+		return GetClothCollection()->GetElements(GetClothCollection()->GetSimAccessoryMeshBoneWeightsAttribute());
+	}
+
 	void FCollectionClothFacade::SetNumRenderPatterns(int32 InNumPatterns)
 	{
 		const int32 NumPatterns = GetNumRenderPatterns();
diff --git a/Engine/Plugins/ChaosClothAsset/Source/ChaosClothAsset/Private/ChaosClothAsset/CollectionClothSeamFacade.cpp b/Engine/Plugins/ChaosClothAsset/Source/ChaosClothAsset/Private/ChaosClothAsset/CollectionClothSeamFacade.cpp
--- a/Engine/Plugins/ChaosClothAsset/Source/ChaosClothAsset/Private/ChaosClothAsset/CollectionClothSeamFacade.cpp
+++ b/Engine/Plugins/ChaosClothAsset/Source/ChaosClothAsset/Private/ChaosClothAsset/CollectionClothSeamFacade.cpp
@@ -741,6 +741,24 @@ static FAutoConsoleVariableRef CVarEnableSeamChecks(TEXT("p.ChaosCloth.EnableSea
 				// Split Tethers
 				SplitTethers(SplittingGroups, Cloth.GetTetherKinematicIndex(), Cloth.GetTetherReferenceLength());
 
+				// Split accessory mesh data
+				for (int32 AccessoryMeshIndex = 0; AccessoryMeshIndex < Cloth.GetNumSimAccessoryMeshes(); ++AccessoryMeshIndex)
+				{
+					FCollectionClothSimAccessoryMeshFacade AccessoryMesh = Cloth.GetSimAccessoryMesh(AccessoryMeshIndex);
+
+					// Split 3D positions
+					SplitCopyVertexData(SplittingGroups, AccessoryMesh.GetSimAccessoryMeshPosition3D());
+
+					// Split normals
+					SplitCopyVertexData(SplittingGroups, AccessoryMesh.GetSimAccessoryMeshNormal());
+
+					// Split BoneIndices
+					SplitCopyVertexData(SplittingGroups, AccessoryMesh.GetSimAccessoryMeshBoneIndices());
+
+					// Split BoneWeights
+					SplitCopyVertexData(SplittingGroups, AccessoryMesh.GetSimAccessoryMeshBoneWeights());
+				}
+
 				// Split Faces
 				// The easiest thing to do here is actually to regenerate the 3D faces from the 2D faces.
 				UpdateSimIndicesAfterSplitting(Cloth.GetSimIndices2D(), Cloth.GetSimVertex3DLookup(), Cloth.GetSimIndices3D());
@@ -845,6 +863,22 @@ static FAutoConsoleVariableRef CVarEnableSeamChecks(TEXT("p.ChaosCloth.EnableSea
 		// Weld Tethers
 		WeldTethers(WeldingMap, WeldingGroups, Cloth.GetTetherKinematicIndex(), Cloth.GetTetherReferenceLength());
 
+		// Weld accessory mesh data
+		for (int32 AccessoryMeshIndex = 0; AccessoryMeshIndex < Cloth.GetNumSimAccessoryMeshes(); ++AccessoryMeshIndex)
+		{
+			FCollectionClothSimAccessoryMeshFacade AccessoryMesh = Cloth.GetSimAccessoryMesh(AccessoryMeshIndex);
+
+			// Weld 3D positions
+			WeldByWeightedAverage(WeldingGroups, AccessoryMesh.GetSimAccessoryMeshPosition3D());
+
+			// Weld normals
+			WeldNormals(WeldingGroups, AccessoryMesh.GetSimAccessoryMeshNormal());
+
+			// Weld BoneIndices and Weights
+			WeldIndexAndFloatArrays<true, FClothCollection::MaxNumBoneInfluences>(WeldingGroups, AccessoryMesh.GetSimAccessoryMeshBoneIndices(), AccessoryMesh.GetSimAccessoryMeshBoneWeights(),
+				[](const TPair<float, int32>& A, const TPair<float, int32>& B) { return A > B; });
+		}
+
 		// Weld Faces
 		// Just go through all faces and fix up. We could store vertex -> face lookups, but we'd have to ensure they stay in sync
 		for (FIntVector3& Index3D : Cloth.GetSimIndices3D())
@@ -872,8 +906,6 @@ static FAutoConsoleVariableRef CVarEnableSeamChecks(TEXT("p.ChaosCloth.EnableSea
 			WeldByWeightedAverage(WeldingGroups, Cloth.GetWeightMap(WeightMapName));
 		}
 
-
-
 		// Gather list of vertices to remove
 		TArray<int32> VerticesToRemove;
 		for (TMap<int32, int32>::TConstIterator WeldingMapIter = WeldingMap.CreateConstIterator(); WeldingMapIter; ++WeldingMapIter)
diff --git a/Engine/Plugins/ChaosClothAsset/Source/ChaosClothAsset/Private/ChaosClothAsset/CollectionClothSimAccessoryMeshFacade.cpp b/Engine/Plugins/ChaosClothAsset/Source/ChaosClothAsset/Private/ChaosClothAsset/CollectionClothSimAccessoryMeshFacade.cpp
new file mode 100644
--- /dev/null
+++ b/Engine/Plugins/ChaosClothAsset/Source/ChaosClothAsset/Private/ChaosClothAsset/CollectionClothSimAccessoryMeshFacade.cpp
@@ -0,0 +1,179 @@
+// Copyright Epic Games, Inc. All Rights Reserved.
+
+#include "ChaosClothAsset/CollectionClothSimAccessoryMeshFacade.h"
+#include "ChaosClothAsset/ClothCollection.h"
+#include "ChaosClothAsset/ClothCollectionGroup.h"
+
+namespace UE::Chaos::ClothAsset
+{
+	FName FCollectionClothSimAccessoryMeshConstFacade::GetSimAccessoryMeshName() const
+	{
+		return ClothCollection->GetSimAccessoryMeshName() && ClothCollection->GetNumElements(ClothCollectionGroup::SimAccessoryMeshes) > GetElementIndex() ?
+			(*ClothCollection->GetSimAccessoryMeshName())[GetElementIndex()] : NAME_None;
+	}
+
+	FName FCollectionClothSimAccessoryMeshConstFacade::GetSimAccessoryMeshPosition3DAttribute() const
+	{
+		return ClothCollection->GetSimAccessoryMeshPosition3DAttribute() && ClothCollection->GetNumElements(ClothCollectionGroup::SimAccessoryMeshes) > GetElementIndex() ?
+			(*ClothCollection->GetSimAccessoryMeshPosition3DAttribute())[GetElementIndex()] : NAME_None;
+	}
+
+	FName FCollectionClothSimAccessoryMeshConstFacade::GetSimAccessoryMeshNormalAttribute() const
+	{
+		return ClothCollection->GetSimAccessoryMeshNormalAttribute() && ClothCollection->GetNumElements(ClothCollectionGroup::SimAccessoryMeshes) > GetElementIndex() ?
+			(*ClothCollection->GetSimAccessoryMeshNormalAttribute())[GetElementIndex()] : NAME_None;
+	}
+
+	FName FCollectionClothSimAccessoryMeshConstFacade::GetSimAccessoryMeshBoneIndicesAttribute() const
+	{
+		return ClothCollection->GetSimAccessoryMeshBoneIndicesAttribute() && ClothCollection->GetNumElements(ClothCollectionGroup::SimAccessoryMeshes) > GetElementIndex() ?
+			(*ClothCollection->GetSimAccessoryMeshBoneIndicesAttribute())[GetElementIndex()] : NAME_None;
+	}
+
+	FName FCollectionClothSimAccessoryMeshConstFacade::GetSimAccessoryMeshBoneWeightsAttribute() const
+	{
+		return ClothCollection->GetSimAccessoryMeshBoneWeightsAttribute() && ClothCollection->GetNumElements(ClothCollectionGroup::SimAccessoryMeshes) > GetElementIndex() ?
+			(*ClothCollection->GetSimAccessoryMeshBoneWeightsAttribute())[GetElementIndex()] : NAME_None;
+	}
+
+	TConstArrayView<FVector3f> FCollectionClothSimAccessoryMeshConstFacade::GetSimAccessoryMeshPosition3D() const
+	{
+		return ClothCollection->GetElements(ClothCollection->GetSimAccessoryMeshPosition3D(GetSimAccessoryMeshPosition3DAttribute()));
+	}
+
+	TConstArrayView<FVector3f> FCollectionClothSimAccessoryMeshConstFacade::GetSimAccessoryMeshNormal() const
+	{
+		return ClothCollection->GetElements(ClothCollection->GetSimAccessoryMeshNormal(GetSimAccessoryMeshNormalAttribute()));
+	}
+
+	TConstArrayView<TArray<int32>> FCollectionClothSimAccessoryMeshConstFacade::GetSimAccessoryMeshBoneIndices() const
+	{
+		return ClothCollection->GetElements(ClothCollection->GetSimAccessoryMeshBoneIndices(GetSimAccessoryMeshBoneIndicesAttribute()));
+	}
+
+	TConstArrayView<TArray<float>> FCollectionClothSimAccessoryMeshConstFacade::GetSimAccessoryMeshBoneWeights() const
+	{
+		return ClothCollection->GetElements(ClothCollection->GetSimAccessoryMeshBoneWeights(GetSimAccessoryMeshBoneWeightsAttribute()));
+	}
+
+	FCollectionClothSimAccessoryMeshConstFacade::FCollectionClothSimAccessoryMeshConstFacade(const TSharedRef<const FConstClothCollection>& ClothCollection, int32 InMeshIndex)
+		: ClothCollection(ClothCollection)
+		, MeshIndex(InMeshIndex)
+	{
+		check(ClothCollection->IsValid(EClothCollectionExtendedSchemas::CookedOnly) || ClothCollection->IsValid(EClothCollectionExtendedSchemas::SimAccessoryMeshes));
+		check(MeshIndex >= 0 && MeshIndex < ClothCollection->GetNumElements(ClothCollectionGroup::SimAccessoryMeshes));
+	}
+
+	void FCollectionClothSimAccessoryMeshFacade::Reset()
+	{
+		SetSimAccessoryMeshName(NAME_None);
+		ClearDynamicAttributes();
+	}
+
+	void FCollectionClothSimAccessoryMeshFacade::Initialize(const FCollectionClothSimAccessoryMeshConstFacade& Other)
+	{
+		GenerateDynamicAttributesIfNecessary();
+
+		//~ Sim Accessory Meshes Group
+		SetSimAccessoryMeshName(Other.GetSimAccessoryMeshName());
+
+		//~ Sim Vertices 3D Group
+		FClothCollection::CopyArrayViewData(GetSimAccessoryMeshPosition3D(), Other.GetSimAccessoryMeshPosition3D());
+		FClothCollection::CopyArrayViewData(GetSimAccessoryMeshNormal(), Other.GetSimAccessoryMeshNormal());
+		FClothCollection::CopyArrayViewData(GetSimAccessoryMeshBoneIndices(), Other.GetSimAccessoryMeshBoneIndices());
+		FClothCollection::CopyArrayViewData(GetSimAccessoryMeshBoneWeights(), Other.GetSimAccessoryMeshBoneWeights());
+	}
+
+	void FCollectionClothSimAccessoryMeshFacade::Initialize(const FName& Name, const TConstArrayView<FVector3f>& Positions, const TConstArrayView<FVector3f>& Normals, const TConstArrayView<TArray<int32>>& BoneIndices, const TConstArrayView<TArray<float>>& BoneWeights)
+	{
+		const int32 NumExpectedVertices = ClothCollection->GetNumElements(ClothCollectionGroup::SimVertices3D);
+		check(Positions.Num() == NumExpectedVertices);
+		check(Normals.Num() == NumExpectedVertices);
+		check(BoneIndices.Num() == NumExpectedVertices);
+		check(BoneWeights.Num() == NumExpectedVertices);
+
+		GenerateDynamicAttributesIfNecessary();
+
+		//~ Sim Accessory Meshes Group
+		SetSimAccessoryMeshName(Name);
+
+		//~ Sim Vertices 3D Group
+		FClothCollection::CopyArrayViewData(GetSimAccessoryMeshPosition3D(), Positions);
+		FClothCollection::CopyArrayViewData(GetSimAccessoryMeshNormal(), Normals);
+		FClothCollection::CopyArrayViewData(GetSimAccessoryMeshBoneIndices(), BoneIndices);
+		FClothCollection::CopyArrayViewData(GetSimAccessoryMeshBoneWeights(), BoneWeights);
+	}
+
+	void FCollectionClothSimAccessoryMeshFacade::SetSimAccessoryMeshName(const FName& MeshName)
+	{
+		(*GetClothCollection()->GetSimAccessoryMeshName())[GetElementIndex()] = MeshName;
+	}
+
+	TArrayView<FVector3f> FCollectionClothSimAccessoryMeshFacade::GetSimAccessoryMeshPosition3D()
+	{
+		return GetClothCollection()->GetElements(
+			GetClothCollection()->GetSimAccessoryMeshPosition3D(GetSimAccessoryMeshPosition3DAttribute()));
+	}
+
+	TArrayView<FVector3f> FCollectionClothSimAccessoryMeshFacade::GetSimAccessoryMeshNormal()
+	{
+		return GetClothCollection()->GetElements(
+			GetClothCollection()->GetSimAccessoryMeshNormal(GetSimAccessoryMeshNormalAttribute()));
+	}
+
+	TArrayView<TArray<int32>> FCollectionClothSimAccessoryMeshFacade::GetSimAccessoryMeshBoneIndices()
+	{
+		return GetClothCollection()->GetElements(
+			GetClothCollection()->GetSimAccessoryMeshBoneIndices(GetSimAccessoryMeshBoneIndicesAttribute()));
+	}
+
+	TArrayView<TArray<float>> FCollectionClothSimAccessoryMeshFacade::GetSimAccessoryMeshBoneWeights()
+	{
+		return GetClothCollection()->GetElements(
+			GetClothCollection()->GetSimAccessoryMeshBoneWeights(GetSimAccessoryMeshBoneWeightsAttribute()));
+	}
+
+	FCollectionClothSimAccessoryMeshFacade::FCollectionClothSimAccessoryMeshFacade(const TSharedRef<FClothCollection>& ClothCollection, int32 InMeshIndex)
+		: FCollectionClothSimAccessoryMeshConstFacade(ClothCollection, InMeshIndex)
+	{
+	}
+
+	void FCollectionClothSimAccessoryMeshFacade::ClearDynamicAttributes()
+	{
+		//~ Sim Vertices 3D Group
+		GetClothCollection()->RemoveSimAccessoryMeshPosition3D(GetSimAccessoryMeshPosition3DAttribute());
+		GetClothCollection()->RemoveSimAccessoryMeshNormal(GetSimAccessoryMeshNormalAttribute());
+		GetClothCollection()->RemoveSimAccessoryMeshBoneIndices(GetSimAccessoryMeshBoneIndicesAttribute());
+		GetClothCollection()->RemoveSimAccessoryMeshBoneWeights(GetSimAccessoryMeshBoneWeightsAttribute());
+		(*GetClothCollection()->GetSimAccessoryMeshPosition3DAttribute())[GetElementIndex()] = NAME_None;
+		(*GetClothCollection()->GetSimAccessoryMeshNormalAttribute())[GetElementIndex()] = NAME_None;
+		(*GetClothCollection()->GetSimAccessoryMeshBoneIndicesAttribute())[GetElementIndex()] = NAME_None;
+		(*GetClothCollection()->GetSimAccessoryMeshBoneWeightsAttribute())[GetElementIndex()] = NAME_None;
+	}
+
+	void FCollectionClothSimAccessoryMeshFacade::GenerateDynamicAttributesIfNecessary()
+	{
+		//~ Sim Vertices 3D Group
+		if (ClothCollection->GetSimAccessoryMeshPosition3D(GetSimAccessoryMeshPosition3DAttribute()) == nullptr)
+		{
+			(*GetClothCollection()->GetSimAccessoryMeshPosition3DAttribute())[GetElementIndex()] = GetClothCollection()->AddSimAccessoryMeshPosition3D();
+		}
+		if (ClothCollection->GetSimAccessoryMeshNormal(GetSimAccessoryMeshNormalAttribute()) == nullptr)
+		{
+			(*GetClothCollection()->GetSimAccessoryMeshNormalAttribute())[GetElementIndex()] = GetClothCollection()->AddSimAccessoryMeshNormal();
+		}
+		if (ClothCollection->GetSimAccessoryMeshBoneIndices(GetSimAccessoryMeshBoneIndicesAttribute()) == nullptr)
+		{
+			(*GetClothCollection()->GetSimAccessoryMeshBoneIndicesAttribute())[GetElementIndex()] = GetClothCollection()->AddSimAccessoryMeshBoneIndices();
+		}
+		if (ClothCollection->GetSimAccessoryMeshBoneWeights(GetSimAccessoryMeshBoneWeightsAttribute()) == nullptr)
+		{
+			(*GetClothCollection()->GetSimAccessoryMeshBoneWeightsAttribute())[GetElementIndex()] = GetClothCollection()->AddSimAccessoryMeshBoneWeights();
+		}
+	}
+
+	TSharedRef<class FClothCollection> FCollectionClothSimAccessoryMeshFacade::GetClothCollection()
+	{
+		return StaticCastSharedRef<class FClothCollection>(ConstCastSharedRef<class FConstClothCollection>(ClothCollection));
+	}
+}  // End namespace UE::Chaos::ClothAsset
diff --git a/Engine/Plugins/ChaosClothAsset/Source/ChaosClothAsset/Private/ChaosClothAsset/CollectionClothSimPatternFacade.cpp b/Engine/Plugins/ChaosClothAsset/Source/ChaosClothAsset/Private/ChaosClothAsset/CollectionClothSimPatternFacade.cpp
--- a/Engine/Plugins/ChaosClothAsset/Source/ChaosClothAsset/Private/ChaosClothAsset/CollectionClothSimPatternFacade.cpp
+++ b/Engine/Plugins/ChaosClothAsset/Source/ChaosClothAsset/Private/ChaosClothAsset/CollectionClothSimPatternFacade.cpp
@@ -53,6 +53,15 @@ namespace UE::Chaos::ClothAsset
 			GetElementIndex());
 	}
 
+	TConstArrayView<int32> FCollectionClothSimPatternConstFacade::GetSimImportVertexID() const
+	{
+		return ClothCollection->GetElements(
+			ClothCollection->GetSimImportVertexID(),
+			ClothCollection->GetSimVertices2DStart(),
+			ClothCollection->GetSimVertices2DEnd(),
+			GetElementIndex());
+	}
+
 	int32 FCollectionClothSimPatternConstFacade::GetNumSimFaces() const
 	{
 		return ClothCollection->GetNumElements(
@@ -209,6 +218,19 @@ namespace UE::Chaos::ClothAsset
 		const int32 SimVertex2DOffset = GetSimVertices2DOffset() - Other.GetSimVertices2DOffset();
 		FClothCollection::CopyArrayViewDataAndApplyOffset(GetSimIndices2D(), Other.GetSimIndices2D(), FIntVector(SimVertex2DOffset));
 		FClothCollection::CopyArrayViewDataAndApplyOffset(GetSimIndices3D(), Other.GetSimIndices3D(), FIntVector(SimVertex3DOffset));
+		if (Other.ClothCollection->IsValid(EClothCollectionExtendedSchemas::Import))
+		{
+			FClothCollection::CopyArrayViewData(GetSimImportVertexID(), Other.GetSimImportVertexID());
+		}
+		else if (ClothCollection->IsValid(EClothCollectionExtendedSchemas::Import))
+		{
+			// Initialize to INDEX_NONE
+			TArrayView<int32> SimImportVertexID = GetSimImportVertexID();
+			for (int32 Index = 0; Index < SimImportVertexID.Num(); ++Index)
+			{
+				SimImportVertexID[Index] = INDEX_NONE;
+			}
+		}
 	}
 
 	TArrayView<FVector2f> FCollectionClothSimPatternFacade::GetSimPosition2D()
@@ -220,6 +242,20 @@ namespace UE::Chaos::ClothAsset
 			GetElementIndex());
 	}
 
+	TArrayView<int32> FCollectionClothSimPatternFacade::GetSimImportVertexID()
+	{
+		if (!GetClothCollection()->IsValid(EClothCollectionExtendedSchemas::Import))
+		{
+			GetClothCollection()->DefineSchema(EClothCollectionExtendedSchemas::Import);
+		}
+
+		return GetClothCollection()->GetElements(
+			GetClothCollection()->GetSimImportVertexID(),
+			GetClothCollection()->GetSimVertices2DStart(),
+			GetClothCollection()->GetSimVertices2DEnd(),
+			GetElementIndex());
+	}
+
 	FIntVector2 FCollectionClothSimPatternFacade::AppendSimVertices(int32 NumSimVertices)
 	{
 		const FIntVector2 OrigNumVertices(GetNumSimVertices2D(), ClothCollection->GetNumElements(ClothCollectionGroup::SimVertices3D));
@@ -234,6 +270,8 @@ namespace UE::Chaos::ClothAsset
 		TArrayView<TArray<int32>> SimVertex2DLookup = GetClothCollection()->GetElements(
 			GetClothCollection()->GetSimVertex2DLookup());
 
+		const bool bSetImportVertexID = ClothCollection->IsValid(EClothCollectionExtendedSchemas::Import);
+		TArrayView<int32> SimImportVertexID = bSetImportVertexID ? GetSimImportVertexID() : TArrayView<int32>();
 		for (int32 NewVertexIndex = 0; NewVertexIndex < NumSimVertices; ++NewVertexIndex)
 		{
 			const int32 LocalVertex2D = OrigNumVertices[0] + NewVertexIndex;
@@ -241,6 +279,10 @@ namespace UE::Chaos::ClothAsset
 			const int32 Vertex3D = OrigNumVertices[1] + NewVertexIndex;
 			SimVertex3DLookup[LocalVertex2D] = Vertex3D;
 			SimVertex2DLookup[Vertex3D].Init(GlobalVertex2D, 1);
+			if (bSetImportVertexID)
+			{
+				SimImportVertexID[LocalVertex2D] = INDEX_NONE;
+			}
 		}
 
 		return OrigNumVertices;
diff --git a/Engine/Plugins/ChaosClothAsset/Source/ChaosClothAsset/Public/ChaosClothAsset/ClothCollectionAttribute.h b/Engine/Plugins/ChaosClothAsset/Source/ChaosClothAsset/Public/ChaosClothAsset/ClothCollectionAttribute.h
--- a/Engine/Plugins/ChaosClothAsset/Source/ChaosClothAsset/Public/ChaosClothAsset/ClothCollectionAttribute.h
+++ b/Engine/Plugins/ChaosClothAsset/Source/ChaosClothAsset/Public/ChaosClothAsset/ClothCollectionAttribute.h
@@ -64,6 +64,7 @@ namespace UE::Chaos::ClothAsset
 		// Sim Vertices 2D Group
 		inline const FName SimPosition2D(TEXT("SimPosition2D"));
 		inline const FName SimVertex3DLookup(TEXT("SimVertex3DLookup"));
+		inline const FName SimImportVertexID(TEXT("SimImportVertexID"));
 
 		// Sim Vertices 3D Group
 		inline const FName SimPosition3D(TEXT("SimPosition3D"));
@@ -77,6 +78,12 @@ namespace UE::Chaos::ClothAsset
 		inline const FName SeamStitchLookup(TEXT("SeamStitchLookup"));
 		inline const FName SimCustomResizingBlend(TEXT("SimCustomResizingBlend"));
 
+		// Sim Vertices 3D Group -- Sim Accessory Mesh
+		inline const FName SimAccessoryMeshPosition3DPrefix(TEXT("SimAccessoryMeshPosition3D"));
+		inline const FName SimAccessoryMeshNormalPrefix(TEXT("SimAccessoryMeshNormal"));
+		inline const FName SimAccessoryMeshBoneIndicesPrefix(TEXT("SimAccessoryMeshBoneIndices"));
+		inline const FName SimAccessoryMeshBoneWeightsPrefix(TEXT("SimAccessoryMeshBoneWeights"));
+
 		// Sim Morph Targets Group
 		inline const FName SimMorphTargetName(TEXT("SimMorphTargetName"));
 		inline const FName SimMorphTargetVerticesStart(TEXT("SimMorphTargetVerticesStart"));
@@ -110,6 +117,12 @@ namespace UE::Chaos::ClothAsset
 		// Custom Resizing Regions Group
 		inline const FName CustomResizingRegionSet(TEXT("CustomResizingRegionSet"));
 		inline const FName CustomResizingRegionType(TEXT("CustomResizingRegionType"));
+		// Sim Accessory Meshes Group
+		inline const FName SimAccessoryMeshName(TEXT("SimAccessoryMeshName"));
+		inline const FName SimAccessoryMeshPosition3DAttribute(TEXT("SimAccessoryMeshPosition3DAttribute"));
+		inline const FName SimAccessoryMeshNormalAttribute(TEXT("SimAccessoryMeshNormalAttribute"));
+		inline const FName SimAccessoryMeshBoneIndicesAttribute(TEXT("SimAccessoryMeshBoneIndicesAttribute"));
+		inline const FName SimAccessoryMeshBoneWeightsAttribute(TEXT("SimAccessoryMeshBoneWeightsAttribute"));
 
 		//~ Deprecated attributes
 		// Lods Group
diff --git a/Engine/Plugins/ChaosClothAsset/Source/ChaosClothAsset/Public/ChaosClothAsset/ClothCollectionExtendedSchemas.h b/Engine/Plugins/ChaosClothAsset/Source/ChaosClothAsset/Public/ChaosClothAsset/ClothCollectionExtendedSchemas.h
--- a/Engine/Plugins/ChaosClothAsset/Source/ChaosClothAsset/Public/ChaosClothAsset/ClothCollectionExtendedSchemas.h
+++ b/Engine/Plugins/ChaosClothAsset/Source/ChaosClothAsset/Public/ChaosClothAsset/ClothCollectionExtendedSchemas.h
@@ -8,7 +8,7 @@
 namespace UE::Chaos::ClothAsset
 {
 	/** Cloth collection extended schemas. */
-	enum class EClothCollectionExtendedSchemas : uint8
+	enum class EClothCollectionExtendedSchemas : uint32
 	{
 		None = 0,
 		RenderDeformer = 1 << 1,
@@ -17,6 +17,8 @@ namespace UE::Chaos::ClothAsset
 		SimMorphTargets = 1 << 4,
 		CookedOnly = 1 << 5, // If this flag is found, only cooked schemas will be considered.
 		Resizing = 1 << 6,
+		SimAccessoryMeshes = 1 << 7,
+		Import = 1 << 8
 	};
 	ENUM_CLASS_FLAGS(EClothCollectionExtendedSchemas)
 }  // End namespace UE::Chaos::ClothAsset
diff --git a/Engine/Plugins/ChaosClothAsset/Source/ChaosClothAsset/Public/ChaosClothAsset/ClothCollectionGroup.h b/Engine/Plugins/ChaosClothAsset/Source/ChaosClothAsset/Public/ChaosClothAsset/ClothCollectionGroup.h
--- a/Engine/Plugins/ChaosClothAsset/Source/ChaosClothAsset/Public/ChaosClothAsset/ClothCollectionGroup.h
+++ b/Engine/Plugins/ChaosClothAsset/Source/ChaosClothAsset/Public/ChaosClothAsset/ClothCollectionGroup.h
@@ -41,5 +41,7 @@ namespace UE::Chaos::ClothAsset
 		inline const FName SimMorphTargetVertices = FName(TEXT("SimMorphTargetVertices"));
 		/** Contains custom region resizing data */
 		inline const FName CustomResizingRegions = FName(TEXT("CustomResizingRegions"));
+		/** Contains sim accessory mesh relationships to other groups. */
+		inline const FName SimAccessoryMeshes = FName(TEXT("SimAccessoryMeshes"));
 	}
 }  // End namespace UE::Chaos::ClothAsset
diff --git a/Engine/Plugins/ChaosClothAsset/Source/ChaosClothAsset/Public/ChaosClothAsset/ClothGeometryTools.h b/Engine/Plugins/ChaosClothAsset/Source/ChaosClothAsset/Public/ChaosClothAsset/ClothGeometryTools.h
--- a/Engine/Plugins/ChaosClothAsset/Source/ChaosClothAsset/Public/ChaosClothAsset/ClothGeometryTools.h
+++ b/Engine/Plugins/ChaosClothAsset/Source/ChaosClothAsset/Public/ChaosClothAsset/ClothGeometryTools.h
@@ -13,6 +13,7 @@ struct FMeshDescription;
 namespace UE::Chaos::ClothAsset
 {
 	class FCollectionClothFacade;
+	class FCollectionClothConstFacade;
 
 	/**
 	 * Geometry tools operating on cloth collections.
@@ -179,6 +180,12 @@ namespace UE::Chaos::ClothAsset
 		 */
 		static UE_API void ApplyProxyDeformer(const TSharedRef<FManagedArrayCollection>& ClothCollection, bool bIgnoreSkinningBlend);
 		static UE_API void ApplyProxyDeformer(FCollectionClothFacade& ClothFacade, bool bIgnoreSkinningBlend, const TConstArrayView<FVector3f>& SimPositions, const TConstArrayView<FVector3f>& SimNormals);
+
+		/** Create a new SimAccessoryMesh by copying the SimMesh data from one collection into another.
+		 *  Optionally use SimImportVertexID to match data rather than a straight copy of SimVertices3D.
+		 *  Any unmatched vertices will copy the data from the existing Sim Mesh.
+		 *  Returns the newly created accessory mesh's index */
+		static UE_API int32 CopySimMeshToSimAccessoryMesh(const FName& AccessoryMeshName, FCollectionClothFacade& ToCloth, const FCollectionClothConstFacade& FromCloth, bool bUseSimImportVertexID);
 	};
 }  // End namespace UE::Chaos::ClothAsset
 
diff --git a/Engine/Plugins/ChaosClothAsset/Source/ChaosClothAsset/Public/ChaosClothAsset/CollectionClothFacade.h b/Engine/Plugins/ChaosClothAsset/Source/ChaosClothAsset/Public/ChaosClothAsset/CollectionClothFacade.h
--- a/Engine/Plugins/ChaosClothAsset/Source/ChaosClothAsset/Public/ChaosClothAsset/CollectionClothFacade.h
+++ b/Engine/Plugins/ChaosClothAsset/Source/ChaosClothAsset/Public/ChaosClothAsset/CollectionClothFacade.h
@@ -10,6 +10,7 @@
 #include "ChaosClothAsset/CollectionClothFabricFacade.h"
 #include "ChaosClothAsset/CollectionClothRenderPatternFacade.h"
 #include "ChaosClothAsset/CollectionClothSeamFacade.h"
+#include "ChaosClothAsset/CollectionClothSimAccessoryMeshFacade.h"
 #include "ChaosClothAsset/CollectionClothSimMorphTargetFacade.h"
 #include "ChaosClothAsset/CollectionClothSimPatternFacade.h"
 #include "ChaosClothAsset/IsUserAttributeType.h"
@@ -104,6 +105,7 @@ namespace UE::Chaos::ClothAsset
 		UE_API int32 GetNumSimVertices2D() const;
 		UE_API TConstArrayView<FVector2f> GetSimPosition2D() const;
 		UE_API TConstArrayView<int32> GetSimVertex3DLookup() const;
+		UE_API TConstArrayView<int32> GetSimImportVertexID() const;
 
 		//~ Sim Vertices 3D Group
 		/** Return the total number of 3D simulation vertices for this collection. */
@@ -145,6 +147,19 @@ namespace UE::Chaos::ClothAsset
 		/** Return a view of all sim morph target names */
 		UE_API TConstArrayView<FString> GetSimMorphTargetName() const;
 
+		//~ Sim Accessory Meshes Group
+		/** Return the number of accessory meshes in this collection. */
+		UE_API int32 GetNumSimAccessoryMeshes() const;
+		/** Lookup sim accessory mesh by name. Returns INDEX_NONE if not found. */
+		UE_API int32 FindSimAccessoryMeshIndexByName(const FName& AccessoryMeshName) const;
+		/** Return the sim accessory mesh facade for the specified accessor mesh index */
+		UE_API FCollectionClothSimAccessoryMeshConstFacade GetSimAccessoryMesh(int32 MeshIndex) const;
+		UE_API TConstArrayView<FName> GetSimAccessoryMeshName() const;
+		UE_API TConstArrayView<FName> GetSimAccessoryMeshPosition3DAttribute() const;
+		UE_API TConstArrayView<FName> GetSimAccessoryMeshNormalAttribute() const;
+		UE_API TConstArrayView<FName> GetSimAccessoryMeshBoneIndicesAttribute() const;
+		UE_API TConstArrayView<FName> GetSimAccessoryMeshBoneWeightsAttribute() const;
+
 		//~ Sim Morph Target Vertices Group
 		/** Return the total number of sim morph target vertices in this collection. */
 		UE_API int32 GetNumSimMorphTargetVertices() const;
@@ -308,6 +323,7 @@ namespace UE::Chaos::ClothAsset
 		//~ Pattern Sim Vertices 2D Group
 		/** SetNumSimVertices2D per pattern within pattern facade. */
 		UE_API TArrayView<FVector2f> GetSimPosition2D();
+		UE_API TArrayView<int32> GetSimImportVertexID();
 
 		//~ Pattern Sim Vertices 3D Group
 		UE_API TArrayView<FVector3f> GetSimPosition3D();
@@ -369,6 +385,23 @@ namespace UE::Chaos::ClothAsset
 		/** Remove a sorted list of sim morph target vertices. */
 		UE_API void RemoveSimMorphTargetVertices3D(const TArray<int32>& SortedDeletionList);
 
+		//~ Sim Accessory Meshes Group
+		/** Set the new number of accessory meshes to this cloth. */
+		UE_API void SetNumSimAccessoryMeshes(int32 NumMeshes, bool bCookedOnly = false);
+		/** Add a new accessory mesh to this cloth and return its index in the accessory mesh list. */
+		UE_API int32 AddSimAccessoryMesh();
+		/** Return an accessory mesh facade for the specified accessory mesh index. */
+		UE_API FCollectionClothSimAccessoryMeshFacade GetSimAccessoryMesh(int32 MeshIndex);
+		/** Add a new accessory mesh to this cloth and return the accessory mesh facade set to its index. */
+		FCollectionClothSimAccessoryMeshFacade AddGetSimAccessoryMesh() { return GetSimAccessoryMesh(AddSimAccessoryMesh()); }
+		/** Remove a sorted list of accessory meshes. */
+		UE_API void RemoveSimAccessoryMeshes(const TArray<int32>& SortedDeletionList);
+		UE_API TArrayView<FName> GetSimAccessoryMeshName();
+		UE_API TArrayView<FName> GetSimAccessoryMeshPosition3DAttribute();
+		UE_API TArrayView<FName> GetSimAccessoryMeshNormalAttribute();
+		UE_API TArrayView<FName> GetSimAccessoryMeshBoneIndicesAttribute();
+		UE_API TArrayView<FName> GetSimAccessoryMeshBoneWeightsAttribute();
+
 		//~ Render Patterns Group
 		/** Set the new number of patterns to this cloth LOD. */
 		UE_API void SetNumRenderPatterns(int32 NumPatterns);
diff --git a/Engine/Plugins/ChaosClothAsset/Source/ChaosClothAsset/Public/ChaosClothAsset/CollectionClothSimAccessoryMeshFacade.h b/Engine/Plugins/ChaosClothAsset/Source/ChaosClothAsset/Public/ChaosClothAsset/CollectionClothSimAccessoryMeshFacade.h
new file mode 100644
--- /dev/null
+++ b/Engine/Plugins/ChaosClothAsset/Source/ChaosClothAsset/Public/ChaosClothAsset/CollectionClothSimAccessoryMeshFacade.h
@@ -0,0 +1,112 @@
+// Copyright Epic Games, Inc. All Rights Reserved.
+
+#pragma once
+
+#include "Templates/SharedPointer.h"
+
+#define UE_API CHAOSCLOTHASSET_API
+
+struct FManagedArrayCollection;
+
+namespace UE::Chaos::ClothAsset
+{
+	/**
+	 * Cloth Asset collection sim accessory mesh facade class to access cloth sim accessory mesh data.
+	 * Constructed from FCollectionClothConstFacade.
+	 * Const access (read only) version.
+	 */
+	class FCollectionClothSimAccessoryMeshConstFacade
+	{
+	public:
+		FCollectionClothSimAccessoryMeshConstFacade() = delete;
+
+		FCollectionClothSimAccessoryMeshConstFacade(const FCollectionClothSimAccessoryMeshConstFacade&) = delete;
+		FCollectionClothSimAccessoryMeshConstFacade& operator=(const FCollectionClothSimAccessoryMeshConstFacade&) = delete;
+
+		FCollectionClothSimAccessoryMeshConstFacade(FCollectionClothSimAccessoryMeshConstFacade&&) = default;
+		FCollectionClothSimAccessoryMeshConstFacade& operator=(FCollectionClothSimAccessoryMeshConstFacade&&) = default;
+
+		virtual ~FCollectionClothSimAccessoryMeshConstFacade() = default;
+
+		//~ Sim Accessory Meshes Group
+		/** Return the name for this accessory mesh */
+		UE_API FName GetSimAccessoryMeshName() const;
+		/** Return the names of dynamic attributes associated with this accessory mesh */
+		UE_API FName GetSimAccessoryMeshPosition3DAttribute() const;
+		UE_API FName GetSimAccessoryMeshNormalAttribute() const;
+		UE_API FName GetSimAccessoryMeshBoneIndicesAttribute() const;
+		UE_API FName GetSimAccessoryMeshBoneWeightsAttribute() const;
+
+		//~ Sim Vertices 3D Group
+		UE_API TConstArrayView<FVector3f> GetSimAccessoryMeshPosition3D() const;
+		UE_API TConstArrayView<FVector3f> GetSimAccessoryMeshNormal() const;
+		UE_API TConstArrayView<TArray<int32>> GetSimAccessoryMeshBoneIndices() const;
+		UE_API TConstArrayView<TArray<float>> GetSimAccessoryMeshBoneWeights() const;
+
+		/** Return the accessory mesh index this facade has been created with. */
+		int32 GetSimAccessoryMeshIndex() const { return MeshIndex; }
+
+	protected:
+		friend class FCollectionClothSimAccessoryMeshFacade;  // For other instances access
+		friend class FCollectionClothConstFacade;
+		FCollectionClothSimAccessoryMeshConstFacade(const TSharedRef<const class FConstClothCollection>& InClothCollection, int32 InMeshIndex);
+
+		static constexpr int32 GetBaseElementIndex() { return 0; }
+		int32 GetElementIndex() const { return GetBaseElementIndex() + MeshIndex; }
+
+		TSharedRef<const class FConstClothCollection> ClothCollection;
+		int32 MeshIndex;
+	};
+
+	/**
+	 * Cloth Asset collection sim accessory mesh facade class to access cloth sim accessory mesh data.
+	 * Constructed from FCollectionClothFacade.
+	 * Non-const access (read/write) version.
+	 */
+	class FCollectionClothSimAccessoryMeshFacade final : public FCollectionClothSimAccessoryMeshConstFacade
+	{
+	public:
+		FCollectionClothSimAccessoryMeshFacade() = delete;
+
+		FCollectionClothSimAccessoryMeshFacade(const FCollectionClothSimAccessoryMeshFacade&) = delete;
+		FCollectionClothSimAccessoryMeshFacade& operator=(const FCollectionClothSimAccessoryMeshFacade&) = delete;
+
+		FCollectionClothSimAccessoryMeshFacade(FCollectionClothSimAccessoryMeshFacade&&) = default;
+		FCollectionClothSimAccessoryMeshFacade& operator=(FCollectionClothSimAccessoryMeshFacade&&) = default;
+
+		virtual ~FCollectionClothSimAccessoryMeshFacade() override = default;
+
+		/** Remove all attributes from this accessory mesh. */
+		UE_API void Reset();
+
+		/** Set any default values for an accessory mesh. */
+		void SetDefaults() {}
+
+		/** Initialize from another accessory mesh. Assumes all indices match between source and target. Does not try to preserve attribute names, just values.*/
+		UE_API void Initialize(const FCollectionClothSimAccessoryMeshConstFacade& Other);
+
+		/** Initialize */
+		UE_API void Initialize(const FName& Name, const TConstArrayView<FVector3f>& Positions, const TConstArrayView<FVector3f>& Normals, const TConstArrayView<TArray<int32>>& BoneIndices, const TConstArrayView<TArray<float>>& BoneWeights);
+
+		//~ Sim Accessory Meshes Group
+		/** Set the name for this accessory mesh. */
+		UE_API void SetSimAccessoryMeshName(const FName& MeshName);
+
+		//~ Sim Vertices 3D Group
+		UE_API TArrayView<FVector3f> GetSimAccessoryMeshPosition3D();
+		UE_API TArrayView<FVector3f> GetSimAccessoryMeshNormal();
+		UE_API TArrayView<TArray<int32>> GetSimAccessoryMeshBoneIndices();
+		UE_API TArrayView<TArray<float>> GetSimAccessoryMeshBoneWeights();
+
+	private:
+		friend class FCollectionClothFacade;
+		FCollectionClothSimAccessoryMeshFacade(const TSharedRef<class FClothCollection>& InClothCollection, int32 InMorphTargetIndex);
+
+		void ClearDynamicAttributes();
+		void GenerateDynamicAttributesIfNecessary();
+
+		TSharedRef<class FClothCollection> GetClothCollection();
+	};
+}  // End namespace UE::Chaos::ClothAsset
+
+#undef UE_API
diff --git a/Engine/Plugins/ChaosClothAsset/Source/ChaosClothAsset/Public/ChaosClothAsset/CollectionClothSimPatternFacade.h b/Engine/Plugins/ChaosClothAsset/Source/ChaosClothAsset/Public/ChaosClothAsset/CollectionClothSimPatternFacade.h
--- a/Engine/Plugins/ChaosClothAsset/Source/ChaosClothAsset/Public/ChaosClothAsset/CollectionClothSimPatternFacade.h
+++ b/Engine/Plugins/ChaosClothAsset/Source/ChaosClothAsset/Public/ChaosClothAsset/CollectionClothSimPatternFacade.h
@@ -40,6 +40,7 @@ namespace UE::Chaos::ClothAsset
 		UE_API int32 GetSimVertices2DOffset() const;
 		UE_API TConstArrayView<FVector2f> GetSimPosition2D() const;
 		UE_API TConstArrayView<int32> GetSimVertex3DLookup() const;
+		UE_API TConstArrayView<int32> GetSimImportVertexID() const;
 
 		//~ Sim Faces Group
 		// Note: SimIndices points to the collection arrays, not the pattern arrays
@@ -101,6 +102,7 @@ namespace UE::Chaos::ClothAsset
 		//~ Sim Vertices 2D Group
 		// Note: Use the FCollectionClothConstFacade accessors instead of these for the array indices to match the SimIndices2D values
 		UE_API TArrayView<FVector2f> GetSimPosition2D();
+		UE_API TArrayView<int32> GetSimImportVertexID();
 
 		/** This will remove the 2D vertices, but the associated seams and 3D vertices will still exist, and point to INDEX_NONE */
 		UE_API void RemoveSimVertices2D(int32 NumSimVertices);
diff --git a/Engine/Plugins/ChaosClothAsset/Source/ChaosClothAssetEngine/Private/ChaosClothAsset/ClothAsset.cpp b/Engine/Plugins/ChaosClothAsset/Source/ChaosClothAssetEngine/Private/ChaosClothAsset/ClothAsset.cpp
--- a/Engine/Plugins/ChaosClothAsset/Source/ChaosClothAssetEngine/Private/ChaosClothAsset/ClothAsset.cpp
+++ b/Engine/Plugins/ChaosClothAsset/Source/ChaosClothAssetEngine/Private/ChaosClothAsset/ClothAsset.cpp
@@ -174,7 +174,7 @@ namespace UE::Chaos::ClothAsset::Private
 					EmbeddedSpringFacade.Append(InEmbeddedSpringFacade, VertexOffset);
 				}
 
-				// Morph targets
+				// Morph targets and accessory meshes
 				const UE::Chaos::ClothAsset::FCollectionClothConstFacade InClothFacade(ClothCollection);
 				if (InClothFacade.IsValid(UE::Chaos::ClothAsset::EClothCollectionExtendedSchemas::CookedOnly))
 				{
diff --git a/Engine/Plugins/ChaosClothAsset/Source/ChaosClothAssetEngine/Private/ChaosClothAsset/ClothSimulationAccessoryMesh.cpp b/Engine/Plugins/ChaosClothAsset/Source/ChaosClothAssetEngine/Private/ChaosClothAsset/ClothSimulationAccessoryMesh.cpp
new file mode 100644
--- /dev/null
+++ b/Engine/Plugins/ChaosClothAsset/Source/ChaosClothAssetEngine/Private/ChaosClothAsset/ClothSimulationAccessoryMesh.cpp
@@ -0,0 +1,51 @@
+// Copyright Epic Games, Inc. All Rights Reserved.
+#include "ChaosClothAsset/ClothSimulationAccessoryMesh.h"
+#include "ChaosClothAsset/CollectionClothFacade.h"
+#include "ClothVertBoneData.h"
+
+namespace UE::Chaos::ClothAsset
+{
+	FClothSimulationAccessoryMesh::FClothSimulationAccessoryMesh(const TSharedRef<FCollectionClothConstFacade>& InClothFacade, const ::Chaos::FClothingSimulationMesh& InMesh, const int32 InAccessoryMeshIndex)
+		: ::Chaos::FClothingSimulationAccessoryMesh(InMesh, InClothFacade->GetSimAccessoryMeshName()[InAccessoryMeshIndex])
+		, AccessoryMeshFacade(InClothFacade->GetSimAccessoryMesh(InAccessoryMeshIndex))
+	{
+		SetBoneData();
+	}
+
+	void FClothSimulationAccessoryMesh::SetBoneData()
+	{
+		TConstArrayView<TArray<int32>> SimBoneIndices = AccessoryMeshFacade.GetSimAccessoryMeshBoneIndices();
+		TConstArrayView<TArray<float>> SimBoneWeights = AccessoryMeshFacade.GetSimAccessoryMeshBoneWeights();
+		check(SimBoneIndices.Num() == SimBoneWeights.Num());
+		BoneData.SetNum(SimBoneIndices.Num());
+		for (int32 VertexIndex = 0; VertexIndex < SimBoneIndices.Num(); ++VertexIndex)
+		{
+			check(SimBoneIndices[VertexIndex].Num() == SimBoneWeights[VertexIndex].Num());
+			check(SimBoneIndices[VertexIndex].Num() <= FClothVertBoneData::MaxTotalInfluences);
+			
+			BoneData[VertexIndex].NumInfluences = SimBoneIndices[VertexIndex].Num();
+			for (int32 BoneIndex = 0; BoneIndex < BoneData[VertexIndex].NumInfluences; ++BoneIndex)
+			{
+				BoneData[VertexIndex].BoneIndices[BoneIndex] = SimBoneIndices[VertexIndex][BoneIndex];
+				BoneData[VertexIndex].BoneWeights[BoneIndex] = SimBoneWeights[VertexIndex][BoneIndex];
+			}
+		}
+	}
+
+	FClothSimulationAccessoryMesh::~FClothSimulationAccessoryMesh() = default;
+
+	int32 FClothSimulationAccessoryMesh::GetNumPoints() const
+	{
+		return AccessoryMeshFacade.GetSimAccessoryMeshPosition3D().Num();
+	}
+
+	TConstArrayView<FVector3f> FClothSimulationAccessoryMesh::GetPositions() const
+	{
+		return AccessoryMeshFacade.GetSimAccessoryMeshPosition3D();
+	}
+
+	TConstArrayView<FVector3f> FClothSimulationAccessoryMesh::GetNormals() const
+	{
+		return AccessoryMeshFacade.GetSimAccessoryMeshNormal();
+	}
+}
\ No newline at end of file
diff --git a/Engine/Plugins/ChaosClothAsset/Source/ChaosClothAssetEngine/Private/ChaosClothAsset/ClothSimulationAccessoryMesh.h b/Engine/Plugins/ChaosClothAsset/Source/ChaosClothAssetEngine/Private/ChaosClothAsset/ClothSimulationAccessoryMesh.h
new file mode 100644
--- /dev/null
+++ b/Engine/Plugins/ChaosClothAsset/Source/ChaosClothAssetEngine/Private/ChaosClothAsset/ClothSimulationAccessoryMesh.h
@@ -0,0 +1,54 @@
+// Copyright Epic Games, Inc. All Rights Reserved.
+#pragma once
+
+#include "ChaosCloth/ChaosClothingSimulationAccessoryMesh.h"
+#include "ChaosClothAsset/CollectionClothSimAccessoryMeshFacade.h"
+
+namespace UE::Chaos::ClothAsset
+{
+	class FCollectionClothConstFacade;
+
+	// Accessory mesh used as an input to a simulation. They are expected to share topology (e.g., Indices, PatternIndices) with their associated FClothingSimulationMesh
+	class FClothSimulationAccessoryMesh : public ::Chaos::FClothingSimulationAccessoryMesh
+	{
+	public:
+		FClothSimulationAccessoryMesh(const TSharedRef<FCollectionClothConstFacade>& InClothFacade, const ::Chaos::FClothingSimulationMesh& InMesh, const int32 InAccessoryMeshIndex);
+
+		virtual ~FClothSimulationAccessoryMesh();
+
+		/* Return the number of points. */
+		virtual int32 GetNumPoints() const override;
+
+		/* Return the number of pattern points (2d, unwelded). */
+		virtual int32 GetNumPatternPoints() const override
+		{
+			return 0;
+		}
+
+		/* Return the source mesh positions (pre-skinning). */
+		virtual TConstArrayView<FVector3f> GetPositions() const override;
+
+		/* Return the source mesh 2d pattern positions. */
+		virtual TConstArrayView<FVector2f> GetPatternPositions() const override
+		{
+			return TConstArrayView<FVector2f>();
+		}
+
+		/* Return the source mesh normals (pre-skinning). */
+		virtual TConstArrayView<FVector3f> GetNormals() const override;
+
+		/* Return the bone data containing bone weights and influences. */
+		virtual TConstArrayView<FClothVertBoneData> GetBoneData() const override
+		{
+			return BoneData;
+		}
+		// ---- End of the Cloth interface ----
+
+	private:
+		void SetBoneData();
+
+		const FCollectionClothSimAccessoryMeshConstFacade AccessoryMeshFacade;
+		TArray<FClothVertBoneData> BoneData;
+	};
+} // namespace UE::Chaos::ClothAsset
+
diff --git a/Engine/Plugins/ChaosClothAsset/Source/ChaosClothAssetEngine/Private/ChaosClothAsset/ClothSimulationMesh.cpp b/Engine/Plugins/ChaosClothAsset/Source/ChaosClothAssetEngine/Private/ChaosClothAsset/ClothSimulationMesh.cpp
--- a/Engine/Plugins/ChaosClothAsset/Source/ChaosClothAssetEngine/Private/ChaosClothAsset/ClothSimulationMesh.cpp
+++ b/Engine/Plugins/ChaosClothAsset/Source/ChaosClothAssetEngine/Private/ChaosClothAsset/ClothSimulationMesh.cpp
@@ -1,6 +1,7 @@
 // Copyright Epic Games, Inc. All Rights Reserved.
 
 #include "ChaosClothAsset/ClothSimulationMesh.h"
+#include "ChaosClothAsset/ClothSimulationAccessoryMesh.h"
 #include "ChaosClothAsset/ClothAsset.h"
 #include "ChaosClothAsset/ClothSimulationModel.h"
 #include "ChaosClothAsset/ClothSimulationContext.h"
@@ -29,9 +30,16 @@ namespace UE::Chaos::ClothAsset
 		}
 #endif
 		ClothFacades.Reserve(ManagedArrayCollections.Num());
+		AccessoryMeshes.Reserve(ManagedArrayCollections.Num());
 		for (const TSharedRef<const FManagedArrayCollection>& Collection : ManagedArrayCollections)
 		{
-			ClothFacades.Emplace(MakeShared<FCollectionClothConstFacade>(Collection));
+			const TSharedRef<FCollectionClothConstFacade>& ClothFacade = ClothFacades.Emplace_GetRef(MakeShared<FCollectionClothConstFacade>(Collection));
+			TMap<FName, TUniquePtr<FClothSimulationAccessoryMesh>>& LODAccessoryMeshes = AccessoryMeshes.AddDefaulted_GetRef();
+			for (int32 AccessoryIndex = 0; AccessoryIndex < ClothFacade->GetNumSimAccessoryMeshes(); ++AccessoryIndex)
+			{
+				const FCollectionClothSimAccessoryMeshConstFacade AccessoryFacade = ClothFacade->GetSimAccessoryMesh(AccessoryIndex);
+				LODAccessoryMeshes.Add(AccessoryFacade.GetSimAccessoryMeshName(), MakeUnique<FClothSimulationAccessoryMesh>(ClothFacade, *this, AccessoryIndex));
+			}
 		}
 	}
 
@@ -289,4 +297,25 @@ namespace UE::Chaos::ClothAsset
 		}
 		return TConstArrayView<int32>();
 	}
+
+	TConstArrayView<FName> FClothSimulationMesh::GetAllAccessoryMeshNames(int32 LODIndex) const
+	{
+		if (ClothFacades.IsValidIndex(LODIndex))
+		{
+			return ClothFacades[LODIndex]->GetSimAccessoryMeshName();
+		}
+		return TConstArrayView<FName>();
+	}
+
+	const ::Chaos::FClothingSimulationAccessoryMesh* FClothSimulationMesh::GetAccessoryMesh(int32 LODIndex, const FName& AccessoryMeshName) const
+	{
+		if (AccessoryMeshes.IsValidIndex(LODIndex))
+		{
+			if (const TUniquePtr<FClothSimulationAccessoryMesh>* FoundMesh = AccessoryMeshes[LODIndex].Find(AccessoryMeshName))
+			{
+				return FoundMesh->Get();
+			}
+		}
+		return nullptr;
+	}
 }
diff --git a/Engine/Plugins/ChaosClothAsset/Source/ChaosClothAssetEngine/Private/ChaosClothAsset/ClothSimulationMesh.h b/Engine/Plugins/ChaosClothAsset/Source/ChaosClothAssetEngine/Private/ChaosClothAsset/ClothSimulationMesh.h
--- a/Engine/Plugins/ChaosClothAsset/Source/ChaosClothAssetEngine/Private/ChaosClothAsset/ClothSimulationMesh.h
+++ b/Engine/Plugins/ChaosClothAsset/Source/ChaosClothAssetEngine/Private/ChaosClothAsset/ClothSimulationMesh.h
@@ -10,6 +10,7 @@ namespace UE::Chaos::ClothAsset
 {
 	struct FClothSimulationContext;
 	class FCollectionClothConstFacade;
+	class FClothSimulationAccessoryMesh;
 
 	class FClothSimulationMesh : public ::Chaos::FClothingSimulationMesh
 	{
@@ -63,6 +64,8 @@ namespace UE::Chaos::ClothAsset
 		virtual TConstArrayView<FVector3f> GetMorphTargetPositionDeltas(int32 LODIndex, int32 MorphTargetIndex) const override;
 		virtual TConstArrayView<FVector3f> GetMorphTargetTangentZDeltas(int32 LODIndex, int32 MorphTargetIndex) const override;
 		virtual TConstArrayView<int32> GetMorphTargetIndices(int32 LODIndex, int32 MorphTargetIndex) const override;
+		virtual TConstArrayView<FName> GetAllAccessoryMeshNames(int32 LODIndex) const override;
+		virtual const ::Chaos::FClothingSimulationAccessoryMesh* GetAccessoryMesh(int32 LODIndex, const FName& AccessoryMeshName) const override;
 		//~ End FClothingSimulationMesh Interface
 
 	private:
@@ -70,5 +73,6 @@ namespace UE::Chaos::ClothAsset
 		const FClothSimulationContext& ClothSimulationContext;
 		const TArray<TSharedRef<const FManagedArrayCollection>>& ManagedArrayCollections;
 		TArray<TSharedRef<FCollectionClothConstFacade>> ClothFacades;
+		TArray<TMap<FName,TUniquePtr<FClothSimulationAccessoryMesh>>> AccessoryMeshes;
 	};
 }
diff --git a/Engine/Plugins/ChaosClothAsset/Source/ChaosClothAssetEngine/Private/ChaosClothAsset/ClothSimulationModel.cpp b/Engine/Plugins/ChaosClothAsset/Source/ChaosClothAssetEngine/Private/ChaosClothAsset/ClothSimulationModel.cpp
--- a/Engine/Plugins/ChaosClothAsset/Source/ChaosClothAssetEngine/Private/ChaosClothAsset/ClothSimulationModel.cpp
+++ b/Engine/Plugins/ChaosClothAsset/Source/ChaosClothAssetEngine/Private/ChaosClothAsset/ClothSimulationModel.cpp
@@ -304,6 +304,24 @@ FChaosClothSimulationModel::FChaosClothSimulationModel(const TArray<TSharedRef<c
 			}
 		}
 
+		// Add any accessory mesh required bones
+		for (int32 AccessoryMeshIndex = 0; AccessoryMeshIndex < Cloth.GetNumSimAccessoryMeshes(); ++AccessoryMeshIndex)
+		{
+			FCollectionClothSimAccessoryMeshConstFacade AccessoryMesh = Cloth.GetSimAccessoryMesh(AccessoryMeshIndex);
+
+			TConstArrayView<TArray<int32>> AccessoryBoneIndices = AccessoryMesh.GetSimAccessoryMeshBoneIndices();
+			for (const TArray<int32>& BoneIndices : AccessoryBoneIndices)
+			{
+				for (const int32 BoneIndex : BoneIndices)
+				{
+					if (BoneIndex != INDEX_NONE)
+					{
+						RequiredExtraSimBones.Add(BoneIndex);
+					}
+				}
+			}
+		}
+
 		// Batch tethers
 		LodModel.TetherData.GenerateTethers(MoveTemp(MergedTetherData));
 
diff --git a/Engine/Plugins/ChaosClothAssetEditor/Source/ChaosClothAssetDataflowNodes/Private/ChaosClothAsset/ChaosClothAssetDataflowNodesModule.cpp b/Engine/Plugins/ChaosClothAssetEditor/Source/ChaosClothAssetDataflowNodes/Private/ChaosClothAsset/ChaosClothAssetDataflowNodesModule.cpp
--- a/Engine/Plugins/ChaosClothAssetEditor/Source/ChaosClothAssetDataflowNodes/Private/ChaosClothAsset/ChaosClothAssetDataflowNodesModule.cpp
+++ b/Engine/Plugins/ChaosClothAssetEditor/Source/ChaosClothAssetDataflowNodes/Private/ChaosClothAsset/ChaosClothAssetDataflowNodesModule.cpp
@@ -33,6 +33,7 @@
 #include "ChaosClothAsset/SelectionToIntMapNode.h"
 #include "ChaosClothAsset/SelectionToWeightMapNode.h"
 #include "ChaosClothAsset/SetPhysicsAssetNode.h"
+#include "ChaosClothAsset/SimAccessoryMeshNode.h"
 #include "ChaosClothAsset/SimulationAerodynamicsConfigNode.h"
 #include "ChaosClothAsset/SimulationAnimDriveConfigNode.h"
 #include "ChaosClothAsset/SimulationBackstopConfigNode.h"
@@ -134,6 +135,7 @@ namespace UE::Chaos::ClothAsset
 			DATAFLOW_NODE_REGISTER_CREATION_FACTORY(FChaosClothAssetSelectionToIntMapNode);
 			DATAFLOW_NODE_REGISTER_CREATION_FACTORY(FChaosClothAssetSelectionToWeightMapNode);
 			DATAFLOW_NODE_REGISTER_CREATION_FACTORY(FChaosClothAssetSetPhysicsAssetNode);
+			DATAFLOW_NODE_REGISTER_CREATION_FACTORY(FChaosClothAssetSimAccessoryMeshNode);
 			DATAFLOW_NODE_REGISTER_CREATION_FACTORY(FChaosClothAssetSimulationAerodynamicsConfigNode);
 			DATAFLOW_NODE_REGISTER_CREATION_FACTORY(FChaosClothAssetSimulationAnimDriveConfigNode);
 			DATAFLOW_NODE_REGISTER_CREATION_FACTORY(FChaosClothAssetSimulationBackstopConfigNode);
diff --git a/Engine/Plugins/ChaosClothAssetEditor/Source/ChaosClothAssetDataflowNodes/Private/ChaosClothAsset/SimAccessoryMeshNode.cpp b/Engine/Plugins/ChaosClothAssetEditor/Source/ChaosClothAssetDataflowNodes/Private/ChaosClothAsset/SimAccessoryMeshNode.cpp
new file mode 100644
--- /dev/null
+++ b/Engine/Plugins/ChaosClothAssetEditor/Source/ChaosClothAssetDataflowNodes/Private/ChaosClothAsset/SimAccessoryMeshNode.cpp
@@ -0,0 +1,131 @@
+// Copyright Epic Games, Inc. All Rights Reserved.
+
+#include "ChaosClothAsset/SimAccessoryMeshNode.h"
+#include "ChaosClothAsset/ClothGeometryTools.h"
+#include "ChaosClothAsset/CollectionClothFacade.h"
+#include "Dataflow/DataflowInputOutput.h"
+#if WITH_EDITOR
+#include "ChaosClothAsset/ClothDataflowViewModes.h"
+#include "Dataflow/DataflowDebugDrawInterface.h"
+#include "Dataflow/DataflowRenderingViewMode.h"
+#endif
+
+#include UE_INLINE_GENERATED_CPP_BY_NAME(SimAccessoryMeshNode)
+
+
+namespace UE::Chaos::ClothAsset::Private
+{
+	struct FDebugDrawAccessoryMesh : public IDataflowDebugDrawInterface::IDebugDrawMesh
+	{
+		virtual int32 GetMaxVertexIndex() const override
+		{
+			return Points.Num();
+		}
+
+		virtual bool IsValidVertex(int32 VertexIndex) const override
+		{
+			return Points.IsValidIndex(VertexIndex);
+		}
+
+		virtual FVector GetVertexPosition(int32 VertexIndex) const override
+		{
+			return FVector(Points[VertexIndex]);
+		}
+
+		virtual FVector GetVertexNormal(int32 VertexIndex) const override
+		{
+			return FVector(Normals[VertexIndex]);
+		}
+
+		virtual int32 GetMaxTriangleIndex() const override
+		{
+			return Elements.Num();
+		}
+
+		virtual bool IsValidTriangle(int32 TriangleIndex) const override
+		{
+			return Elements.IsValidIndex(TriangleIndex);
+		}
+
+		virtual FIntVector3 GetTriangle(int32 TriangleIndex) const override
+		{
+			return Elements[TriangleIndex];
+		}
+
+		TConstArrayView<FIntVector3> Elements;
+		TConstArrayView<FVector3f> Points;
+		TConstArrayView<FVector3f> Normals;
+	};
+}
+
+FChaosClothAssetSimAccessoryMeshNode::FChaosClothAssetSimAccessoryMeshNode(const UE::Dataflow::FNodeParameters& InParam, FGuid InGuid)
+	: FDataflowNode(InParam, InGuid)
+{
+	RegisterInputConnection(&Collection);
+	RegisterInputConnection(&SimAccessoryMeshCollection);
+	RegisterOutputConnection(&Collection, &Collection);
+	RegisterOutputConnection(&AccessoryMeshName);
+}
+
+void FChaosClothAssetSimAccessoryMeshNode::Evaluate(UE::Dataflow::FContext& Context, const FDataflowOutput* Out) const
+{
+	if (Out->IsA<FManagedArrayCollection>(&Collection) || Out->IsA<FString>(&AccessoryMeshName))
+	{
+		using namespace UE::Chaos::ClothAsset;
+
+		// Evaluate in collection
+		FManagedArrayCollection InCollection = GetValue<FManagedArrayCollection>(Context, &Collection);
+		const TSharedRef<FManagedArrayCollection> ClothCollection = MakeShared<FManagedArrayCollection>(MoveTemp(InCollection));
+		FCollectionClothFacade ClothFacade(ClothCollection);
+		FString OutAccessoryMeshName = AccessoryMeshName;
+		if (ClothFacade.IsValid())
+		{
+			FManagedArrayCollection InAccessoryCollection = GetValue<FManagedArrayCollection>(Context, &SimAccessoryMeshCollection);  // Can't use a const reference here sadly since the facade needs a SharedRef to be created
+			const TSharedRef<const FManagedArrayCollection> AccessoryClothCollection = MakeShared<const FManagedArrayCollection>(MoveTemp(InAccessoryCollection));
+			FCollectionClothConstFacade AccessoryClothFacade(AccessoryClothCollection);
+
+			if (AccessoryClothFacade.IsValid())
+			{
+				const int32 AccessoryMeshIndex = FClothGeometryTools::CopySimMeshToSimAccessoryMesh(FName(OutAccessoryMeshName), ClothFacade, AccessoryClothFacade, bUseSimImportVertexID);
+				if (ClothFacade.GetSimAccessoryMeshName().IsValidIndex(AccessoryMeshIndex))
+				{
+					OutAccessoryMeshName = ClothFacade.GetSimAccessoryMeshName()[AccessoryMeshIndex].ToString();
+				}
+			}
+		}
+
+		SetValue(Context, MoveTemp(*ClothCollection), &Collection);
+		SetValue(Context, MoveTemp(OutAccessoryMeshName), &AccessoryMeshName);
+	}
+}
+
+#if WITH_EDITOR
+bool FChaosClothAssetSimAccessoryMeshNode::CanDebugDrawViewMode(const FName& ViewModeName) const
+{
+	return ViewModeName == UE::Dataflow::FDataflowConstruction3DViewMode::Name || ViewModeName == UE::Chaos::ClothAsset::FCloth3DSimViewMode::Name;
+}
+
+void FChaosClothAssetSimAccessoryMeshNode::DebugDraw(UE::Dataflow::FContext& Context, IDataflowDebugDrawInterface& DataflowRenderingInterface, const FDebugDrawParameters& DebugDrawParameters) const
+{
+	if (DebugDrawParameters.bNodeIsSelected || DebugDrawParameters.bNodeIsPinned)
+	{
+		using namespace UE::Chaos::ClothAsset;
+		const FName OutAccessoryMeshName = FName(GetOutputValue(Context, &AccessoryMeshName, AccessoryMeshName));
+		FManagedArrayCollection OutCollection = GetOutputValue<FManagedArrayCollection>(Context, &Collection, Collection);
+		const TSharedRef<FManagedArrayCollection> ClothCollection = MakeShared<FManagedArrayCollection>(MoveTemp(OutCollection));
+		FCollectionClothConstFacade ClothFacade(ClothCollection);
+		const int32 AccessoryMeshIndex = ClothFacade.FindSimAccessoryMeshIndexByName(OutAccessoryMeshName);
+		if (AccessoryMeshIndex >= 0 && AccessoryMeshIndex < ClothFacade.GetNumSimAccessoryMeshes())
+		{			
+			FCollectionClothSimAccessoryMeshConstFacade SimAccessoryMeshFacade = ClothFacade.GetSimAccessoryMesh(AccessoryMeshIndex);
+
+			UE::Chaos::ClothAsset::Private::FDebugDrawAccessoryMesh DrawAccessoryMesh;
+			DrawAccessoryMesh.Elements = ClothFacade.GetSimIndices3D();
+			DrawAccessoryMesh.Points = SimAccessoryMeshFacade.GetSimAccessoryMeshPosition3D();
+			DrawAccessoryMesh.Normals = SimAccessoryMeshFacade.GetSimAccessoryMeshNormal();
+
+			DataflowRenderingInterface.DrawMesh(DrawAccessoryMesh);
+		}
+	}
+}
+#endif
diff --git a/Engine/Plugins/ChaosClothAssetEditor/Source/ChaosClothAssetDataflowNodes/Private/ChaosClothAsset/SimAccessoryMeshNode.h b/Engine/Plugins/ChaosClothAssetEditor/Source/ChaosClothAssetDataflowNodes/Private/ChaosClothAsset/SimAccessoryMeshNode.h
new file mode 100644
--- /dev/null
+++ b/Engine/Plugins/ChaosClothAssetEditor/Source/ChaosClothAssetDataflowNodes/Private/ChaosClothAsset/SimAccessoryMeshNode.h
@@ -0,0 +1,48 @@
+// Copyright Epic Games, Inc. All Rights Reserved.
+
+#pragma once 
+
+#include "Dataflow/DataflowNode.h"
+#include "GeometryCollection/ManagedArrayCollection.h"
+#include "SimAccessoryMeshNode.generated.h"
+
+
+/** Add a SimAccessoryMesh by converting a cloth collection into an accessory mesh and attaching it to an existing cloth collection. Any unmatched vertices will use the existing cloth collection's sim mesh data to populate the accessory mesh data */
+USTRUCT(Meta = (DataflowCloth))
+struct FChaosClothAssetSimAccessoryMeshNode : public FDataflowNode
+{
+	GENERATED_USTRUCT_BODY()
+	DATAFLOW_NODE_DEFINE_INTERNAL(FChaosClothAssetSimAccessoryMeshNode, "SimAccessoryMeshNode", "Cloth", "Cloth Sim Accessory Mesh")
+	DATAFLOW_NODE_RENDER_TYPE("SurfaceRender", FName("FClothCollection"), "Collection")
+
+public:
+
+	FChaosClothAssetSimAccessoryMeshNode(const UE::Dataflow::FNodeParameters& InParam, FGuid InGuid = FGuid::NewGuid());
+
+private:
+	//~ Begin FDataflowNode interface
+	virtual void Evaluate(UE::Dataflow::FContext& Context, const FDataflowOutput* Out) const override;
+#if WITH_EDITOR
+	virtual bool CanDebugDraw() const override
+	{
+		return true;
+	}
+	virtual bool CanDebugDrawViewMode(const FName& ViewModeName) const override;
+	virtual void DebugDraw(UE::Dataflow::FContext& Context, IDataflowDebugDrawInterface& DataflowRenderingInterface, const FDebugDrawParameters& DebugDrawParameters) const override;
+#endif
+	//~ End FDataflowNode interface
+
+	UPROPERTY(Meta = (DataflowInput, DataflowOutput, DataflowPassthrough = "Collection"))
+	FManagedArrayCollection Collection;
+
+	UPROPERTY(Meta = (DataflowInput))
+	FManagedArrayCollection SimAccessoryMeshCollection;
+
+	/** Name of the new accessory mesh */
+	UPROPERTY(EditAnywhere, Category = "Sim Accessory Mesh", Meta = (DataflowOutput))
+	FString AccessoryMeshName = TEXT("AccessoryMesh");
+
+	/** Use SimImportVertexID (e.g., imported vertex ID from an FBX->SKM->ClothCollection conversion) to match vertices*/
+	UPROPERTY(EditAnywhere, Category = "Sim Accessory Mesh")
+	bool bUseSimImportVertexID = true;
+};
diff --git a/Engine/Plugins/ChaosClothAssetEditor/Source/ChaosClothAssetDataflowNodes/Private/ChaosClothAsset/SimulationBackstopConfigNode.cpp b/Engine/Plugins/ChaosClothAssetEditor/Source/ChaosClothAssetDataflowNodes/Private/ChaosClothAsset/SimulationBackstopConfigNode.cpp
--- a/Engine/Plugins/ChaosClothAssetEditor/Source/ChaosClothAssetDataflowNodes/Private/ChaosClothAsset/SimulationBackstopConfigNode.cpp
+++ b/Engine/Plugins/ChaosClothAssetEditor/Source/ChaosClothAssetDataflowNodes/Private/ChaosClothAsset/SimulationBackstopConfigNode.cpp
@@ -15,10 +15,14 @@ FChaosClothAssetSimulationBackstopConfigNode::FChaosClothAssetSimulationBackstop
 	RegisterInputConnection(&BackstopRadius.WeightMap)
 		.SetCanHidePin(true)
 		.SetPinIsHidden(true);
+	RegisterInputConnection(&BackstopMeshName.StringValue, GET_MEMBER_NAME_CHECKED(FChaosClothAssetConnectableIStringValue, StringValue))
+		.SetCanHidePin(true)
+		.SetPinIsHidden(true);
 }
 
 void FChaosClothAssetSimulationBackstopConfigNode::AddProperties(FPropertyHelper& PropertyHelper) const
 {
 	PropertyHelper.SetPropertyWeighted(this, &BackstopDistance);
 	PropertyHelper.SetPropertyWeighted(this, &BackstopRadius);
+	PropertyHelper.SetPropertyString(this, &BackstopMeshName, {}, ECollectionPropertyFlags::None);
 }
diff --git a/Engine/Plugins/ChaosClothAssetEditor/Source/ChaosClothAssetDataflowNodes/Private/ChaosClothAsset/SkeletalMeshImportNode.cpp b/Engine/Plugins/ChaosClothAssetEditor/Source/ChaosClothAssetDataflowNodes/Private/ChaosClothAsset/SkeletalMeshImportNode.cpp
--- a/Engine/Plugins/ChaosClothAssetEditor/Source/ChaosClothAssetDataflowNodes/Private/ChaosClothAsset/SkeletalMeshImportNode.cpp
+++ b/Engine/Plugins/ChaosClothAssetEditor/Source/ChaosClothAssetDataflowNodes/Private/ChaosClothAsset/SkeletalMeshImportNode.cpp
@@ -51,7 +51,7 @@ void FChaosClothAssetSkeletalMeshImportNode_v2::Evaluate(UE::Dataflow::FContext&
 	{
 		const TSharedRef<FManagedArrayCollection> ClothCollection = MakeShared<FManagedArrayCollection>();
 		FCollectionClothFacade ClothFacade(ClothCollection);
-		ClothFacade.DefineSchema();
+		ClothFacade.DefineSchema(EClothCollectionExtendedSchemas::Import);
 
 		if (const USkeletalMesh* const InSkeletalMesh = GetValue(Context, &SkeletalMesh))
 		{
@@ -104,63 +104,80 @@ void FChaosClothAssetSkeletalMeshImportNode_v2::Evaluate(UE::Dataflow::FContext&
 				}
 			}
 
-			if (bImportSimMorphTargets)
+			if (bImportSimMesh)
 			{
-				// Use Sim2DToSKMIndex to get a map from SKMIndex to Sim3D index
-				const TMap<int32, int32> SKMToSim3DLookup = Private::GenerateSKMToSim3DLookup(ClothFacade, Sim2DToSKMIndex);
+				// Set SimImportVertexID
+				TArrayView<int32> SimImportVertexID = ClothFacade.GetSimImportVertexID();
+				for (int32 Vertex2DIndex = 0; Vertex2DIndex < SimImportVertexID.Num(); ++Vertex2DIndex)
+				{
+					if (LODModel.MeshToImportVertexMap.IsValidIndex(Sim2DToSKMIndex[Vertex2DIndex]))
+					{
+						SimImportVertexID[Vertex2DIndex] = LODModel.MeshToImportVertexMap[Sim2DToSKMIndex[Vertex2DIndex]];
+					}
+					else
+					{
+						SimImportVertexID[Vertex2DIndex] = INDEX_NONE;
+					}
+				}
 
-				const TArray<TObjectPtr<UMorphTarget>>& MorphTargets = InSkeletalMesh->GetMorphTargets();
-				for (const UMorphTarget* const MorphTarget : MorphTargets)
+				if (bImportSimMorphTargets)
 				{
-					if (MorphTarget && MorphTarget->HasDataForLOD(LODIndex))
+					// Use Sim2DToSKMIndex to get a map from SKMIndex to Sim3D index
+					const TMap<int32, int32> SKMToSim3DLookup = Private::GenerateSKMToSim3DLookup(ClothFacade, Sim2DToSKMIndex);
+
+					const TArray<TObjectPtr<UMorphTarget>>& MorphTargets = InSkeletalMesh->GetMorphTargets();
+					for (const UMorphTarget* const MorphTarget : MorphTargets)
 					{
-						const FMorphTargetLODModel& MorphTargetLODModel = MorphTarget->GetMorphLODModels()[LODIndex];
-						
-						TArray<int32> MissingIndices;
-						TArray<FVector3f> PositionDelta;
-						TArray<FVector3f> NormalZDelta;
-						TArray<int32> SimIndex;
-						PositionDelta.Reserve(MorphTargetLODModel.NumVertices);
-						NormalZDelta.Reserve(MorphTargetLODModel.NumVertices);
-						SimIndex.Reserve(MorphTargetLODModel.NumVertices);
-						for (const FMorphTargetDelta& Delta : MorphTargetLODModel.Vertices)
+						if (MorphTarget && MorphTarget->HasDataForLOD(LODIndex))
 						{
-							if (const int32* const SimIndex3D = SKMToSim3DLookup.Find(Delta.SourceIdx))
-							{
-								PositionDelta.Add(Delta.PositionDelta);
-								NormalZDelta.Add(Delta.TangentZDelta);
-								SimIndex.Add(*SimIndex3D);
-							}
-							else
+							const FMorphTargetLODModel& MorphTargetLODModel = MorphTarget->GetMorphLODModels()[LODIndex];
+
+							TArray<int32> MissingIndices;
+							TArray<FVector3f> PositionDelta;
+							TArray<FVector3f> NormalZDelta;
+							TArray<int32> SimIndex;
+							PositionDelta.Reserve(MorphTargetLODModel.NumVertices);
+							NormalZDelta.Reserve(MorphTargetLODModel.NumVertices);
+							SimIndex.Reserve(MorphTargetLODModel.NumVertices);
+							for (const FMorphTargetDelta& Delta : MorphTargetLODModel.Vertices)
 							{
-								MissingIndices.Add(Delta.SourceIdx);
+								if (const int32* const SimIndex3D = SKMToSim3DLookup.Find(Delta.SourceIdx))
+								{
+									PositionDelta.Add(Delta.PositionDelta);
+									NormalZDelta.Add(Delta.TangentZDelta);
+									SimIndex.Add(*SimIndex3D);
+								}
+								else
+								{
+									MissingIndices.Add(Delta.SourceIdx);
+								}
 							}
-						}
 
-						const int32 NumTargetVertices = PositionDelta.Num();
-						check(NormalZDelta.Num() == NumTargetVertices);
-						check(SimIndex.Num() == NumTargetVertices);
-						if (NumTargetVertices > 0)
-						{
-							FCollectionClothSimMorphTargetFacade MorphTargetFacade = ClothFacade.AddGetSimMorphTarget();
-							MorphTargetFacade.Initialize(MorphTarget->GetName(), PositionDelta, NormalZDelta, SimIndex);
-							if (!MissingIndices.IsEmpty())
+							const int32 NumTargetVertices = PositionDelta.Num();
+							check(NormalZDelta.Num() == NumTargetVertices);
+							check(SimIndex.Num() == NumTargetVertices);
+							if (NumTargetVertices > 0)
 							{
-								FString MissingIndicesString;
-								for (int32 Index = 0; Index < MissingIndices.Num(); ++Index)
+								FCollectionClothSimMorphTargetFacade MorphTargetFacade = ClothFacade.AddGetSimMorphTarget();
+								MorphTargetFacade.Initialize(MorphTarget->GetName(), PositionDelta, NormalZDelta, SimIndex);
+								if (!MissingIndices.IsEmpty())
 								{
-									MissingIndicesString += FString::FromInt(MissingIndices[Index]);
-									if (Index != MissingIndices.Num() - 1)
+									FString MissingIndicesString;
+									for (int32 Index = 0; Index < MissingIndices.Num(); ++Index)
 									{
-										MissingIndicesString += ", ";
+										MissingIndicesString += FString::FromInt(MissingIndices[Index]);
+										if (Index != MissingIndices.Num() - 1)
+										{
+											MissingIndicesString += ", ";
+										}
 									}
+									Context.Warning(FString::Printf(TEXT("Failed to find corresponding sim vertex index for the following indices for morph target '%s': %s"), *MorphTarget->GetName(), *MissingIndicesString), this, Out);
 								}
-								Context.Warning(FString::Printf(TEXT("Failed to find corresponding sim vertex index for the following indices for morph target '%s': %s"), *MorphTarget->GetName(), *MissingIndicesString), this, Out);
 							}
-						}
-						else
-						{
-							Context.Warning(FString::Printf(TEXT("Failed to import morph target '%s': no valid vertices found"), *MorphTarget->GetName()), this, Out);
+							else
+							{
+								Context.Warning(FString::Printf(TEXT("Failed to import morph target '%s': no valid vertices found"), *MorphTarget->GetName()), this, Out);
+							}
 						}
 					}
 				}
diff --git a/Engine/Plugins/ChaosClothAssetEditor/Source/ChaosClothAssetDataflowNodes/Public/ChaosClothAsset/SimulationBackstopConfigNode.h b/Engine/Plugins/ChaosClothAssetEditor/Source/ChaosClothAssetDataflowNodes/Public/ChaosClothAsset/SimulationBackstopConfigNode.h
--- a/Engine/Plugins/ChaosClothAssetEditor/Source/ChaosClothAssetDataflowNodes/Public/ChaosClothAsset/SimulationBackstopConfigNode.h
+++ b/Engine/Plugins/ChaosClothAssetEditor/Source/ChaosClothAssetDataflowNodes/Public/ChaosClothAsset/SimulationBackstopConfigNode.h
@@ -33,6 +33,11 @@ public:
 	UPROPERTY(EditAnywhere, Category = "Backstop Properties", Meta = (UIMin = "0", UIMax = "100", ClampMin = "0", InteractorName = "BackstopRadius"))
 	FChaosClothAssetWeightedValue BackstopRadius = { true, 0.f, 100.f, TEXT("BackstopRadius") };
 
+	/** The name of the accessory mesh to use as the backstop mesh. If none is specified, the default animated mesh will be used */
+	UPROPERTY(EditAnywhere, Category = "Backstop Properties")
+	FChaosClothAssetConnectableIStringValue BackstopMeshName;
+
+
 	FChaosClothAssetSimulationBackstopConfigNode(const UE::Dataflow::FNodeParameters& InParam, FGuid InGuid = FGuid::NewGuid());
 
 private:
diff --git a/Engine/Plugins/ChaosClothAssetEditor/Source/ChaosClothAssetDataflowNodes/Public/ChaosClothAsset/SkeletalMeshImportNode.h b/Engine/Plugins/ChaosClothAssetEditor/Source/ChaosClothAssetDataflowNodes/Public/ChaosClothAsset/SkeletalMeshImportNode.h
--- a/Engine/Plugins/ChaosClothAssetEditor/Source/ChaosClothAssetDataflowNodes/Public/ChaosClothAsset/SkeletalMeshImportNode.h
+++ b/Engine/Plugins/ChaosClothAssetEditor/Source/ChaosClothAssetDataflowNodes/Public/ChaosClothAsset/SkeletalMeshImportNode.h
@@ -70,7 +70,7 @@ public:
 	bool bSetPhysicsAsset = false;
 
 	/** Import morph targets as Sim Mesh Morph Targets */
-	UPROPERTY(EditAnywhere, Category = "Skeletal Mesh Import")
+	UPROPERTY(EditAnywhere, Category = "Skeletal Mesh Import", Meta = (EditCondition = "bImportSimMesh"))
 	bool bImportSimMorphTargets = false;
 
 	FChaosClothAssetSkeletalMeshImportNode_v2(const UE::Dataflow::FNodeParameters& InParam, FGuid InGuid = FGuid::NewGuid());
diff --git a/Engine/Plugins/ChaosClothAssetEditor/Source/ChaosClothAssetEditor/Private/ChaosClothAsset/ClothEditorSimulationVisualization.cpp b/Engine/Plugins/ChaosClothAssetEditor/Source/ChaosClothAssetEditor/Private/ChaosClothAsset/ClothEditorSimulationVisualization.cpp
--- a/Engine/Plugins/ChaosClothAssetEditor/Source/ChaosClothAssetEditor/Private/ChaosClothAsset/ClothEditorSimulationVisualization.cpp
+++ b/Engine/Plugins/ChaosClothAssetEditor/Source/ChaosClothAssetEditor/Private/ChaosClothAsset/ClothEditorSimulationVisualization.cpp
@@ -438,6 +438,32 @@ const FVisualizationOption FVisualizationOption::OptionData[] =
 		{
 			EditorVisualization.ExtendViewportShowMenuMorphTargetSelector(MenuBuilder);
 		})),
+	FVisualizationOption(
+		FClothVisualizationDebugDraw::CreateLambda([](const FClothEditorSimulationVisualization& EditorVisualization, const ::Chaos::FClothVisualizationNoGC& Visualization, FPrimitiveDrawInterface* PDI)
+		{
+			const FString* const AccessoryMeshName = EditorVisualization.GetCurrentlySelectedAccessoryMesh();
+			Visualization.DrawAccessoryMesh(PDI, AccessoryMeshName ? FName(*AccessoryMeshName) : NAME_None);
+		}),
+		LOCTEXT("ChaosVisName_DrawAccessoryMesh", "Accessory Mesh"), 
+		LOCTEXT("ChaosVisName_DrawAccessoryMesh_ToolTip", "Draw an accessory mesh."), 
+		/*bDisablesSimulation =*/false, /*bHidesClothSections=*/false, /*bDefaultFlagValue=*/false,
+		FAdditionalMenus::CreateLambda([](FClothEditorSimulationVisualization& EditorVisualization, FMenuBuilder& MenuBuilder)
+		{
+			EditorVisualization.ExtendViewportShowMenuAccessoryMeshSelector(MenuBuilder);
+		})),
+	FVisualizationOption(
+		FClothVisualizationDebugDraw::CreateLambda([](const FClothEditorSimulationVisualization& EditorVisualization, const ::Chaos::FClothVisualizationNoGC& Visualization, FPrimitiveDrawInterface* PDI)
+		{
+			const FString* const AccessoryMeshName = EditorVisualization.GetCurrentlySelectedAccessoryMesh();
+			Visualization.DrawAccessoryMeshNormals(PDI, AccessoryMeshName ? FName(*AccessoryMeshName) : NAME_None, EditorVisualization.GetAccessoryMeshNormalLength());
+		}),
+		LOCTEXT("ChaosVisName_DrawAccessoryMeshNormals", "Accessory Mesh Normals"), 
+		LOCTEXT("ChaosVisName_DrawAccessoryMeshNormals_ToolTip", "Draw an accessory mesh's normals."), 
+		/*bDisablesSimulation =*/false, /*bHidesClothSections=*/false, /*bDefaultFlagValue=*/false,
+		FAdditionalMenus::CreateLambda([](FClothEditorSimulationVisualization& EditorVisualization, FMenuBuilder& MenuBuilder)
+		{
+			EditorVisualization.ExtendViewportShowMenuAccessoryMeshNormalsLength(MenuBuilder);
+		})),
 	FVisualizationOption(
 		FClothVisualizationDebugDraw::CreateLambda([](const FClothEditorSimulationVisualization&, const ::Chaos::FClothVisualizationNoGC& Visualization, FPrimitiveDrawInterface* PDI)
 		{
@@ -583,10 +609,17 @@ void FClothEditorSimulationVisualization::RefreshMenusForClothComponent(const UC
 {
 	TArray<FString> WeightMaps;
 	TArray<FString> MorphTargets;
+	TArray<FString> AccessoryMeshes;
 	if (const ::Chaos::FClothVisualizationNoGC* const Visualization = ClothComponent && ClothComponent->GetClothSimulationProxy() ? ClothComponent->GetClothSimulationProxy()->GetClothVisualization() : nullptr)
 	{
 		WeightMaps = Visualization->GetAllWeightMapNames();
 		MorphTargets = Visualization->GetAllMorphTargetNames();
+		const TArray<FName> AccessoryMeshNames = Visualization->GetAllAccessoryMeshNames();
+		AccessoryMeshes.Reserve(AccessoryMeshNames.Num());
+		for (const FName& AccessoryMeshName : AccessoryMeshNames)
+		{
+			AccessoryMeshes.Add(AccessoryMeshName.ToString());
+		}
 	}
 
 	auto RefreshMenuNames = [](const TArray<FString>& AllNames, FNameSelectionData& SelectionData)
@@ -622,6 +655,7 @@ void FClothEditorSimulationVisualization::RefreshMenusForClothComponent(const UC
 
 	RefreshMenuNames(WeightMaps, WeightMapSelection);
 	RefreshMenuNames(MorphTargets, MorphTargetSelection);
+	RefreshMenuNames(AccessoryMeshes, AccessoryMeshSelection);
 }
 
 void FClothEditorSimulationVisualization::DebugDrawSimulation(const UChaosClothComponent* ClothComponent, FPrimitiveDrawInterface* PDI)
diff --git a/Engine/Plugins/ChaosClothAssetEditor/Source/ChaosClothAssetEditor/Private/ChaosClothAsset/ClothEditorSimulationVisualization.h b/Engine/Plugins/ChaosClothAssetEditor/Source/ChaosClothAssetEditor/Private/ChaosClothAsset/ClothEditorSimulationVisualization.h
--- a/Engine/Plugins/ChaosClothAssetEditor/Source/ChaosClothAssetEditor/Private/ChaosClothAsset/ClothEditorSimulationVisualization.h
+++ b/Engine/Plugins/ChaosClothAssetEditor/Source/ChaosClothAssetEditor/Private/ChaosClothAsset/ClothEditorSimulationVisualization.h
@@ -47,6 +47,12 @@ public:
 	{
 		ExtendViewportShowMenuNameSelector(MenuBuilder, MorphTargetSelection);
 	}
+	// Accessory Meshes
+	const FString* GetCurrentlySelectedAccessoryMesh() const { return AccessoryMeshSelection.CurrentlySelectedName.Get(); }
+	void ExtendViewportShowMenuAccessoryMeshSelector(FMenuBuilder& MenuBuilder)
+	{
+		ExtendViewportShowMenuNameSelector(MenuBuilder, AccessoryMeshSelection);
+	}
 	// Normals
 	void ExtendViewportShowMenuPointNormalsLength(FMenuBuilder& MenuBuilder)
 	{
@@ -60,8 +66,13 @@ public:
 	{
 		ExtendViewportShowMenuSpinBox(MenuBuilder, AerodynamicsLengthScale, 0.f, FLT_MAX, 0.f, 40.f);
 	}
+	void ExtendViewportShowMenuAccessoryMeshNormalsLength(FMenuBuilder& MenuBuilder)
+	{
+		ExtendViewportShowMenuSpinBox(MenuBuilder, AccessoryMeshNormalLength, 0.f, FLT_MAX, 0.f, 40.f);
+	}
 	float GetPointNormalLength() const { return PointNormalLength; }
 	float GetAnimatedNormalLength() const { return AnimatedNormalLength; }
+	float GetAccessoryMeshNormalLength() const { return AccessoryMeshNormalLength; }
 	float GetAerodynamicsLengthScale() const { return AerodynamicsLengthScale; }
 
 private:
@@ -87,8 +98,10 @@ private:
 	TBitArray<> Flags;
 	FNameSelectionData WeightMapSelection;
 	FNameSelectionData MorphTargetSelection;
+	FNameSelectionData AccessoryMeshSelection;
 	float PointNormalLength = 20.f;
 	float AnimatedNormalLength = 20.f;
+	float AccessoryMeshNormalLength = 20.f;
 	float AerodynamicsLengthScale = 10.f;
 	
 };
diff --git a/Engine/Source/Runtime/Experimental/Chaos/Private/Chaos/PBDSphericalConstraint.cpp b/Engine/Source/Runtime/Experimental/Chaos/Private/Chaos/PBDSphericalConstraint.cpp
--- a/Engine/Source/Runtime/Experimental/Chaos/Private/Chaos/PBDSphericalConstraint.cpp
+++ b/Engine/Source/Runtime/Experimental/Chaos/Private/Chaos/PBDSphericalConstraint.cpp
@@ -112,8 +112,8 @@ void FPBDSphericalBackstopConstraint::ApplyLegacyHelperISPC(SolverParticlesOrRan
 {
 	check(bRealTypeCompatibleWithISPC);
 
-	const TConstArrayView<FSolverVec3> AnimationPositionsView = Particles.GetConstArrayView(AnimationPositions);
-	const TConstArrayView<FSolverVec3> AnimationNormalsView = Particles.GetConstArrayView(AnimationNormals);
+	const TConstArrayView<FSolverVec3> AnimationPositionsView = bMeshUseGlobalIndexation ? Particles.GetConstArrayView(AnimationPositions) : TConstArrayView<FSolverVec3>(AnimationPositions);
+	const TConstArrayView<FSolverVec3> AnimationNormalsView = bMeshUseGlobalIndexation ? Particles.GetConstArrayView(AnimationNormals) : TConstArrayView<FSolverVec3>(AnimationNormals);
 #if INTEL_ISPC
 	if (SphereRadii.Num() == ParticleCount)
 	{
@@ -180,8 +180,8 @@ void FPBDSphericalBackstopConstraint::ApplyHelperISPC(SolverParticlesOrRange& Pa
 {
 	check(bRealTypeCompatibleWithISPC);
 
-	const TConstArrayView<FSolverVec3> AnimationPositionsView = Particles.GetConstArrayView(AnimationPositions);
-	const TConstArrayView<FSolverVec3> AnimationNormalsView = Particles.GetConstArrayView(AnimationNormals);
+	const TConstArrayView<FSolverVec3> AnimationPositionsView = bMeshUseGlobalIndexation ? Particles.GetConstArrayView(AnimationPositions) : TConstArrayView<FSolverVec3>(AnimationPositions);
+	const TConstArrayView<FSolverVec3> AnimationNormalsView = bMeshUseGlobalIndexation ? Particles.GetConstArrayView(AnimationNormals) : TConstArrayView<FSolverVec3>(AnimationNormals);
 #if INTEL_ISPC
 	if (SphereRadii.Num() == ParticleCount)
 	{
diff --git a/Engine/Source/Runtime/Experimental/Chaos/Public/Chaos/PBDSphericalConstraint.h b/Engine/Source/Runtime/Experimental/Chaos/Public/Chaos/PBDSphericalConstraint.h
--- a/Engine/Source/Runtime/Experimental/Chaos/Public/Chaos/PBDSphericalConstraint.h
+++ b/Engine/Source/Runtime/Experimental/Chaos/Public/Chaos/PBDSphericalConstraint.h
@@ -170,11 +170,12 @@ public:
 	FPBDSphericalBackstopConstraint(
 		const int32 InParticleOffset,
 		const int32 InParticleCount,
-		const TArray<FSolverVec3>& InAnimationPositions,  // Use global indexation (will need adding ParticleOffset)
-		const TArray<FSolverVec3>& InAnimationNormals,  // Use global indexation (will need adding ParticleOffset)
+		const TArray<FSolverVec3>& InAnimationPositions,
+		const TArray<FSolverVec3>& InAnimationNormals,
 		const TMap<FString, TConstArrayView<FRealSingle>>& WeightMaps,  // Use local indexation
 		const FCollectionPropertyConstFacade& PropertyCollection,
-		FSolverReal MeshScale
+		FSolverReal MeshScale, 
+		bool bInUseGlobalIndexation = true   // Use global indexation (will need adding ParticleOffset)
 	)
 		: AnimationPositions(InAnimationPositions)
 		, AnimationNormals(InAnimationNormals)
@@ -188,6 +189,7 @@ public:
 		, BackstopDistanceBase((FSolverReal)GetLowBackstopDistance(PropertyCollection, 0.f))
 		, BackstopDistanceRange((FSolverReal)GetHighBackstopDistance(PropertyCollection, 1.f) - BackstopDistanceBase)
 		, bUseLegacyBackstop(GetUseLegacyBackstop(PropertyCollection, false))  // Only set the legacy backstop in constructor
+		, bMeshUseGlobalIndexation(bInUseGlobalIndexation)
 		, BackstopDistanceIndex(PropertyCollection)
 		, BackstopRadiusIndex(PropertyCollection)
 		, UseLegacyBackstopIndex(PropertyCollection)
@@ -197,11 +199,12 @@ public:
 	FPBDSphericalBackstopConstraint(
 		const int32 InParticleOffset,
 		const int32 InParticleCount,
-		const TArray<FSolverVec3>& InAnimationPositions,  // Use global indexation (will need adding ParticleOffset)
-		const TArray<FSolverVec3>& InAnimationNormals,  // Use global indexation (will need adding ParticleOffset)
+		const TArray<FSolverVec3>& InAnimationPositions,
+		const TArray<FSolverVec3>& InAnimationNormals,
 		const TConstArrayView<FRealSingle>& InSphereRadii,  // Use local indexation
 		const TConstArrayView<FRealSingle>& InSphereOffsetDistances,  // Use local indexation
-		const bool bInUseLegacyBackstop  // Do not include the sphere radius in the distance calculations when this is true
+		const bool bInUseLegacyBackstop,  // Do not include the sphere radius in the distance calculations when this is true, 
+		bool bInUseGlobalIndexation = true  // Use global indexation (will need adding ParticleOffset)
 	)
 		: AnimationPositions(InAnimationPositions)
 		, AnimationNormals(InAnimationNormals)
@@ -212,6 +215,7 @@ public:
 		, Scale((FSolverReal)1.)
 		, bEnabled(true)
 		, bUseLegacyBackstop(bInUseLegacyBackstop)
+		, bMeshUseGlobalIndexation(bInUseGlobalIndexation)
 		, BackstopDistanceIndex(ForceInit)
 		, BackstopRadiusIndex(ForceInit)
 		, UseLegacyBackstopIndex(ForceInit)
@@ -331,14 +335,15 @@ public:
 	UE_CHAOS_DECLARE_INDEXLESS_PROPERTYCOLLECTION_NAME(BackstopDistance, float);
 	UE_CHAOS_DECLARE_INDEXLESS_PROPERTYCOLLECTION_NAME(BackstopRadius, float);
 	UE_CHAOS_DECLARE_INDEXLESS_PROPERTYCOLLECTION_NAME(UseLegacyBackstop, bool);
+	UE_CHAOS_DECLARE_INDEXLESS_PROPERTYCOLLECTION_NAME(BackstopMeshName, bool); // Actually just using for string name
 
 private:
 
 	template<bool bHasBackstopDistance, bool bHasBackstopRadius, typename SolverParticlesOrRange>
 	void ApplyHelper(SolverParticlesOrRange& Particles, const FSolverReal Dt) const
 	{
-		const TConstArrayView<FSolverVec3> AnimationPositionsView = Particles.GetConstArrayView(AnimationPositions);
-		const TConstArrayView<FSolverVec3> AnimationNormalsView = Particles.GetConstArrayView(AnimationNormals);
+		const TConstArrayView<FSolverVec3> AnimationPositionsView = bMeshUseGlobalIndexation ? Particles.GetConstArrayView(AnimationPositions) : TConstArrayView<FSolverVec3>(AnimationPositions);
+		const TConstArrayView<FSolverVec3> AnimationNormalsView = bMeshUseGlobalIndexation ? Particles.GetConstArrayView(AnimationNormals) : TConstArrayView<FSolverVec3>(AnimationNormals);
 		PhysicsParallelFor(ParticleCount, [this, &Particles, Dt, &AnimationPositionsView, &AnimationNormalsView](int32 Index)  // TODO: profile need for parallel loop based on particle count
 		{
 			const int32 ParticleIndex = ParticleOffset + Index;
@@ -375,8 +380,8 @@ private:
 	template<bool bHasBackstopDistance, bool bHasBackstopRadius, typename SolverParticlesOrRange>
 	void ApplyLegacyHelper(SolverParticlesOrRange& Particles, const FSolverReal Dt) const
 	{
-		const TConstArrayView<FSolverVec3> AnimationPositionsView = Particles.GetConstArrayView(AnimationPositions);
-		const TConstArrayView<FSolverVec3> AnimationNormalsView = Particles.GetConstArrayView(AnimationNormals);
+		const TConstArrayView<FSolverVec3> AnimationPositionsView = bMeshUseGlobalIndexation ? Particles.GetConstArrayView(AnimationPositions) : TConstArrayView<FSolverVec3>(AnimationPositions);
+		const TConstArrayView<FSolverVec3> AnimationNormalsView = bMeshUseGlobalIndexation ? Particles.GetConstArrayView(AnimationNormals) : TConstArrayView<FSolverVec3>(AnimationNormals);
 		PhysicsParallelFor(ParticleCount, [this, &Particles, Dt, &AnimationPositionsView, &AnimationNormalsView](int32 Index)  // TODO: profile need for parallel loop based on particle count
 		{
 			const int32 ParticleIndex = ParticleOffset + Index;
@@ -416,8 +421,8 @@ private:
 	CHAOS_API void ApplyHelperISPC(SolverParticlesOrRange& Particles, const FSolverReal Dt) const;
 
 private:
-	const TArray<FSolverVec3>& AnimationPositions;  // Positions of spheres, use global indexation (will need adding ParticleOffset)
-	const TArray<FSolverVec3>& AnimationNormals; // Sphere offset directions, use global indexation (will need adding ParticleOffset)
+	const TArray<FSolverVec3>& AnimationPositions;  // Positions of spheres
+	const TArray<FSolverVec3>& AnimationNormals; // Sphere offset directions
 	TConstArrayView<FRealSingle> SphereRadii; // Start at index 0, use local indexation
 	TConstArrayView<FRealSingle> SphereOffsetDistances;  // Sphere position offsets, use local indexation
 	const int32 ParticleOffset;
@@ -429,6 +434,7 @@ private:
 	FSolverReal BackstopDistanceRange = (FSolverReal)1.;
 	bool bEnabled = true;
 	bool bUseLegacyBackstop = false;
+	bool bMeshUseGlobalIndexation = true; // use global indexation (will need adding ParticleOffset)
 
 	UE_CHAOS_DECLARE_INDEXED_PROPERTYCOLLECTION_NAME(BackstopDistance, float);
 	UE_CHAOS_DECLARE_INDEXED_PROPERTYCOLLECTION_NAME(BackstopRadius, float);
