Chaos SKM Cloth Asset- Fixed bounds calculation.

#rnx

diff --git a/Engine/Plugins/ChaosClothAsset/Source/ChaosClothAssetEngine/Private/ChaosClothAsset/ClothAssetSKMClothingSimulation.cpp b/Engine/Plugins/ChaosClothAsset/Source/ChaosClothAssetEngine/Private/ChaosClothAsset/ClothAssetSKMClothingSimulation.cpp
--- a/Engine/Plugins/ChaosClothAsset/Source/ChaosClothAssetEngine/Private/ChaosClothAsset/ClothAssetSKMClothingSimulation.cpp
+++ b/Engine/Plugins/ChaosClothAsset/Source/ChaosClothAssetEngine/Private/ChaosClothAsset/ClothAssetSKMClothingSimulation.cpp
@@ -153,7 +153,7 @@ namespace UE::Chaos::ClothAsset
 	{
 		if (ClothSimulationProxy.IsValid())
 		{
-			ClothSimulationProxy->CalculateBounds_AnyThread();
+			return ClothSimulationProxy->CalculateBounds_AnyThread();
 		}
 		return FBoxSphereBounds(ForceInit);
 	}
