Chaos Cloth Asset - Fixed thumbnail rendering broken with CL 42904341.

#rnx

diff --git a/Engine/Plugins/ChaosClothAsset/Source/ChaosClothAssetEngine/Private/ChaosClothAsset/ClothComponent.cpp b/Engine/Plugins/ChaosClothAsset/Source/ChaosClothAssetEngine/Private/ChaosClothAsset/ClothComponent.cpp
--- a/Engine/Plugins/ChaosClothAsset/Source/ChaosClothAssetEngine/Private/ChaosClothAsset/ClothComponent.cpp
+++ b/Engine/Plugins/ChaosClothAsset/Source/ChaosClothAssetEngine/Private/ChaosClothAsset/ClothComponent.cpp
@@ -762,6 +762,7 @@ void UChaosClothComponent::UpdateVisibility()
 	{
 		bHasValidRenderDataForVisibility = false;
 	}
+	SetVisibility(bHasValidRenderDataForVisibility);
 }
 
 UChaosClothAssetInteractor* UChaosClothComponent::GetClothOutfitInteractor(int32 ModelIndex, const FName ClothSimulationModelName)
