Chaos Cloth Asset:
- Added SKM Clothing Asset Simulation and Clothing Interactors.
- Improved the component update mechanism when a Cloth Asset has changed.

#rnx

diff --git a/Engine/Plugins/ChaosClothAsset/Source/ChaosClothAssetEngine/Private/ChaosClothAsset/ClothAsset.cpp b/Engine/Plugins/ChaosClothAsset/Source/ChaosClothAssetEngine/Private/ChaosClothAsset/ClothAsset.cpp
--- a/Engine/Plugins/ChaosClothAsset/Source/ChaosClothAssetEngine/Private/ChaosClothAsset/ClothAsset.cpp
+++ b/Engine/Plugins/ChaosClothAsset/Source/ChaosClothAssetEngine/Private/ChaosClothAsset/ClothAsset.cpp
@@ -211,7 +211,7 @@ PRAGMA_ENABLE_DEPRECATION_WARNINGS
 	TSharedRef<FManagedArrayCollection> ClothCollection = MakeShared<FManagedArrayCollection>();
 	UE::Chaos::ClothAsset::FCollectionClothFacade ClothFacade(ClothCollection);
 	ClothFacade.DefineSchema();
-	GetClothCollections().Emplace(MoveTemp(ClothCollection));
+	GetClothCollectionsInternal().Emplace(MoveTemp(ClothCollection));
 }
 
 UChaosClothAsset::UChaosClothAsset(FVTableHelper& Helper)
@@ -239,14 +239,14 @@ void UChaosClothAsset::Serialize(FArchive& Ar)
 	if (bCooked && Ar.IsSaving())
 	{
 		TArray<TSharedRef<const FManagedArrayCollection>> OutputCollections = 
-			Private::TrimOnCook(GetPathName(), GetClothCollections());
+			Private::TrimOnCook(GetPathName(), GetClothCollectionsInternal());
 		Chaos::FChaosArchive ChaosArchive(Ar);
 		Private::Serialize(ChaosArchive, OutputCollections);
 	}
 	else
 	{
 		Chaos::FChaosArchive ChaosArchive(Ar);
-		Private::Serialize(ChaosArchive, GetClothCollections());
+		Private::Serialize(ChaosArchive, GetClothCollectionsInternal());
 	}
 
 #if WITH_EDITOR
@@ -264,7 +264,7 @@ void UChaosClothAsset::Serialize(FArchive& Ar)
 	if (Ar.CustomVer(FFortniteMainBranchObjectVersion::GUID) < FFortniteMainBranchObjectVersion::ClothAssetSkinweightsValidation)
 	{
 		// Fix the skeleton mesh binding, which can cause crashes in the render code, or make the sim mesh disappear when missing
-		for (TSharedRef<const FManagedArrayCollection>& ClothCollection : GetClothCollections())
+		for (TSharedRef<const FManagedArrayCollection>& ClothCollection : GetClothCollectionsInternal())
 		{
 			const FCollectionClothConstFacade ClothConstFacade(ClothCollection);
 			if (ClothConstFacade.IsValid())
@@ -393,21 +393,21 @@ void UChaosClothAsset::BeginPostLoadInternal(FSkinnedAssetPostLoadContext& Conte
 	FSkinnedAssetAsyncBuildScope AsyncBuildScope(this);
 
 	// Make sure that there is at least one valid collection
-	if (GetClothCollections().IsEmpty())
+	if (GetClothCollectionsInternal().IsEmpty())
 	{
 		UE_LOG(LogChaosClothAsset, Warning, TEXT("Invalid Cloth Collection (no LODs) found while loading Cloth Asset %s."), *GetFullName());
 		TSharedRef<FManagedArrayCollection> ClothCollection = MakeShared<FManagedArrayCollection>();
 		FCollectionClothFacade ClothFacade(ClothCollection);
 		ClothFacade.DefineSchema();
-		GetClothCollections().Emplace(MoveTemp(ClothCollection));
+		GetClothCollectionsInternal().Emplace(MoveTemp(ClothCollection));
 	}
 
 	// Check that all LODs have the cloth schema
-	const int32 NumLods = GetClothCollections().Num();
+	const int32 NumLods = GetClothCollectionsInternal().Num();
 	check(NumLods >= 1);  // The default LOD 0 should be present now if it ever was missing
 	for (int32 LODIndex = 0; LODIndex < NumLods; ++LODIndex)
 	{
-		const TSharedRef<const FManagedArrayCollection>& ClothCollection = GetClothCollections()[LODIndex];
+		const TSharedRef<const FManagedArrayCollection>& ClothCollection = GetClothCollectionsInternal()[LODIndex];
 
 		const FCollectionClothConstFacade ClothConstFacade(ClothCollection);
 		if (!ClothConstFacade.IsValid())
@@ -416,7 +416,7 @@ void UChaosClothAsset::BeginPostLoadInternal(FSkinnedAssetPostLoadContext& Conte
 			TSharedRef<FManagedArrayCollection> NewClothCollection = MakeShared<FManagedArrayCollection>();
 			FCollectionClothFacade NewClothFacade = FCollectionClothFacade(NewClothCollection);
 			NewClothFacade.DefineSchema();
-			GetClothCollections()[LODIndex] = MoveTemp(NewClothCollection);
+			GetClothCollectionsInternal()[LODIndex] = MoveTemp(NewClothCollection);
 		}
 	}
 
@@ -507,6 +507,17 @@ void UChaosClothAsset::CalculateBounds()
 }
 #endif
 
+void UChaosClothAsset::SetClothCollections(TArray<TSharedRef<const FManagedArrayCollection>>&& InClothCollections)
+{
+	WaitUntilAsyncPropertyReleased(EClothAssetAsyncProperties::ClothCollection, ESkinnedAssetAsyncPropertyLockType::WriteOnly);
+
+PRAGMA_DISABLE_DEPRECATION_WARNINGS
+	ClothCollections = MoveTemp(InClothCollections);
+PRAGMA_ENABLE_DEPRECATION_WARNINGS
+
+	OnPropertyChanged();
+}
+
 void UChaosClothAsset::Build(
 	const TArray<TSharedRef<const FManagedArrayCollection>>& InClothCollections,
 	TArray<FChaosClothAssetLodTransitionDataCache>* InOutTransitionCache,
@@ -516,7 +527,7 @@ void UChaosClothAsset::Build(
 	using namespace UE::Chaos::ClothAsset;
 
 	// Reset the asset's collection
-	TArray<TSharedRef<const FManagedArrayCollection>>& OutClothCollections = GetClothCollections();
+	TArray<TSharedRef<const FManagedArrayCollection>>& OutClothCollections = GetClothCollectionsInternal();
 	OutClothCollections.Reset(InClothCollections.Num());
 
 	// Reset the asset's material list
@@ -673,7 +684,7 @@ void UChaosClothAsset::Build(TArray<FChaosClothAssetLodTransitionDataCache>* InO
 	CalculateInvRefMatrices();
 
 	// Add LODs to the render data
-	const int32 NumLods = FMath::Max(GetClothCollections().Num(), 1);  // The render data will always look for at least one default LOD 0
+	const int32 NumLods = FMath::Max(GetClothCollectionsInternal().Num(), 1);  // The render data will always look for at least one default LOD 0
 
 	// Rebuild LOD Infos
 	LODInfo.Reset(NumLods);
@@ -877,8 +888,8 @@ void UChaosClothAsset::UpdateSkeletonFromCollection(bool /*bRebuildModels*/)
 {
 	using namespace UE::Chaos::ClothAsset;
 
-	check(GetClothCollections().Num());
-	FCollectionClothConstFacade ClothFacade(GetClothCollections()[0]);
+	check(GetClothCollectionsInternal().Num());
+	FCollectionClothConstFacade ClothFacade(GetClothCollectionsInternal()[0]);
 	check(ClothFacade.IsValid());
 
 	const FSoftObjectPath& SkeletalMeshPathName = ClothFacade.GetSkeletalMeshSoftObjectPathName();
@@ -891,7 +902,7 @@ void UChaosClothAsset::UpdateSkeletonFromCollection(bool /*bRebuildModels*/)
 void UChaosClothAsset::CopySimMeshToRenderMesh(UMaterialInterface* Material)
 {
 	using namespace UE::Chaos::ClothAsset;
-	check(GetClothCollections().Num());
+	check(GetClothCollectionsInternal().Num());
 
 	// Add a default material if none is specified
 	const FSoftObjectPath RenderMaterialPathName = FSoftObjectPath(Material ?
@@ -899,7 +910,7 @@ void UChaosClothAsset::CopySimMeshToRenderMesh(UMaterialInterface* Material)
 		FString(TEXT("/Engine/EditorMaterials/Cloth/CameraLitDoubleSided.CameraLitDoubleSided")));
 
 	bool bAnyLodHasRenderMesh = false;
-	for (TSharedRef<const FManagedArrayCollection>& ClothCollection : GetClothCollections())
+	for (TSharedRef<const FManagedArrayCollection>& ClothCollection : GetClothCollectionsInternal())
 	{
 		TSharedRef<FManagedArrayCollection> NewClothCollection = MakeShared<FManagedArrayCollection>(*ClothCollection);
 		constexpr bool bSingleRenderPattern = true;
diff --git a/Engine/Plugins/ChaosClothAsset/Source/ChaosClothAssetEngine/Private/ChaosClothAsset/ClothAssetBase.cpp b/Engine/Plugins/ChaosClothAsset/Source/ChaosClothAssetEngine/Private/ChaosClothAsset/ClothAssetBase.cpp
--- a/Engine/Plugins/ChaosClothAsset/Source/ChaosClothAssetEngine/Private/ChaosClothAsset/ClothAssetBase.cpp
+++ b/Engine/Plugins/ChaosClothAsset/Source/ChaosClothAssetEngine/Private/ChaosClothAsset/ClothAssetBase.cpp
@@ -19,6 +19,9 @@
 #endif
 #include "Rendering/SkeletalMeshRenderData.h"
 #include "UObject/FortniteMainBranchObjectVersion.h"
+#include "ClothingSimulationFactory.h"
+#include "ClothingSimulationInstance.h"
+#include "ClothingSimulationInteractor.h"
 
 #include UE_INLINE_GENERATED_CPP_BY_NAME(ClothAssetBase)
 namespace UE::Chaos::ClothAsset::Private
@@ -599,7 +602,52 @@ bool UChaosClothAssetBase::ExportToSkeletalMesh(USkeletalMesh& SkeletalMesh) con
 }
 #endif  // #if WITH_EDITOR
 
-void UChaosClothAssetBase::ReregisterComponents()
+void UChaosClothAssetBase::OnPropertyChanged() const
+{
+	// Update cloth components properties
+	for (TObjectIterator<UChaosClothComponent> ObjectIterator; ObjectIterator; ++ObjectIterator)
+	{
+		if (UChaosClothComponent* const Component = *ObjectIterator)
+		{
+			if (Component->GetAsset() == this)
+			{
+				Component->UpdateConfigProperties();
+			}
+		}
+	}
+	// Update skeletal mesh components properties
+	for (TObjectIterator<USkeletalMeshComponent> ObjectIterator; ObjectIterator; ++ObjectIterator)
+	{
+		if (USkeletalMeshComponent* const Component = *ObjectIterator)
+		{
+			if (Component->GetSkeletalMeshAsset())
+			{
+				for (const TObjectPtr<UClothingAssetBase>& ClothingAsset : Component->GetSkeletalMeshAsset()->GetMeshClothingAssets())
+				{
+					if (const UChaosClothAssetSKMClothingAsset* const ClothAssetSKMClothingAsset = Cast<UChaosClothAssetSKMClothingAsset>(ClothingAsset))
+					{
+						if (ClothAssetSKMClothingAsset->GetAsset() == this)
+						{
+							for (const FClothingSimulationInstance& ClothingSimulationInstance : Component->GetClothingSimulationInstances())
+							{
+								if (ClothingSimulationInstance.GetClothingSimulationFactory() &&
+									ClothingSimulationInstance.GetClothingSimulationFactory()->SupportsAsset(ClothAssetSKMClothingAsset))
+								{
+									if (UClothingSimulationInteractor* const ClothingSimulationInteractor = ClothingSimulationInstance.GetClothingSimulationInteractor())
+									{
+										ClothingSimulationInteractor->ClothConfigUpdated();
+									}
+								}
+							}
+						}
+					}
+				}
+			}
+		}
+	}
+}
+
+void UChaosClothAssetBase::ReregisterComponents() const
 {
 	// Reregister cloth components
 	for (TObjectIterator<UChaosClothComponent> ObjectIterator; ObjectIterator; ++ObjectIterator)
diff --git a/Engine/Plugins/ChaosClothAsset/Source/ChaosClothAssetEngine/Private/ChaosClothAsset/ClothAssetInteractor.cpp b/Engine/Plugins/ChaosClothAsset/Source/ChaosClothAssetEngine/Private/ChaosClothAsset/ClothAssetInteractor.cpp
--- a/Engine/Plugins/ChaosClothAsset/Source/ChaosClothAssetEngine/Private/ChaosClothAsset/ClothAssetInteractor.cpp
+++ b/Engine/Plugins/ChaosClothAsset/Source/ChaosClothAssetEngine/Private/ChaosClothAsset/ClothAssetInteractor.cpp
@@ -141,7 +141,7 @@ namespace UE::Chaos::ClothAsset::Private
 
 
 	template<typename T>
-	static T GetValueWithAlias(TSharedPtr<::Chaos::Softs::FCollectionPropertyFacade>& PropertyFacade, const FName& PropertyName, const T& DefaultValue, const TFunctionRef<T(int32 KeyIndex)>& GetValue)
+	static T GetValueWithAlias(TSharedPtr<::Chaos::Softs::FCollectionPropertyFacade>& PropertyFacade, const FName PropertyName, const T& DefaultValue, const TFunctionRef<T(int32 KeyIndex)>& GetValue)
 	{
 		check(PropertyFacade);
 		{
@@ -167,7 +167,7 @@ namespace UE::Chaos::ClothAsset::Private
 		return DefaultValue;
 	}
 
-	void SetValueWithAlias(TSharedPtr<::Chaos::Softs::FCollectionPropertyFacade>& PropertyFacade, const FName& PropertyName, const TFunctionRef<void(const FName&)>& SetValue)
+	void SetValueWithAlias(TSharedPtr<::Chaos::Softs::FCollectionPropertyFacade>& PropertyFacade, const FName PropertyName, const TFunctionRef<void(const FName&)>& SetValue)
 	{
 		check(PropertyFacade);
 		if (const TArray<FName>* FoundAliases = Aliases.Find(PropertyName))
@@ -269,7 +269,7 @@ TArray<FString> UChaosClothAssetInteractor::GetAllProperties(int32 LODIndex) con
 	return Result;
 }
 
-float UChaosClothAssetInteractor::GetFloatPropertyValue(const FName& PropertyName, int32 LODIndex, float DefaultValue) const
+float UChaosClothAssetInteractor::GetFloatPropertyValue(const FName PropertyName, int32 LODIndex, float DefaultValue) const
 {
 	using namespace UE::Chaos::ClothAsset::Private;
 	if (CollectionPropertyFacades.IsValidIndex(LODIndex))
@@ -286,7 +286,7 @@ float UChaosClothAssetInteractor::GetFloatPropertyValue(const FName& PropertyNam
 	return DefaultValue;
 }
 
-float UChaosClothAssetInteractor::GetLowFloatPropertyValue(const FName& PropertyName, int32 LODIndex, float DefaultValue) const
+float UChaosClothAssetInteractor::GetLowFloatPropertyValue(const FName PropertyName, int32 LODIndex, float DefaultValue) const
 {
 	using namespace UE::Chaos::ClothAsset::Private;
 	if (CollectionPropertyFacades.IsValidIndex(LODIndex))
@@ -303,7 +303,7 @@ float UChaosClothAssetInteractor::GetLowFloatPropertyValue(const FName& Property
 	return DefaultValue;
 }
 
-float UChaosClothAssetInteractor::GetHighFloatPropertyValue(const FName& PropertyName, int32 LODIndex, float DefaultValue) const
+float UChaosClothAssetInteractor::GetHighFloatPropertyValue(const FName PropertyName, int32 LODIndex, float DefaultValue) const
 {
 	using namespace UE::Chaos::ClothAsset::Private;
 	if (CollectionPropertyFacades.IsValidIndex(LODIndex))
@@ -320,7 +320,7 @@ float UChaosClothAssetInteractor::GetHighFloatPropertyValue(const FName& Propert
 	return DefaultValue;
 }
 
-FVector2D UChaosClothAssetInteractor::GetWeightedFloatPropertyValue(const FName& PropertyName, int32 LODIndex, FVector2D DefaultValue) const
+FVector2D UChaosClothAssetInteractor::GetWeightedFloatPropertyValue(const FName PropertyName, int32 LODIndex, FVector2D DefaultValue) const
 {
 	using namespace UE::Chaos::ClothAsset::Private;
 	if (CollectionPropertyFacades.IsValidIndex(LODIndex))
@@ -337,7 +337,7 @@ FVector2D UChaosClothAssetInteractor::GetWeightedFloatPropertyValue(const FName&
 	return DefaultValue;
 }
 
-int32 UChaosClothAssetInteractor::GetIntPropertyValue(const FName& PropertyName, int32 LODIndex, int32 DefaultValue) const
+int32 UChaosClothAssetInteractor::GetIntPropertyValue(const FName PropertyName, int32 LODIndex, int32 DefaultValue) const
 {
 	using namespace UE::Chaos::ClothAsset::Private;
 	if (CollectionPropertyFacades.IsValidIndex(LODIndex))
@@ -354,7 +354,7 @@ int32 UChaosClothAssetInteractor::GetIntPropertyValue(const FName& PropertyName,
 	return DefaultValue;
 }
 
-FVector UChaosClothAssetInteractor::GetVectorPropertyValue(const FName& PropertyName, int32 LODIndex, FVector DefaultValue) const
+FVector UChaosClothAssetInteractor::GetVectorPropertyValue(const FName PropertyName, int32 LODIndex, FVector DefaultValue) const
 {
 	using namespace UE::Chaos::ClothAsset::Private;
 	if (CollectionPropertyFacades.IsValidIndex(LODIndex))
@@ -371,7 +371,7 @@ FVector UChaosClothAssetInteractor::GetVectorPropertyValue(const FName& Property
 	return DefaultValue;
 }
 
-FString UChaosClothAssetInteractor::GetStringPropertyValue(const FName& PropertyName, int32 LODIndex, const FString& DefaultValue) const
+FString UChaosClothAssetInteractor::GetStringPropertyValue(const FName PropertyName, int32 LODIndex, const FString& DefaultValue) const
 {
 	using namespace UE::Chaos::ClothAsset::Private;
 	if (CollectionPropertyFacades.IsValidIndex(LODIndex))
@@ -388,7 +388,7 @@ FString UChaosClothAssetInteractor::GetStringPropertyValue(const FName& Property
 	return DefaultValue;
 }
 
-void UChaosClothAssetInteractor::SetFloatPropertyValue(const FName& PropertyName, int32 LODIndex, float Value)
+void UChaosClothAssetInteractor::SetFloatPropertyValue(const FName PropertyName, int32 LODIndex, float Value)
 {
 	using namespace UE::Chaos::ClothAsset::Private;
 	if (LODIndex == INDEX_NONE)
@@ -418,7 +418,7 @@ void UChaosClothAssetInteractor::SetFloatPropertyValue(const FName& PropertyName
 	}
 }
 
-void UChaosClothAssetInteractor::SetLowFloatPropertyValue(const FName& PropertyName, int32 LODIndex, float Value)
+void UChaosClothAssetInteractor::SetLowFloatPropertyValue(const FName PropertyName, int32 LODIndex, float Value)
 {
 	using namespace UE::Chaos::ClothAsset::Private;
 	if (LODIndex == INDEX_NONE)
@@ -448,7 +448,7 @@ void UChaosClothAssetInteractor::SetLowFloatPropertyValue(const FName& PropertyN
 	}
 }
 
-void UChaosClothAssetInteractor::SetHighFloatPropertyValue(const FName& PropertyName, int32 LODIndex, float Value)
+void UChaosClothAssetInteractor::SetHighFloatPropertyValue(const FName PropertyName, int32 LODIndex, float Value)
 {
 	using namespace UE::Chaos::ClothAsset::Private;
 	if (LODIndex == INDEX_NONE)
@@ -478,7 +478,7 @@ void UChaosClothAssetInteractor::SetHighFloatPropertyValue(const FName& Property
 	}
 }
 
-void UChaosClothAssetInteractor::SetWeightedFloatPropertyValue(const FName& PropertyName, int32 LODIndex, FVector2D Value)
+void UChaosClothAssetInteractor::SetWeightedFloatPropertyValue(const FName PropertyName, int32 LODIndex, FVector2D Value)
 {
 	using namespace UE::Chaos::ClothAsset::Private;
 	if (LODIndex == INDEX_NONE)
@@ -508,7 +508,7 @@ void UChaosClothAssetInteractor::SetWeightedFloatPropertyValue(const FName& Prop
 	}
 }
 
-void UChaosClothAssetInteractor::SetIntPropertyValue(const FName& PropertyName, int32 LODIndex, int32 Value)
+void UChaosClothAssetInteractor::SetIntPropertyValue(const FName PropertyName, int32 LODIndex, int32 Value)
 {
 	using namespace UE::Chaos::ClothAsset::Private;
 	if (LODIndex == INDEX_NONE)
@@ -538,7 +538,7 @@ void UChaosClothAssetInteractor::SetIntPropertyValue(const FName& PropertyName,
 	}
 }
 
-void UChaosClothAssetInteractor::SetVectorPropertyValue(const FName& PropertyName, int32 LODIndex, FVector Value)
+void UChaosClothAssetInteractor::SetVectorPropertyValue(const FName PropertyName, int32 LODIndex, FVector Value)
 {
 	using namespace UE::Chaos::ClothAsset::Private;
 	if (LODIndex == INDEX_NONE)
@@ -568,7 +568,7 @@ void UChaosClothAssetInteractor::SetVectorPropertyValue(const FName& PropertyNam
 	}
 }
 
-void UChaosClothAssetInteractor::SetStringPropertyValue(const FName& PropertyName, int32 LODIndex, const FString& Value)
+void UChaosClothAssetInteractor::SetStringPropertyValue(const FName PropertyName, int32 LODIndex, const FString& Value)
 {
 	using namespace UE::Chaos::ClothAsset::Private;
 	if (LODIndex == INDEX_NONE)
diff --git a/Engine/Plugins/ChaosClothAsset/Source/ChaosClothAssetEngine/Private/ChaosClothAsset/ClothAssetSKMClothingSimulation.cpp b/Engine/Plugins/ChaosClothAsset/Source/ChaosClothAssetEngine/Private/ChaosClothAsset/ClothAssetSKMClothingSimulation.cpp
--- a/Engine/Plugins/ChaosClothAsset/Source/ChaosClothAssetEngine/Private/ChaosClothAsset/ClothAssetSKMClothingSimulation.cpp
+++ b/Engine/Plugins/ChaosClothAsset/Source/ChaosClothAssetEngine/Private/ChaosClothAsset/ClothAssetSKMClothingSimulation.cpp
@@ -1,13 +1,18 @@
 // Copyright Epic Games, Inc. All Rights Reserved.
 
 #include "ChaosClothAsset/ClothAssetSKMClothingSimulation.h"
+#include "ChaosClothAsset/ClothAssetSKMClothingSimulationInteractor.h"
+#include "Chaos/CollectionPropertyFacade.h"
+#include "ChaosCloth/ChaosClothingSimulationSolver.h"
 #include "ChaosClothAsset/ClothAssetBase.h"
 #include "ChaosClothAsset/ClothAssetSKMClothingAsset.h"
 #include "ChaosClothAsset/ClothComponent.h"  // For FChaosClothSimulationProperties  TODO: Move FChaosClothSimulationProperties to its own header file?
 #include "ChaosClothAsset/ClothSimulationProxy.h"
 #include "ChaosClothAsset/CollisionSources.h"
 #include "Components/SkeletalMeshComponent.h"
+#include "Engine/SkeletalMesh.h"
 #include "HAL/IConsoleManager.h"
+#include "Rendering/SkeletalMeshRenderData.h"
 #include "ClothCollisionData.h"
 
 #include UE_INLINE_GENERATED_CPP_BY_NAME(ClothAssetSKMClothingSimulation)
@@ -17,6 +22,43 @@ namespace UE::Chaos::ClothAsset
 	FSKMClothingSimulation::FSKMClothingSimulation() = default;
 	FSKMClothingSimulation::~FSKMClothingSimulation() = default;
 
+	void FSKMClothingSimulation::ResetConfigProperties()
+	{
+		check(IsInGameThread());
+
+		// Update collection, facade and outfit interactor
+		for (FAssetClothSimulationProperties& Properties : AssetClothSimulationProperties)
+		{
+			if (const UChaosClothAssetSKMClothingAsset* const ClothAssetSKMClothingAsset = Properties.Key)
+			{
+				if (const int32 ClothSimulationModelIndex = ClothAssetSKMClothingAsset->GetClothSimulationModelIndex(); ClothSimulationModelIndex != INDEX_NONE)
+				{
+					if (const UChaosClothAssetBase* const Asset = ClothAssetSKMClothingAsset->GetAsset(); ensure(Asset))  // ClothSimulationModelIndex has already been checked, so the ClothAsset can't be null
+					{
+						Properties.Value.Initialize(Asset->GetCollections(ClothSimulationModelIndex));
+					}
+				}
+			}
+		}
+	}
+
+	void FSKMClothingSimulation::RecreateClothSimulationProxy()
+	{
+		ClothSimulationProxy = MakeShared<FClothSimulationProxy>(*this);
+		if (ClothSimulationProxy.IsValid())
+		{
+			// Re-add all assets to the proxy
+			ClothSimulationProxy->PostConstructor();
+		}
+	}
+
+	UChaosClothAssetInteractor* FSKMClothingSimulation::GetPropertyInteractor(const UChaosClothAssetBase* Asset, int32 ModelIndex) const
+	{
+		const int32 SimulationGroupId = GetSimulationGroupId(Asset, ModelIndex);
+		checkf(SimulationGroupId != INDEX_NONE, TEXT("Invalid arguments"));
+		return AssetClothSimulationProperties[SimulationGroupId].Value.ClothOutfitInteractor;
+	}
+
 	void FSKMClothingSimulation::Initialize()
 	{
 		Shutdown();
@@ -73,7 +115,7 @@ namespace UE::Chaos::ClothAsset
 							Properties.Key->GetClothSimulationModelIndex() == ClothSimulationModelIndex;
 					};
 
-				const int32 FirstSimDataIndex =  AssetClothSimulationProperties.IndexOfByPredicate(ClothAssetPredicate);
+				const int32 FirstSimDataIndex = AssetClothSimulationProperties.IndexOfByPredicate(ClothAssetPredicate);
 				if (FirstSimDataIndex == INDEX_NONE)
 				{
 					FChaosClothSimulationProperties& ClothSimulationProperties = AssetClothSimulationProperties[SimDataIndex].Value;
@@ -179,6 +221,39 @@ namespace UE::Chaos::ClothAsset
 		OutCollisions.Reset();
 	}
 
+	void FSKMClothingSimulation::SetNumIterations(int32 NumIterations)
+	{
+		for (const TSharedPtr<::Chaos::Softs::FCollectionPropertyFacade>& SolverPropertyFacade : GetSolverPropertyFacades())
+		{
+			if (SolverPropertyFacade)
+			{
+				SolverPropertyFacade->SetValue(::Chaos::FClothingSimulationSolver::NumIterationsName, NumIterations);
+			}
+		}
+	}
+
+	void FSKMClothingSimulation::SetMaxNumIterations(int32 MaxNumIterations)
+	{
+		for (const TSharedPtr<::Chaos::Softs::FCollectionPropertyFacade>& SolverPropertyFacade : GetSolverPropertyFacades())
+		{
+			if (SolverPropertyFacade)
+			{
+				SolverPropertyFacade->SetValue(::Chaos::FClothingSimulationSolver::MaxNumIterationsName, MaxNumIterations);
+			}
+		}
+	}
+
+	void FSKMClothingSimulation::SetNumSubsteps(int32 NumSubsteps)
+	{
+		for (const TSharedPtr<::Chaos::Softs::FCollectionPropertyFacade>& SolverPropertyFacade : GetSolverPropertyFacades())
+		{
+			if (SolverPropertyFacade)
+			{
+				SolverPropertyFacade->SetValue(::Chaos::FClothingSimulationSolver::NumSubstepsName, NumSubsteps);
+			}
+		}
+	}
+
 	const USkinnedMeshComponent& FSKMClothingSimulation::GetOwnerComponent() const
 	{
 		checkf(OwnerComponent, TEXT("No actors, owner component not set. Call CreateActor first."));
@@ -209,7 +284,7 @@ namespace UE::Chaos::ClothAsset
 			const FAssetClothSimulationProperties& Properties = AssetClothSimulationProperties[Index];
 			if (const UChaosClothAssetSKMClothingAsset* const ClothAssetSKMClothingAsset = Properties.Key)
 			{
-				if (ClothAssetSKMClothingAsset->GetAsset() == Asset && 
+				if (ClothAssetSKMClothingAsset->GetAsset() == Asset &&
 					ClothAssetSKMClothingAsset->GetClothSimulationModelIndex() == ModelIndex)
 				{
 					return Index;
@@ -255,18 +330,41 @@ namespace UE::Chaos::ClothAsset
 
 	const TArray<TSharedPtr<const FManagedArrayCollection>>& FSKMClothingSimulation::GetSolverPropertyCollections() const
 	{
-		for (const FAssetClothSimulationProperties& Properties : AssetClothSimulationProperties)
+		if (const FChaosClothSimulationProperties* const SolverProperties = FSKMClothingSimulation::GetSolverProperties())
 		{
-			if (const UChaosClothAssetSKMClothingAsset* const ClothAssetSKMClothingAsset = Properties.Key)
+			return SolverProperties->PropertyCollections;
+		}
+		static const TArray<TSharedPtr<const FManagedArrayCollection>> EmptyManagedArrayCollectionArray;
+		return EmptyManagedArrayCollectionArray;
+	}
+
+	bool FSKMClothingSimulation::HasAnyRenderClothMappingData(int32 LodIndex) const
+	{
+		if (const USkeletalMesh* const SkeletalMesh = OwnerComponent->GetSkeletalMeshAsset())
+		{
+			if (const FSkeletalMeshRenderData* const SkeletalMeshRenderData = SkeletalMesh->GetResourceForRendering())
 			{
-				if (ClothAssetSKMClothingAsset->GetAsset())
+				if (SkeletalMeshRenderData->LODRenderData.IsValidIndex(LodIndex))
 				{
-					return Properties.Value.PropertyCollections;
+					for (const FSkelMeshRenderSection& SkelMeshRenderSection : SkeletalMeshRenderData->LODRenderData[LodIndex].RenderSections)
+					{
+						if (SkelMeshRenderSection.HasClothingData() &&
+							AssetClothSimulationProperties.IsValidIndex(SkelMeshRenderSection.CorrespondClothAssetIndex))
+						{
+							if (const UChaosClothAssetSKMClothingAsset* const ClothAssetSKMClothingAsset = AssetClothSimulationProperties[SkelMeshRenderSection.CorrespondClothAssetIndex].Key)
+							{
+								if (SkelMeshRenderSection.ClothingData.AssetGuid == ClothAssetSKMClothingAsset->GetAssetGuid() &&
+									SkelMeshRenderSection.ClothingData.AssetLodIndex == LodIndex)
+								{
+									return true;
+								}
+							}
+						}
+					}
 				}
 			}
 		}
-		static const TArray<TSharedPtr<const FManagedArrayCollection>> EmptyManagedArrayCollectionArray;
-		return EmptyManagedArrayCollectionArray;
+		return false;
 	}
 
 	void FSKMClothingSimulation::AddReferencedObjects(FReferenceCollector& Collector)
@@ -284,6 +382,31 @@ namespace UE::Chaos::ClothAsset
 	{
 		return TEXT("UE::Chaos::ClothAsset::FSKMClothingSimulation");
 	}
+
+	const FChaosClothSimulationProperties* FSKMClothingSimulation::GetSolverProperties() const
+	{
+		for (const FAssetClothSimulationProperties& Properties : AssetClothSimulationProperties)
+		{
+			if (const UChaosClothAssetSKMClothingAsset* const ClothAssetSKMClothingAsset = Properties.Key)
+			{
+				if (ClothAssetSKMClothingAsset->GetAsset())
+				{
+					return &Properties.Value;  // Return the properties for the first valid asset for now, there might be way to improve this in the future
+				}
+			}
+		}
+		return nullptr;
+	}
+
+	const TArray<TSharedPtr<::Chaos::Softs::FCollectionPropertyFacade>>& FSKMClothingSimulation::GetSolverPropertyFacades() const
+	{
+		if (const FChaosClothSimulationProperties* const SolverProperties = FSKMClothingSimulation::GetSolverProperties())
+		{
+			return SolverProperties->CollectionPropertyFacades;
+		}
+		static const TArray<TSharedPtr<::Chaos::Softs::FCollectionPropertyFacade>> EmptyManagedArrayCollectionArray;
+		return EmptyManagedArrayCollectionArray;
+	}
 }  // namespace UE::Chaos::ClothAsset
 
 UChaosClothAssetSKMClothingSimulationFactory::UChaosClothAssetSKMClothingSimulationFactory() = default;
@@ -293,3 +416,8 @@ bool UChaosClothAssetSKMClothingSimulationFactory::SupportsAsset(const UClothing
 {
 	return Cast<UChaosClothAssetSKMClothingAsset>(Asset) != nullptr;
 }
+
+UClothingSimulationInteractor* UChaosClothAssetSKMClothingSimulationFactory::CreateInteractor()
+{
+	return NewObject<UChaosClothAssetSKMClothingSimulationInteractor>(GetTransientPackage());
+}
diff --git a/Engine/Plugins/ChaosClothAsset/Source/ChaosClothAssetEngine/Private/ChaosClothAsset/ClothAssetSKMClothingSimulation.h b/Engine/Plugins/ChaosClothAsset/Source/ChaosClothAssetEngine/Private/ChaosClothAsset/ClothAssetSKMClothingSimulation.h
--- a/Engine/Plugins/ChaosClothAsset/Source/ChaosClothAssetEngine/Private/ChaosClothAsset/ClothAssetSKMClothingSimulation.h
+++ b/Engine/Plugins/ChaosClothAsset/Source/ChaosClothAssetEngine/Private/ChaosClothAsset/ClothAssetSKMClothingSimulation.h
@@ -3,19 +3,25 @@
 #pragma once
 
 #include "ChaosClothAsset/ClothComponentAdapter.h"
+#include "ChaosClothAsset/ClothSimulationProxy.h"
 #include "ClothingSimulation.h"
 #include "ClothingSimulationFactory.h"
 #include "UObject/GCObject.h"
 #include "ClothAssetSKMClothingSimulation.generated.h"
 
 class USkinnedMeshComponent;
+class UChaosClothAssetInteractor;
 class UChaosClothAssetSKMClothingAsset;
 struct FChaosClothSimulationProperties;
 
+namespace Chaos::Softs
+{
+	class FCollectionPropertyFacade;
+}
+
 namespace UE::Chaos::ClothAsset
 {
 	using FSKMClothingSimulationContext = FClothingSimulationContextCommon;
-	class FClothSimulationProxy;
 
 	class FSKMClothingSimulation final
 		: public FClothingSimulationCommon
@@ -26,6 +32,15 @@ namespace UE::Chaos::ClothAsset
 		FSKMClothingSimulation();
 		virtual ~FSKMClothingSimulation() override;
 
+		/** Reset all cloth simulation config properties to the values stored in the original cloth asset. Game thread only. */
+		void ResetConfigProperties();
+
+		/** Hard reset the cloth simulation by recreating the proxy. */
+		void RecreateClothSimulationProxy();
+
+		/** Return a property interactor for the specified asset and model. */
+		UChaosClothAssetInteractor* GetPropertyInteractor(const UChaosClothAssetBase* Asset, int32 ModelIndex) const;
+
 	private:
 		//~ Begin IClothingSimulation Interface
 		virtual void Initialize() override;
@@ -53,6 +68,16 @@ namespace UE::Chaos::ClothAsset
 		virtual void AddExternalCollisions(const FClothCollisionData& Data) override;
 		virtual void ClearExternalCollisions() override;
 		virtual void GetCollisions(FClothCollisionData& OutCollisions, bool bIncludeExternal = true) const override;
+		virtual void SetNumIterations(int32 NumIterations) override;
+		virtual void SetMaxNumIterations(int32 MaxNumIterations) override;
+		virtual void SetNumSubsteps(int32 NumSubsteps) override;
+		virtual int32 GetNumCloths() const override { return ClothSimulationProxy.IsValid() ? ClothSimulationProxy->GetNumCloths() : 0; }
+		virtual int32 GetNumKinematicParticles() const override { return ClothSimulationProxy.IsValid() ? ClothSimulationProxy->GetNumKinematicParticles() : 0; }
+		virtual int32 GetNumDynamicParticles() const override { return ClothSimulationProxy.IsValid() ? ClothSimulationProxy->GetNumDynamicParticles() : 0; }
+		virtual int32 GetNumIterations() const override { return ClothSimulationProxy.IsValid() ? ClothSimulationProxy->GetNumIterations() : 0; }
+		virtual int32 GetNumSubsteps() const override { return ClothSimulationProxy.IsValid() ? ClothSimulationProxy->GetNumSubsteps() : 0; }
+		virtual float GetSimulationTime() const override { return ClothSimulationProxy.IsValid() ? ClothSimulationProxy->GetSimulationTime() : 0.f; }
+		virtual bool IsTeleported() const override { return ClothSimulationProxy.IsValid() ? ClothSimulationProxy->IsTeleported() : false; }
 		//~ End IClothingSimulation Interface
 
 		//~ Begin IClothComponentAdapter Interface
@@ -80,6 +105,7 @@ namespace UE::Chaos::ClothAsset
 		virtual float GetClothGeometryScale() const override;
 		virtual const TArray<TSharedPtr<const FManagedArrayCollection>>& GetPropertyCollections(const UChaosClothAssetBase* Asset, int32 ModelIndex) const override;
 		virtual const TArray<TSharedPtr<const FManagedArrayCollection>>& GetSolverPropertyCollections() const override;
+		virtual bool HasAnyRenderClothMappingData(int32 LodIndex) const override;
 		//~ End IClothComponentAdapter Interface
 
 		//~ Begin FGCObject Interface
@@ -87,6 +113,9 @@ namespace UE::Chaos::ClothAsset
 		virtual FString GetReferencerName() const override;
 		//~ End FGCObject Interface
 
+		const FChaosClothSimulationProperties* GetSolverProperties() const;
+		const TArray<TSharedPtr<::Chaos::Softs::FCollectionPropertyFacade>>& GetSolverPropertyFacades() const;
+
 		USkeletalMeshComponent* OwnerComponent = nullptr;
 		TUniquePtr<FCollisionSources> CollisionSources;
 		TSharedPtr<FClothSimulationProxy> ClothSimulationProxy;
@@ -124,13 +153,10 @@ private:
 
 	virtual bool SupportsRuntimeInteraction() const override
 	{
-		return false;
+		return true;
 	}
 
-	virtual UClothingSimulationInteractor* CreateInteractor() override
-	{
-		return nullptr;  // NewObject<UChaosClothAssetSimulationInteractor>(GetTransientPackage());
-	}
+	virtual UClothingSimulationInteractor* CreateInteractor() override;
 
 	virtual TArrayView<const TSubclassOf<UClothConfigBase>> GetClothConfigClasses() const override
 	{
diff --git a/Engine/Plugins/ChaosClothAsset/Source/ChaosClothAssetEngine/Private/ChaosClothAsset/ClothAssetSKMClothingSimulationInteractor.cpp b/Engine/Plugins/ChaosClothAsset/Source/ChaosClothAssetEngine/Private/ChaosClothAsset/ClothAssetSKMClothingSimulationInteractor.cpp
new file mode 100644
--- /dev/null
+++ b/Engine/Plugins/ChaosClothAsset/Source/ChaosClothAssetEngine/Private/ChaosClothAsset/ClothAssetSKMClothingSimulationInteractor.cpp
@@ -0,0 +1,132 @@
+// Copyright Epic Games, Inc. All Rights Reserved.
+
+#include "ChaosClothAsset/ClothAssetSKMClothingSimulationInteractor.h"
+#include "ChaosClothAsset/ClothAssetInteractor.h"
+#include "ChaosClothAsset/ClothAssetSKMClothingAsset.h"
+#include "ChaosClothAsset/ClothAssetSKMClothingSimulation.h"
+#include "ChaosClothAsset/ClothAssetPrivate.h"
+#include "Chaos/PBDAnimDriveConstraint.h"
+
+#include UE_INLINE_GENERATED_CPP_BY_NAME(ClothAssetSKMClothingSimulationInteractor)
+
+void UChaosClothAssetSKMClothingInteractor::Initialize(const UChaosClothAssetSKMClothingAsset* InClothingAsset, int32 InClothingId)
+{
+	ClothingId = InClothingId;
+	ClothingAsset = InClothingAsset;
+}
+
+void UChaosClothAssetSKMClothingSimulationInteractor::Sync(IClothingSimulation* Simulation, IClothingSimulationContext* Context)
+{
+	using namespace UE::Chaos::ClothAsset;
+
+	check(Simulation);
+	check(Context);
+
+	FSKMClothingSimulation* const SKMClothingSimulation = static_cast<FSKMClothingSimulation*>(Simulation);
+
+	// Sync interactors
+	for (TPair<FName, TObjectPtr<UClothingInteractor>>& ClothingInteractor : ClothingInteractors)
+	{
+		if (const UChaosClothAssetSKMClothingInteractor* const DummyClothingInteractor = Cast<UChaosClothAssetSKMClothingInteractor>(ClothingInteractor.Value))
+		{
+			if (const UChaosClothAssetSKMClothingAsset* const ClothingAsset = DummyClothingInteractor->ClothingAsset)
+			{
+				if (const int32 ModelIndex = ClothingAsset->GetClothSimulationModelIndex(); ModelIndex != INDEX_NONE)
+				{
+					ClothingInteractor.Value = SKMClothingSimulation->GetPropertyInteractor(ClothingAsset->GetAsset(), ModelIndex);
+				}
+			}
+		}
+		// Note: The iteration could theorically stop at the first non dummy interactor, but since the creation is done outside this file it is better not
+		//       to assume all of them are being created at the same time, especially as there will only be a couple of these interactors in most cases.
+	}
+
+	// Execute commands
+	for (FCommand& Command : Commands)
+	{
+		Command.Execute(SKMClothingSimulation);
+	}
+	Commands.Reset();
+
+	// Call base class' sync 
+	UClothingSimulationInteractor::Sync(Simulation, Context);
+}
+
+void UChaosClothAssetSKMClothingSimulationInteractor::PhysicsAssetUpdated()
+{
+	using namespace UE::Chaos::ClothAsset;
+	Commands.Add(FCommand::CreateLambda([](FSKMClothingSimulation* SKMClothingSimulation)
+		{
+			SKMClothingSimulation->RecreateClothSimulationProxy();
+		}));
+}
+
+void UChaosClothAssetSKMClothingSimulationInteractor::ClothConfigUpdated()
+{
+	using namespace UE::Chaos::ClothAsset;
+	Commands.Add(FCommand::CreateLambda([](FSKMClothingSimulation* SKMClothingSimulation)
+		{
+			SKMClothingSimulation->ResetConfigProperties();
+		}));
+}
+
+void UChaosClothAssetSKMClothingSimulationInteractor::SetAnimDriveSpringStiffness(float Stiffness)
+{
+	for (TPair<FName, TObjectPtr<UClothingInteractor>>& ClothingInteractor : ClothingInteractors)
+	{
+		if (UChaosClothAssetInteractor* const ClothAssetInteractor = Cast<UChaosClothAssetInteractor>(ClothingInteractor.Value))
+		{
+			constexpr int32 AllLODIndices = -1;
+			ClothAssetInteractor->SetFloatPropertyValue(::Chaos::Softs::FPBDAnimDriveConstraint::AnimDriveStiffnessName, AllLODIndices, Stiffness);
+		}
+	}
+}
+
+void UChaosClothAssetSKMClothingSimulationInteractor::EnableGravityOverride(const FVector& Gravity)
+{
+	UE_LOG(LogChaosClothAsset, Warning, TEXT("EnableGravityOverride is not implemented for the Cloth Asset Clothing Simulation Interactor."));
+}
+
+void UChaosClothAssetSKMClothingSimulationInteractor::DisableGravityOverride()
+{
+	UE_LOG(LogChaosClothAsset, Warning, TEXT("DisableGravityOverride is not implemented for the Cloth Asset Clothing Simulation Interactor."));
+}
+
+void UChaosClothAssetSKMClothingSimulationInteractor::SetNumIterations(int32 NumIterations)
+{
+	using namespace UE::Chaos::ClothAsset;
+	Commands.Add(FCommand::CreateLambda([NumIterations](FSKMClothingSimulation* SKMClothingSimulation)
+		{
+			static_cast<IClothingSimulation*>(SKMClothingSimulation)->SetNumIterations(NumIterations);
+		}));
+}
+
+void UChaosClothAssetSKMClothingSimulationInteractor::SetMaxNumIterations(int32 MaxNumIterations)
+{
+	using namespace UE::Chaos::ClothAsset;
+	Commands.Add(FCommand::CreateLambda([MaxNumIterations](FSKMClothingSimulation* SKMClothingSimulation)
+		{
+			static_cast<IClothingSimulation*>(SKMClothingSimulation)->SetMaxNumIterations(MaxNumIterations);
+		}));
+}
+
+void UChaosClothAssetSKMClothingSimulationInteractor::SetNumSubsteps(int32 NumSubsteps)
+{
+	using namespace UE::Chaos::ClothAsset;
+	Commands.Add(FCommand::CreateLambda([NumSubsteps](FSKMClothingSimulation* SKMClothingSimulation)
+		{
+			static_cast<IClothingSimulation*>(SKMClothingSimulation)->SetNumSubsteps(NumSubsteps);
+		}));
+}
+
+void UChaosClothAssetSKMClothingSimulationInteractor::CreateClothingInteractor(const UClothingAssetBase* ClothingAsset, int32 ClothingId)
+{
+	if (const UChaosClothAssetSKMClothingAsset* const SKMClothingAsset = Cast<const UChaosClothAssetSKMClothingAsset>(ClothingAsset))
+	{
+		UChaosClothAssetSKMClothingInteractor* const ClothingInteractor = NewObject<UChaosClothAssetSKMClothingInteractor>(this);
+		check(ClothingInteractor);
+		ClothingInteractor->Initialize(SKMClothingAsset, ClothingId);
+
+		ClothingInteractors.Emplace(ClothingAsset->GetName(), ClothingInteractor);  // The actual interactors will be populated during the Sync call since it needs the simulation pointer
+	}
+}
diff --git a/Engine/Plugins/ChaosClothAsset/Source/ChaosClothAssetEngine/Private/ChaosClothAsset/ClothAssetSKMClothingSimulationInteractor.h b/Engine/Plugins/ChaosClothAsset/Source/ChaosClothAssetEngine/Private/ChaosClothAsset/ClothAssetSKMClothingSimulationInteractor.h
new file mode 100644
--- /dev/null
+++ b/Engine/Plugins/ChaosClothAsset/Source/ChaosClothAssetEngine/Private/ChaosClothAsset/ClothAssetSKMClothingSimulationInteractor.h
@@ -0,0 +1,54 @@
+// Copyright Epic Games, Inc. All Rights Reserved.
+
+#pragma once
+
+#include "ClothingSimulationInteractor.h"
+#include "ClothAssetSKMClothingSimulationInteractor.generated.h"
+
+namespace UE::Chaos::ClothAsset
+{
+	class FSKMClothingSimulation;
+}
+class UChaosClothAssetSKMClothingAsset;
+
+/** Dummy interactor not exposed to Blueprint used during initializations, until the correct clothing interactor is known from the simualtion. */
+UCLASS(Transient, NotBlueprintType, Hidden, HideDropdown, MinimalAPI)
+class UChaosClothAssetSKMClothingInteractor final : public UClothingInteractor
+{
+	GENERATED_BODY()
+
+private:
+	friend UChaosClothAssetSKMClothingSimulationInteractor;
+
+	void Initialize(const UChaosClothAssetSKMClothingAsset* InClothingAsset, int32 InClothingId);
+
+	UPROPERTY()
+	TObjectPtr<const UChaosClothAssetSKMClothingAsset> ClothingAsset;
+};
+
+/** Simulation interactor for Cloth Assets running through a Skeletal Mesh simulation. */
+UCLASS(Transient, BlueprintType, MinimalAPI)
+class UChaosClothAssetSKMClothingSimulationInteractor final : public UClothingSimulationInteractor
+{
+	GENERATED_BODY()
+
+private:
+	DECLARE_DELEGATE_OneParam(FCommand, UE::Chaos::ClothAsset::FSKMClothingSimulation*)
+
+	//~ Begin UClothingSimulationInteractor Interface
+	virtual void CreateClothingInteractor(const UClothingAssetBase* ClothingAsset, int32 ClothingId) override;
+	virtual void Sync(IClothingSimulation* InSimulation, IClothingSimulationContext* InContext) override;
+
+	virtual void PhysicsAssetUpdated() override;
+	virtual void ClothConfigUpdated() override;
+	virtual void SetAnimDriveSpringStiffness(float InStiffness) override;
+	virtual void EnableGravityOverride(const FVector& InVector) override;
+	virtual void DisableGravityOverride() override;
+
+	virtual void SetNumIterations(int32 NumIterations) override;
+	virtual void SetMaxNumIterations(int32 MaxNumIterations) override;
+	virtual void SetNumSubsteps(int32 NumSubsteps) override;
+	//~ End UClothingSimulationInteractor Interface
+
+	TArray<FCommand> Commands;
+};
diff --git a/Engine/Plugins/ChaosClothAsset/Source/ChaosClothAssetEngine/Private/ChaosClothAsset/ClothComponent.cpp b/Engine/Plugins/ChaosClothAsset/Source/ChaosClothAssetEngine/Private/ChaosClothAsset/ClothComponent.cpp
--- a/Engine/Plugins/ChaosClothAsset/Source/ChaosClothAssetEngine/Private/ChaosClothAsset/ClothComponent.cpp
+++ b/Engine/Plugins/ChaosClothAsset/Source/ChaosClothAssetEngine/Private/ChaosClothAsset/ClothComponent.cpp
@@ -131,7 +131,6 @@ void UChaosClothComponent::ResetConfigProperties()
 	}
 }
 
-#if WITH_EDITOR
 void UChaosClothComponent::UpdateConfigProperties()
 {
 	using namespace ::Chaos::Softs;
@@ -156,7 +155,6 @@ void UChaosClothComponent::UpdateConfigProperties()
 		}
 	}
 }
-#endif
 
 void UChaosClothComponent::WaitForExistingParallelClothSimulation_GameThread()
 {
@@ -628,6 +626,20 @@ void UChaosClothComponent::UnregisterOnBoneTransformsFinalizedDelegate(const FDe
 	OnBoneTransformsFinalizedMC.Remove(DelegateHandle);
 }
 
+bool UChaosClothComponent::HasAnyRenderClothMappingData(int32 LODIndex) const
+{
+	if (const UChaosClothAssetBase* const AssetToRender = GetAsset())
+	{
+		if (const FSkeletalMeshRenderData* const SkeletalMeshRenderData = AssetToRender->GetResourceForRendering())
+		{
+			return
+				SkeletalMeshRenderData->LODRenderData.IsValidIndex(LODIndex) &&
+				SkeletalMeshRenderData->LODRenderData[LODIndex].HasClothData();
+		}
+	}
+	return false;
+}
+
 TSharedPtr<UE::Chaos::ClothAsset::FClothSimulationProxy> UChaosClothComponent::CreateClothSimulationProxy()
 {
 	using namespace UE::Chaos::ClothAsset;
diff --git a/Engine/Plugins/ChaosClothAsset/Source/ChaosClothAssetEngine/Private/ChaosClothAsset/ClothSimulationContext.cpp b/Engine/Plugins/ChaosClothAsset/Source/ChaosClothAssetEngine/Private/ChaosClothAsset/ClothSimulationContext.cpp
--- a/Engine/Plugins/ChaosClothAsset/Source/ChaosClothAssetEngine/Private/ChaosClothAsset/ClothSimulationContext.cpp
+++ b/Engine/Plugins/ChaosClothAsset/Source/ChaosClothAssetEngine/Private/ChaosClothAsset/ClothSimulationContext.cpp
@@ -24,24 +24,11 @@ namespace UE::Chaos::ClothAsset
 
 		// Retrieve the list of assets to simulate
 		const TArray<const UChaosClothAssetBase*> Assets = ClothComponentAdapter.GetAssets();
-		if (!bIsInitialization)
+
+		// Do not simulate if there's no asset to render
+		if (!bIsInitialization && !ClothComponentAdapter.HasAnyRenderClothMappingData(LodIndex))
 		{
-			bool bHasAssetValidClothRenderData = false;
-			for (const UChaosClothAssetBase* const Asset : Assets)
-			{
-				if (Asset &&
-					Asset->GetResourceForRendering() &&
-					Asset->GetResourceForRendering()->LODRenderData.IsValidIndex(LodIndex) &&
-					Asset->GetResourceForRendering()->LODRenderData[LodIndex].HasClothData())
-				{
-					bHasAssetValidClothRenderData = true;
-					break;  
-				}
-			}
-			if (!bHasAssetValidClothRenderData)
-			{
-				return;  // This Lod will not simulate, so no need to update the context.
-			}
+			return;  // This Lod will not simulate, so no need to update the context.
 		}
 
 		// Set the teleport mode
@@ -138,21 +125,17 @@ namespace UE::Chaos::ClothAsset
 		{
 			for (const UChaosClothAssetBase* const Asset : Assets)
 			{
-				if (const FSkeletalMeshRenderData* const RenderData = Asset ? Asset->GetResourceForRendering() : nullptr)
+				if (Asset)
 				{
-					if (RenderData->LODRenderData.IsValidIndex(LodIndex))  // TODO Disable this conditional to drive skeletal meshes
+					for (int32 ModelIndex = 0; ModelIndex < Asset->GetNumClothSimulationModels(); ++ModelIndex)
 					{
-						for (int32 ModelIndex = 0; ModelIndex < Asset->GetNumClothSimulationModels(); ++ModelIndex)
+						const TSharedPtr<const FChaosClothSimulationModel>& ClothModel = Asset->GetClothSimulationModel(ModelIndex);
+						if (ClothModel && ClothModel->IsValidLodIndex(LodIndex))
 						{
-							const TSharedPtr<const FChaosClothSimulationModel>& ClothModel = Asset->GetClothSimulationModel(ModelIndex);
-							if (ClothModel && ClothModel->IsValidLodIndex(LodIndex))
-							{
-								RequiredExtraBones.Append(ClothModel->ClothSimulationLodModels[LodIndex].RequiredExtraBoneIndices);
-							}
+							RequiredExtraBones.Append(ClothModel->ClothSimulationLodModels[LodIndex].RequiredExtraBoneIndices);
 						}
-						ClothComponentAdapter.GetOwnerComponent().GetCurrentRefToLocalMatrices(RefToLocalMatrices, LodIndex, &RequiredExtraBones);
-						bSetRefToLocalMatricesToIdentity = false;
 					}
+					bSetRefToLocalMatricesToIdentity = false;
 				}
 			}
 		}
@@ -166,6 +149,10 @@ namespace UE::Chaos::ClothAsset
 				RefToLocalMatrices[BoneIndex] = FMatrix44f::Identity;
 			}
 		}
+		else
+		{
+			ClothComponentAdapter.GetOwnerComponent().GetCurrentRefToLocalMatrices(RefToLocalMatrices, LodIndex, &RequiredExtraBones);
+		}
 		
 		// Update gravity
 		const UWorld* const World = ClothComponentAdapter.GetOwnerComponent().GetWorld();
diff --git a/Engine/Plugins/ChaosClothAsset/Source/ChaosClothAssetEngine/Private/ChaosClothAsset/ClothSimulationProxy.cpp b/Engine/Plugins/ChaosClothAsset/Source/ChaosClothAssetEngine/Private/ChaosClothAsset/ClothSimulationProxy.cpp
--- a/Engine/Plugins/ChaosClothAsset/Source/ChaosClothAssetEngine/Private/ChaosClothAsset/ClothSimulationProxy.cpp
+++ b/Engine/Plugins/ChaosClothAsset/Source/ChaosClothAssetEngine/Private/ChaosClothAsset/ClothSimulationProxy.cpp
@@ -362,26 +362,12 @@ PRAGMA_ENABLE_DEPRECATION_WARNINGS
 		}
 
 		// Check that the render mesh current LOD isn't just fully skinned
-		const TArray<const UChaosClothAssetBase*> Assets = ClothComponentAdapter.GetAssets();
-		const int32 LodIndex = ClothSimulationContext->LodIndex;
-		bool bHasAssetValidClothRenderData = false;
-		for (const UChaosClothAssetBase* const Asset : Assets)
-		{
-			if (Asset &&
-				Asset->GetResourceForRendering() &&
-				Asset->GetResourceForRendering()->LODRenderData.IsValidIndex(LodIndex) &&
-				Asset->GetResourceForRendering()->LODRenderData[LodIndex].HasClothData())
-			{
-				bHasAssetValidClothRenderData = true;
-				break;
-			}
-		}
-		if (!bHasAssetValidClothRenderData)
+		if (!ClothComponentAdapter.HasAnyRenderClothMappingData(ClothSimulationContext->LodIndex))
 		{
 			return false;
 		}
 
-		// Check whether there some actual simulation needs to happen
+		// Check whether some actual simulation needs to happen
 		bIsSimulating = (DeltaTime > 0.f && !ClothComponentAdapter.IsSimulationSuspended() && bIsSolverEnabled);
 		if (!bIsSimulating)
 		{
diff --git a/Engine/Plugins/ChaosClothAsset/Source/ChaosClothAssetEngine/Public/ChaosClothAsset/ClothAsset.h b/Engine/Plugins/ChaosClothAsset/Source/ChaosClothAssetEngine/Public/ChaosClothAsset/ClothAsset.h
--- a/Engine/Plugins/ChaosClothAsset/Source/ChaosClothAssetEngine/Public/ChaosClothAsset/ClothAsset.h
+++ b/Engine/Plugins/ChaosClothAsset/Source/ChaosClothAssetEngine/Public/ChaosClothAsset/ClothAsset.h
@@ -80,16 +80,13 @@ public:
 #endif
 	//~ End USkinnedAsset interface
 
-	/** Return the enclosed Cloth Collection object. */
+	UE_DEPRECATED(5.7, "Use SetClothCollections() or GetClothCollectionsInternal() instead.")
 	TArray<TSharedRef<const FManagedArrayCollection>>& GetClothCollections()
 	{
-		WaitUntilAsyncPropertyReleased(EClothAssetAsyncProperties::ClothCollection);
-		PRAGMA_DISABLE_DEPRECATION_WARNINGS
-		return ClothCollections;
-		PRAGMA_ENABLE_DEPRECATION_WARNINGS
+		return GetClothCollectionsInternal();
 	}
 
-	/** Return the enclosed Cloth Collection object, const version. */
+	/** Return the enclosed Cloth Collections. */
 	const TArray<TSharedRef<const FManagedArrayCollection>>& GetClothCollections() const
 	{
 		WaitUntilAsyncPropertyReleased(EClothAssetAsyncProperties::ClothCollection, ESkinnedAssetAsyncPropertyLockType::ReadOnly);
@@ -98,6 +95,9 @@ public:
 		PRAGMA_ENABLE_DEPRECATION_WARNINGS
 	}
 
+	/** Set new Cloth Collections for this Cloth Asset. */
+	UE_API void SetClothCollections(TArray<TSharedRef<const FManagedArrayCollection>>&& InClothCollections);
+
 	/**
 	 * Rebuild this asset from the specified managed array collection data.
 	 * @param InClothCollections The source managed array collection to build/rebuild the asset with.
@@ -221,6 +221,15 @@ private:
 	/** Build the clothing simulation meshes from the Cloth Collection. */
 	UE_API void BuildClothSimulationModel(TArray<FChaosClothAssetLodTransitionDataCache>* InOutTransitionCache = nullptr);
 
+	/** Return the cloth collection array with async protection access. */
+	TArray<TSharedRef<const FManagedArrayCollection>>& GetClothCollectionsInternal()  // TODO: Can be renamed GetClothCollections() once the current public method is removed in 5.9
+	{
+		WaitUntilAsyncPropertyReleased(EClothAssetAsyncProperties::ClothCollection);
+		PRAGMA_DISABLE_DEPRECATION_WARNINGS
+		return ClothCollections;
+		PRAGMA_ENABLE_DEPRECATION_WARNINGS
+	}
+
 #if WITH_EDITORONLY_DATA
 	/** Dataflow asset. */
 	UPROPERTY(meta = (DeprecatedProperty, DeprecationMessage = "Dataflow Asset is now stored in DataflowInstance"))
diff --git a/Engine/Plugins/ChaosClothAsset/Source/ChaosClothAssetEngine/Public/ChaosClothAsset/ClothAssetBase.h b/Engine/Plugins/ChaosClothAsset/Source/ChaosClothAssetEngine/Public/ChaosClothAsset/ClothAssetBase.h
--- a/Engine/Plugins/ChaosClothAsset/Source/ChaosClothAssetEngine/Public/ChaosClothAsset/ClothAssetBase.h
+++ b/Engine/Plugins/ChaosClothAsset/Source/ChaosClothAssetEngine/Public/ChaosClothAsset/ClothAssetBase.h
@@ -275,8 +275,11 @@ public:
 #endif
 
 protected:
+	/** Update all properties for components currently using this asset. */
+	UE_API void OnPropertyChanged() const;
+
 	/** Reregister all components using this asset to reset the simulation in case anything has changed. */
-	UE_API void ReregisterComponents();
+	UE_API void ReregisterComponents() const;
 
 	/**
 	 * Wait for the asset to finish compilation to protect internal skinned asset data from race conditions during async build.
diff --git a/Engine/Plugins/ChaosClothAsset/Source/ChaosClothAssetEngine/Public/ChaosClothAsset/ClothAssetInteractor.h b/Engine/Plugins/ChaosClothAsset/Source/ChaosClothAssetEngine/Public/ChaosClothAsset/ClothAssetInteractor.h
--- a/Engine/Plugins/ChaosClothAsset/Source/ChaosClothAssetEngine/Public/ChaosClothAsset/ClothAssetInteractor.h
+++ b/Engine/Plugins/ChaosClothAsset/Source/ChaosClothAssetEngine/Public/ChaosClothAsset/ClothAssetInteractor.h
@@ -1,7 +1,7 @@
 // Copyright Epic Games, Inc. All Rights Reserved.
 
 #pragma once
-#include "UObject/Object.h"
+#include "ClothingSimulationInteractor.h"
 #include "CoreTypes.h"
 #include "Math/MathFwd.h"
 #include "ClothAssetInteractor.generated.h"
@@ -12,7 +12,7 @@ class FCollectionPropertyFacade;
 }
 
 UCLASS(BlueprintType, MinimalAPI)
-class UChaosClothAssetInteractor : public UObject
+class UChaosClothAssetInteractor : public UClothingInteractor
 {
 	GENERATED_BODY()
 public:
@@ -29,59 +29,59 @@ public:
 
 	/** Get the value for a property cast to float. DefaultValue will be returned if the property is not found.*/
 	UFUNCTION(BlueprintCallable, Category = ClothProperty)
-	CHAOSCLOTHASSETENGINE_API float GetFloatPropertyValue(const FName& PropertyName, int32 LODIndex = 0, float DefaultValue = 0.f) const;
+	CHAOSCLOTHASSETENGINE_API float GetFloatPropertyValue(const FName PropertyName, int32 LODIndex = 0, float DefaultValue = 0.f) const;
 
 	/** Get the low value for a weighted property value (same as GetFloatValue). DefaultValue will be returned if the property is not found.*/
 	UFUNCTION(BlueprintCallable, Category = ClothProperty)
-	CHAOSCLOTHASSETENGINE_API float GetLowFloatPropertyValue(const FName& PropertyName, int32 LODIndex = 0, float DefaultValue = 0.f) const;
+	CHAOSCLOTHASSETENGINE_API float GetLowFloatPropertyValue(const FName PropertyName, int32 LODIndex = 0, float DefaultValue = 0.f) const;
 
 	/** Get the high value for a weighted property value. DefaultValue will be returned if the property is not found.*/
 	UFUNCTION(BlueprintCallable, Category = ClothProperty)
-	CHAOSCLOTHASSETENGINE_API float GetHighFloatPropertyValue(const FName& PropertyName, int32 LODIndex = 0, float DefaultValue = 0.f) const;
+	CHAOSCLOTHASSETENGINE_API float GetHighFloatPropertyValue(const FName PropertyName, int32 LODIndex = 0, float DefaultValue = 0.f) const;
 
 	/** Get the low and high values for a weighted property value. DefaultValue will be returned if the property is not found.*/
 	UFUNCTION(BlueprintCallable, Category = ClothProperty)
-	CHAOSCLOTHASSETENGINE_API FVector2D GetWeightedFloatPropertyValue(const FName& PropertyName, int32 LODIndex = 0, FVector2D DefaultValue = FVector2D(0., 0.)) const;
+	CHAOSCLOTHASSETENGINE_API FVector2D GetWeightedFloatPropertyValue(const FName PropertyName, int32 LODIndex = 0, FVector2D DefaultValue = FVector2D(0., 0.)) const;
 
 	/** Get the value for a property cast to int. DefaultValue will be returned if the property is not found.*/
 	UFUNCTION(BlueprintCallable, Category = ClothProperty)
-	CHAOSCLOTHASSETENGINE_API int32 GetIntPropertyValue(const FName& PropertyName, int32 LODIndex = 0, int32 DefaultValue = 0) const;
+	CHAOSCLOTHASSETENGINE_API int32 GetIntPropertyValue(const FName PropertyName, int32 LODIndex = 0, int32 DefaultValue = 0) const;
 
 	/** Get the value for a property cast to vector. DefaultValue will be returned if the property is not found.*/
 	UFUNCTION(BlueprintCallable, Category = ClothProperty)
-	CHAOSCLOTHASSETENGINE_API FVector GetVectorPropertyValue(const FName& PropertyName, int32 LODIndex = 0, FVector DefaultValue = FVector(0.)) const;
+	CHAOSCLOTHASSETENGINE_API FVector GetVectorPropertyValue(const FName PropertyName, int32 LODIndex = 0, FVector DefaultValue = FVector(0.)) const;
 
 	/** Get the string value for a property (typically the associated map name for weighted values). DefaultValue will be returned if the property is not found.*/
 	UFUNCTION(BlueprintCallable, Category = ClothProperty)
-	CHAOSCLOTHASSETENGINE_API FString GetStringPropertyValue(const FName& PropertyName, int32 LODIndex = 0, const FString& DefaultValue = "") const;
+	CHAOSCLOTHASSETENGINE_API FString GetStringPropertyValue(const FName PropertyName, int32 LODIndex = 0, const FString& DefaultValue = "") const;
 
 	/**Set the value for a property (if it exists). This sets the Low and High values for weighted values. All LODs will be set when LODIndex = -1.*/
 	UFUNCTION(BlueprintCallable, Category = ClothProperty)
-	CHAOSCLOTHASSETENGINE_API void SetFloatPropertyValue(const FName& PropertyName, int32 LODIndex = -1, float Value = 0.f);
+	CHAOSCLOTHASSETENGINE_API void SetFloatPropertyValue(const FName PropertyName, int32 LODIndex = -1, float Value = 0.f);
 
 	/**Set the low value for a weighted property (if it exists). All LODs will be set when LODIndex = -1.*/
 	UFUNCTION(BlueprintCallable, Category = ClothProperty)
-	CHAOSCLOTHASSETENGINE_API void SetLowFloatPropertyValue(const FName& PropertyName, int32 LODIndex = -1, float Value = 0.f);
+	CHAOSCLOTHASSETENGINE_API void SetLowFloatPropertyValue(const FName PropertyName, int32 LODIndex = -1, float Value = 0.f);
 
 	/**Set the high value for a weighted property (if it exists). All LODs will be set when LODIndex = -1.*/
 	UFUNCTION(BlueprintCallable, Category = ClothProperty)
-	CHAOSCLOTHASSETENGINE_API void SetHighFloatPropertyValue(const FName& PropertyName, int32 LODIndex = -1, float Value = 0.f);
+	CHAOSCLOTHASSETENGINE_API void SetHighFloatPropertyValue(const FName PropertyName, int32 LODIndex = -1, float Value = 0.f);
 
 	/**Set the low and high values for a weighted property (if it exists). All LODs will be set when LODIndex = -1.*/
 	UFUNCTION(BlueprintCallable, Category = ClothProperty)
-	CHAOSCLOTHASSETENGINE_API void SetWeightedFloatPropertyValue(const FName& PropertyName, int32 LODIndex = -1, FVector2D Value = FVector2D(0., 0.));
+	CHAOSCLOTHASSETENGINE_API void SetWeightedFloatPropertyValue(const FName PropertyName, int32 LODIndex = -1, FVector2D Value = FVector2D(0., 0.));
 
 	/**Set the value for a property (if it exists). All LODs will be set when LODIndex = -1.*/
 	UFUNCTION(BlueprintCallable, Category = ClothProperty)
-	CHAOSCLOTHASSETENGINE_API void SetIntPropertyValue(const FName& PropertyName, int32 LODIndex = -1, int32 Value = 0);
+	CHAOSCLOTHASSETENGINE_API void SetIntPropertyValue(const FName PropertyName, int32 LODIndex = -1, int32 Value = 0);
 
 	/**Set the value for a property (if it exists). All LODs will be set when LODIndex = -1.*/
 	UFUNCTION(BlueprintCallable, Category = ClothProperty)
-	CHAOSCLOTHASSETENGINE_API void SetVectorPropertyValue(const FName& PropertyName, int32 LODIndex = -1, FVector Value = FVector(0.));
+	CHAOSCLOTHASSETENGINE_API void SetVectorPropertyValue(const FName PropertyName, int32 LODIndex = -1, FVector Value = FVector(0.));
 
 	/**Set the string value for a property (if it exists). This is typically the map name associated with a property. All LODs will be set when LODIndex = -1.*/
 	UFUNCTION(BlueprintCallable, Category = ClothProperty)
-	CHAOSCLOTHASSETENGINE_API void SetStringPropertyValue(const FName& PropertyName, int32 LODIndex = -1, const FString& Value = "");
+	CHAOSCLOTHASSETENGINE_API void SetStringPropertyValue(const FName PropertyName, int32 LODIndex = -1, const FString& Value = "");
 
 	// Deprecated string-based key versions.
 	/** Generate a list of all properties held by this interactor. Properties for all LODs will be returned if LODIndex = -1.*/
diff --git a/Engine/Plugins/ChaosClothAsset/Source/ChaosClothAssetEngine/Public/ChaosClothAsset/ClothComponent.h b/Engine/Plugins/ChaosClothAsset/Source/ChaosClothAssetEngine/Public/ChaosClothAsset/ClothComponent.h
--- a/Engine/Plugins/ChaosClothAsset/Source/ChaosClothAssetEngine/Public/ChaosClothAsset/ClothComponent.h
+++ b/Engine/Plugins/ChaosClothAsset/Source/ChaosClothAssetEngine/Public/ChaosClothAsset/ClothComponent.h
@@ -290,10 +290,8 @@ public:
 	UFUNCTION(BlueprintSetter, Category = ClothComponent)
 	UE_API void SetTeleportRotationThreshold(float Threshold);
 
-#if WITH_EDITOR
 	/** Update config properties from the asset. Will only update existing values.*/
 	UE_API void UpdateConfigProperties();
-#endif
 
 	/** Stalls on any currently running clothing simulations.*/
 	UE_API void WaitForExistingParallelClothSimulation_GameThread();
@@ -375,6 +373,7 @@ protected:
 		checkf(GetSimulationGroupId(InAsset, ModelIndex) != INDEX_NONE, TEXT("Invalid arguments"));
 		return GetPropertyCollections(ModelIndex);  // There's only this one asset for now, therefore this is the same as calling GetPropertyCollections(ModelIndex)
 	}
+	UE_API virtual bool HasAnyRenderClothMappingData(int32 LODIndex) const override;
 	//~ End IClothComponentAdapter Interface
 
 	/** Override this function for setting up custom simulation proxies when the component is registered. */
diff --git a/Engine/Plugins/ChaosClothAsset/Source/ChaosClothAssetEngine/Public/ChaosClothAsset/ClothComponentAdapter.h b/Engine/Plugins/ChaosClothAsset/Source/ChaosClothAssetEngine/Public/ChaosClothAsset/ClothComponentAdapter.h
--- a/Engine/Plugins/ChaosClothAsset/Source/ChaosClothAssetEngine/Public/ChaosClothAsset/ClothComponentAdapter.h
+++ b/Engine/Plugins/ChaosClothAsset/Source/ChaosClothAssetEngine/Public/ChaosClothAsset/ClothComponentAdapter.h
@@ -18,20 +18,51 @@ namespace UE::Chaos::ClothAsset
 	class IClothComponentAdapter
 	{
 	public:
+		/**
+		 * Return the SkinnedMeshComponent owning this adapter.
+		 * It can be this object pointer or a different owner component if this interface isn't directly implemented on a component.
+		 */
 		virtual const USkinnedMeshComponent& GetOwnerComponent() const = 0;
+
+		/** Return the collision sources to use in cloth simulation proxy. */
 		virtual FCollisionSources& GetCollisionSources() const = 0;
+
+		/** Return the list of cloth/outfit assets in use by this adapter. */
 		virtual TArray<const UChaosClothAssetBase*> GetAssets() const = 0;
-		/** Returns the index of the model assigned to the render section, or INDEX_NONE when the model shouldn't be simulated. */
+
+		/** Return the index of the model assigned to the render section, or INDEX_NONE when the model shouldn't be simulated. */
 		virtual int32 GetSimulationGroupId(const UChaosClothAssetBase* Asset, int32 ModelId) const = 0;
+
+		/** Return the reference skeleton as passed to the component renderer. */
 		virtual const FReferenceSkeleton* GetReferenceSkeleton() const = 0;
+
+		/** Return the current teleport mode. */
 		virtual EClothingTeleportMode GetClothTeleportMode() const = 0;
+
+		/** Return whether the simulation is currently running. */
 		virtual bool IsSimulationEnabled() const = 0;
+
+		/** Return whether the simulation is currently suspended. */
 		virtual bool IsSimulationSuspended() const = 0;
+
+		/** Return whether the cloth mesh rest lengths needs to be reset. */
 		virtual bool NeedsResetRestLengths() const = 0;
+
+		/** Return the name of the current morph target. */
 		virtual const FString& GetRestLengthsMorphTargetName() const = 0;
+
+		/** Return the scale applied to all cloth geometries (cloth meshes and collisions) in order to simulate in a different scale space. */
 		virtual float GetClothGeometryScale() const = 0;
+
+		/** Return the runtime property collections used for modifications of the specified cloth/outfit asset simulation model. */
 		virtual const TArray<TSharedPtr<const FManagedArrayCollection>>& GetPropertyCollections(const UChaosClothAssetBase* Asset, int32 ModelIndex) const = 0;
+
+		/** Return the runtime solver property collections. */
 		virtual const TArray<TSharedPtr<const FManagedArrayCollection>>& GetSolverPropertyCollections() const = 0;
+
+		/** Return true when any of the rendered assets LOD sections has valid Cloth Mapping Data, or false otherwise. */
+		virtual bool HasAnyRenderClothMappingData(int32 LodIndex) const = 0;
+
 	protected:
 		virtual ~IClothComponentAdapter() = default;
 	};
diff --git a/Engine/Plugins/ChaosClothAssetEditor/Source/ChaosClothAssetDataflowNodes/Private/ChaosClothAsset/MakeClothAssetNode.cpp b/Engine/Plugins/ChaosClothAssetEditor/Source/ChaosClothAssetDataflowNodes/Private/ChaosClothAsset/MakeClothAssetNode.cpp
--- a/Engine/Plugins/ChaosClothAssetEditor/Source/ChaosClothAssetDataflowNodes/Private/ChaosClothAsset/MakeClothAssetNode.cpp
+++ b/Engine/Plugins/ChaosClothAssetEditor/Source/ChaosClothAssetDataflowNodes/Private/ChaosClothAsset/MakeClothAssetNode.cpp
@@ -37,8 +37,7 @@ void FChaosClothAssetMakeClothAssetNode::Evaluate(UE::Dataflow::FContext& Contex
 	{
 		TObjectPtr<UChaosClothAsset> OutClothAsset = NewObject<UChaosClothAsset>();
 
-		TArray<TSharedRef<const FManagedArrayCollection>> InClothCollections = GetCleanedCollectionLodValues(Context);
-		TArray<TSharedRef<const FManagedArrayCollection>>& ClothCollections = OutClothAsset->GetClothCollections();
+		const TArray<TSharedRef<const FManagedArrayCollection>> InClothCollections = GetCleanedCollectionLodValues(Context);
 
 		FText ErrorText;
 		FText VerboseText;
diff --git a/Engine/Plugins/ChaosClothAssetEditor/Source/ChaosClothAssetDataflowNodes/Private/ChaosClothAsset/TerminalNode.cpp b/Engine/Plugins/ChaosClothAssetEditor/Source/ChaosClothAssetDataflowNodes/Private/ChaosClothAsset/TerminalNode.cpp
--- a/Engine/Plugins/ChaosClothAssetEditor/Source/ChaosClothAssetDataflowNodes/Private/ChaosClothAsset/TerminalNode.cpp
+++ b/Engine/Plugins/ChaosClothAssetEditor/Source/ChaosClothAssetDataflowNodes/Private/ChaosClothAsset/TerminalNode.cpp
@@ -112,7 +112,7 @@ void FChaosClothAssetTerminalNode_v2::SetAssetValue(TObjectPtr<UObject> Asset, U
 		using namespace UE::Chaos::ClothAsset;
 
 		TArray<TSharedRef<const FManagedArrayCollection>> InClothCollections = GetCleanedCollectionLodValues(Context);
-		TArray<TSharedRef<const FManagedArrayCollection>>& ClothCollections = ClothAsset->GetClothCollections();
+		TArray<TSharedRef<const FManagedArrayCollection>> ClothCollections = static_cast<const UChaosClothAsset*>(ClothAsset)->GetClothCollections();
 
 		const uint32 PreviousChecksum = ClothColllectionChecksum;
 		const bool bPreviousChecksumsValid = bClothCollectionChecksumValid;
@@ -132,6 +132,8 @@ void FChaosClothAssetTerminalNode_v2::SetAssetValue(TObjectPtr<UObject> Asset, U
 				ClothCollections[LODIndex] = MoveTemp(ClothCollection);
 			}
 
+			ClothAsset->SetClothCollections(MoveTemp(ClothCollections));
+
 			// Asset must be resaved
 			ClothAsset->MarkPackageDirty();
 
@@ -306,16 +308,17 @@ PRAGMA_DISABLE_DEPRECATION_WARNINGS
 PRAGMA_ENABLE_DEPRECATION_WARNINGS
 
 		TArray<TSharedRef<const FManagedArrayCollection>> InClothCollections = GetCleanedCollectionLodValues(Context);
-		TArray<TSharedRef<const FManagedArrayCollection>>& ClothCollections = ClothAsset->GetClothCollections();
+		const TArray<TSharedRef<const FManagedArrayCollection>>& RefClothCollections = static_cast<const UChaosClothAsset*>(ClothAsset)->GetClothCollections();
 
 		const uint32 PreviousChecksum = ClothColllectionChecksum;
 		const bool bPreviousChecksumsValid = bClothCollectionChecksumValid;
 		ClothColllectionChecksum = Private::CalculateClothChecksum(InClothCollections);
 		bClothCollectionChecksumValid = InClothCollections.Num() > 0;
 
-		if (bPreviousChecksumsValid && PreviousChecksum == ClothColllectionChecksum && Private::PropertyKeysAndSolverTypesMatch(InClothCollections, ClothCollections))
+		if (bPreviousChecksumsValid && PreviousChecksum == ClothColllectionChecksum && Private::PropertyKeysAndSolverTypesMatch(InClothCollections, RefClothCollections))
 		{
 			// Cloth and property keys match. Just update property values.
+			TArray<TSharedRef<const FManagedArrayCollection>> ClothCollections = RefClothCollections;  // Copy on write
 			check(InClothCollections.Num() == ClothCollections.Num());
 			check(ClothCollections.Num() > 0);
 			for (int32 LODIndex = 0; LODIndex < InClothCollections.Num(); ++LODIndex)
@@ -326,6 +329,8 @@ PRAGMA_ENABLE_DEPRECATION_WARNINGS
 				ClothCollections[LODIndex] = MoveTemp(ClothCollection);
 			}
 
+			ClothAsset->SetClothCollections(MoveTemp(ClothCollections));
+
 			// Asset must be resaved
 			ClothAsset->MarkPackageDirty();
 
diff --git a/Engine/Plugins/ChaosClothAssetEditor/Source/ChaosClothAssetEditor/Private/ChaosClothAsset/ClothEditorPreviewScene.cpp b/Engine/Plugins/ChaosClothAssetEditor/Source/ChaosClothAssetEditor/Private/ChaosClothAsset/ClothEditorPreviewScene.cpp
--- a/Engine/Plugins/ChaosClothAssetEditor/Source/ChaosClothAssetEditor/Private/ChaosClothAsset/ClothEditorPreviewScene.cpp
+++ b/Engine/Plugins/ChaosClothAssetEditor/Source/ChaosClothAssetEditor/Private/ChaosClothAsset/ClothEditorPreviewScene.cpp
@@ -367,12 +367,7 @@ void FChaosClothPreviewScene::SetClothAsset(UChaosClothAsset* Asset)
 	check(SceneActor);
 	check(ClothComponent);
 
-	if (Cast<UChaosClothAsset>(ClothComponent->GetAsset()) == Asset)
-	{
-		// Update the config properties on the component from the asset.
-		ClothComponent->UpdateConfigProperties();
-	}
-	else
+	if (Cast<UChaosClothAsset>(ClothComponent->GetAsset()) != Asset)
 	{
 		ClothComponent->SetAsset(Asset);
 	}
diff --git a/Engine/Source/Runtime/ClothingSystemRuntimeInterface/Private/ClothingSimulationInstance.cpp b/Engine/Source/Runtime/ClothingSystemRuntimeInterface/Private/ClothingSimulationInstance.cpp
--- a/Engine/Source/Runtime/ClothingSystemRuntimeInterface/Private/ClothingSimulationInstance.cpp
+++ b/Engine/Source/Runtime/ClothingSystemRuntimeInterface/Private/ClothingSimulationInstance.cpp
@@ -28,6 +28,7 @@ FClothingSimulationInstance::FClothingSimulationInstance(USkeletalMeshComponent*
 			if (ClothingSimulationInteractor)
 			{
 				ClothingSimulationInteractor->CreateClothingInteractor(ClothingAsset, ClothingAssetIndex);
+				ClothingSimulationInteractor->Sync(ClothingSimulation, ClothingSimulationContext);  // An early sync is required by some interactors to complete the initializations
 			}
 		}
 	}
