Chaos Cloth Asset:
- Fixed the wind in the skeletal mesh editor
- Implemented a light weight ClothAssetSKMContext derived from the IClothingSimulationContext interface instead of relyingon the FClothingSimulationContextCommon implementation (the actual simulation context is managed by the Cloth SImulation Proxy).

#rnx

diff --git a/Engine/Plugins/ChaosClothAsset/Source/ChaosClothAssetEngine/Private/ChaosClothAsset/ClothAssetSKMClothingSimulation.cpp b/Engine/Plugins/ChaosClothAsset/Source/ChaosClothAssetEngine/Private/ChaosClothAsset/ClothAssetSKMClothingSimulation.cpp
--- a/Engine/Plugins/ChaosClothAsset/Source/ChaosClothAssetEngine/Private/ChaosClothAsset/ClothAssetSKMClothingSimulation.cpp
+++ b/Engine/Plugins/ChaosClothAsset/Source/ChaosClothAssetEngine/Private/ChaosClothAsset/ClothAssetSKMClothingSimulation.cpp
@@ -74,6 +74,14 @@ namespace UE::Chaos::ClothAsset
 		ClothSimulationProxy = nullptr;
 	}
 
+	void FSKMClothingSimulation::FillContext(const USkeletalMeshComponent* /*InComponent*/, float InDeltaTime, IClothingSimulationContext* InOutContext, bool /*bIsInitialization*/) const
+	{
+		if (FSKMClothingSimulationContext* const Context = static_cast<FSKMClothingSimulationContext*>(InOutContext))
+		{
+			Context->DeltaTime = InDeltaTime;
+		}
+	}
+
 	void FSKMClothingSimulation::CreateActor(USkeletalMeshComponent* InOwnerComponent, const UClothingAssetBase* ClothingAsset, int32 SimDataIndex)
 	{
 		check(InOwnerComponent);
@@ -172,7 +180,7 @@ namespace UE::Chaos::ClothAsset
 			if (ClothSimulationProxy.IsValid())
 			{
 				constexpr bool bUseParallelTask = false;  // The SKM component already deals with calling Simulate() in its own parallel task
-				ClothSimulationProxy->Tick_GameThread(SKMClothingSimulationContext->DeltaSeconds, bUseParallelTask);
+				ClothSimulationProxy->Tick_GameThread(SKMClothingSimulationContext->DeltaTime, bUseParallelTask);
 			}
 		}
 		// Reset the bNeedsResetRestLengths at every pass
diff --git a/Engine/Plugins/ChaosClothAsset/Source/ChaosClothAssetEngine/Private/ChaosClothAsset/ClothAssetSKMClothingSimulation.h b/Engine/Plugins/ChaosClothAsset/Source/ChaosClothAssetEngine/Private/ChaosClothAsset/ClothAssetSKMClothingSimulation.h
--- a/Engine/Plugins/ChaosClothAsset/Source/ChaosClothAssetEngine/Private/ChaosClothAsset/ClothAssetSKMClothingSimulation.h
+++ b/Engine/Plugins/ChaosClothAsset/Source/ChaosClothAssetEngine/Private/ChaosClothAsset/ClothAssetSKMClothingSimulation.h
@@ -4,7 +4,7 @@
 
 #include "ChaosClothAsset/ClothComponentAdapter.h"
 #include "ChaosClothAsset/ClothSimulationProxy.h"
-#include "ClothingSimulation.h"
+#include "ClothingSimulationInterface.h"
 #include "ClothingSimulationFactory.h"
 #include "UObject/GCObject.h"
 #include "ClothAssetSKMClothingSimulation.generated.h"
@@ -21,10 +21,21 @@ namespace Chaos::Softs
 
 namespace UE::Chaos::ClothAsset
 {
-	using FSKMClothingSimulationContext = FClothingSimulationContextCommon;
 
+	/** Context for the SKM simulation, only used to temporarily store the simulation timestep, the actual simulation context is managed by the simulation proxy. */
+	class FSKMClothingSimulationContext : public IClothingSimulationContext
+	{
+	public:
+		FSKMClothingSimulationContext() : IClothingSimulationContext() {}
+		virtual ~FSKMClothingSimulationContext() override {}
+
+		/** Simulation timestep. */
+		float DeltaTime = 0.f;
+	};
+
+	/** Skeletal Mesh simulation class that interfaces between the Skeletal Mesh Component and the Cloth Simulation Proxy. */
 	class FSKMClothingSimulation final
-		: public FClothingSimulationCommon
+		: public IClothingSimulation
 		, public IClothComponentAdapter
 		, public FGCObject
 	{
@@ -50,6 +61,7 @@ namespace UE::Chaos::ClothAsset
 		{
 			return new FSKMClothingSimulationContext();
 		}
+		virtual void FillContext(const USkeletalMeshComponent* InComponent, float InDeltaTime, IClothingSimulationContext* InOutContext, bool bIsInitialization) const override;
 		virtual void DestroyContext(IClothingSimulationContext* Context) override
 		{
 			delete Context;
diff --git a/Engine/Plugins/ChaosClothAsset/Source/ChaosClothAssetEngine/Private/ChaosClothAsset/ClothSimulationContext.cpp b/Engine/Plugins/ChaosClothAsset/Source/ChaosClothAssetEngine/Private/ChaosClothAsset/ClothSimulationContext.cpp
--- a/Engine/Plugins/ChaosClothAsset/Source/ChaosClothAssetEngine/Private/ChaosClothAsset/ClothSimulationContext.cpp
+++ b/Engine/Plugins/ChaosClothAsset/Source/ChaosClothAssetEngine/Private/ChaosClothAsset/ClothSimulationContext.cpp
@@ -160,7 +160,7 @@ namespace UE::Chaos::ClothAsset
 		WorldGravity = FVector(0.f, 0.f, World ? World->GetGravityZ() : EarthGravity);
 
 		// Update wind velocity
-		if (World && World->Scene && World->IsGameWorld())
+		if (World && World->Scene)
 		{
 			const FVector Position = ClothComponentAdapter.GetOwnerComponent().GetComponentTransform().GetTranslation();
 
