Fix ChaosClothAsset async build issues.

Fix ChaosClothAsset [Acquire/Release]AsyncProperty to account for enum offsets. Make UChaosClothAsset::GetCollections only require a read lock.

diff --git a/Engine/Plugins/ChaosClothAsset/Source/ChaosClothAssetEngine/Private/ChaosClothAsset/ClothAsset.cpp b/Engine/Plugins/ChaosClothAsset/Source/ChaosClothAssetEngine/Private/ChaosClothAsset/ClothAsset.cpp
--- a/Engine/Plugins/ChaosClothAsset/Source/ChaosClothAssetEngine/Private/ChaosClothAsset/ClothAsset.cpp
+++ b/Engine/Plugins/ChaosClothAsset/Source/ChaosClothAssetEngine/Private/ChaosClothAsset/ClothAsset.cpp
@@ -363,7 +363,7 @@ void UChaosClothAsset::BeginBuildInternal(FSkinnedAssetBuildContext& Context)
 	ReleaseResourcesFence.Wait();
 
 	// Lock all properties that should not be modified/accessed during async post-load
-	AcquireAsyncProperty();
+	USkinnedAsset::AcquireAsyncProperty();
 }
 
 void UChaosClothAsset::FinishBuildInternal(FSkinnedAssetBuildContext& Context)
@@ -372,7 +372,7 @@ void UChaosClothAsset::FinishBuildInternal(FSkinnedAssetBuildContext& Context)
 
 	ClearInternalFlags(EInternalObjectFlags::Async);
 
-	ReleaseAsyncProperty();
+	USkinnedAsset::ReleaseAsyncProperty();
 }
 #endif // #if WITH_EDITOR
 
@@ -387,7 +387,7 @@ void UChaosClothAsset::BeginPostLoadInternal(FSkinnedAssetPostLoadContext& Conte
 	SetInternalFlags(EInternalObjectFlags::Async);
 
 	// Lock all properties that should not be modified/accessed during async post-load
-	AcquireAsyncProperty();
+	USkinnedAsset::AcquireAsyncProperty();
 
 	// This scope allows us to use any locked properties without causing stalls
 	FSkinnedAssetAsyncBuildScope AsyncBuildScope(this);
@@ -421,7 +421,7 @@ void UChaosClothAsset::BeginPostLoadInternal(FSkinnedAssetPostLoadContext& Conte
 	}
 
 	// We're done touching the ClothCollections, so can unlock for read
-	ReleaseAsyncProperty((uint64)EClothAssetAsyncProperties::ClothCollection, ESkinnedAssetAsyncPropertyLockType::WriteOnly);
+	ReleaseAsyncProperty(EClothAssetAsyncProperties::ClothCollection, ESkinnedAssetAsyncPropertyLockType::WriteOnly);
 
 	// Build the cloth simulation model
 	BuildClothSimulationModel();  // TODO: Cache ClothSimulationModel in the DDC
@@ -481,7 +481,7 @@ void UChaosClothAsset::FinishPostLoadInternal(FSkinnedAssetPostLoadContext& Cont
 	CalculateInvRefMatrices();
 
 #if WITH_EDITOR
-	ReleaseAsyncProperty();
+	USkinnedAsset::ReleaseAsyncProperty();
 #endif
 }
 
@@ -718,7 +718,7 @@ void UChaosClothAsset::Build(TArray<FChaosClothAssetLodTransitionDataCache>* InO
 void UChaosClothAsset::PrepareMeshModel()
 {
 	WaitUntilAsyncPropertyReleased(EClothAssetAsyncProperties::ImportedModel);
-	WaitUntilAsyncPropertyReleased(EClothAssetAsyncProperties::ClothCollection);
+	WaitUntilAsyncPropertyReleased(EClothAssetAsyncProperties::ClothCollection, ESkinnedAssetAsyncPropertyLockType::ReadOnly);
 
 PRAGMA_DISABLE_DEPRECATION_WARNINGS
 	const int32 NumLods = ClothCollections.Num();
@@ -747,14 +747,6 @@ FString UChaosClothAsset::GetAsyncPropertyName(uint64 Property) const
 	return StaticEnum<EClothAssetAsyncProperties>()->GetNameByValue(Property).ToString();
 }
 
-const TArray<TSharedRef<const FManagedArrayCollection>>& UChaosClothAsset::GetCollections(int32 /*ModelIndex*/) const
-{
-	WaitUntilAsyncPropertyReleased(EClothAssetAsyncProperties::ClothCollection);
-PRAGMA_DISABLE_DEPRECATION_WARNINGS
-	return ClothCollections;
-PRAGMA_ENABLE_DEPRECATION_WARNINGS
-}
-
 #if WITH_EDITOR
 void UChaosClothAsset::CacheDerivedData(FSkinnedAssetCompilationContext* Context)
 {
diff --git a/Engine/Plugins/ChaosClothAsset/Source/ChaosClothAssetEngine/Public/ChaosClothAsset/ClothAsset.h b/Engine/Plugins/ChaosClothAsset/Source/ChaosClothAssetEngine/Public/ChaosClothAsset/ClothAsset.h
--- a/Engine/Plugins/ChaosClothAsset/Source/ChaosClothAssetEngine/Public/ChaosClothAsset/ClothAsset.h
+++ b/Engine/Plugins/ChaosClothAsset/Source/ChaosClothAssetEngine/Public/ChaosClothAsset/ClothAsset.h
@@ -170,7 +170,10 @@ PRAGMA_ENABLE_DEPRECATION_WARNINGS
 	{
 		return GetFName();
 	}
-	UE_API virtual const TArray<TSharedRef<const FManagedArrayCollection>>& GetCollections(int32 /*ModelIndex*/) const override;
+	virtual const TArray<TSharedRef<const FManagedArrayCollection>>& GetCollections(int32 /*ModelIndex*/) const override
+	{
+		return GetClothCollections();
+	}
 	virtual const UPhysicsAsset* GetPhysicsAssetForModel(int32 /*ModelIndex*/) const override
 	{
 		return PhysicsAsset;
diff --git a/Engine/Plugins/ChaosClothAsset/Source/ChaosClothAssetEngine/Public/ChaosClothAsset/ClothAssetBase.h b/Engine/Plugins/ChaosClothAsset/Source/ChaosClothAssetEngine/Public/ChaosClothAsset/ClothAssetBase.h
--- a/Engine/Plugins/ChaosClothAsset/Source/ChaosClothAssetEngine/Public/ChaosClothAsset/ClothAssetBase.h
+++ b/Engine/Plugins/ChaosClothAsset/Source/ChaosClothAssetEngine/Public/ChaosClothAsset/ClothAssetBase.h
@@ -281,6 +281,22 @@ protected:
 	/** Reregister all components using this asset to reset the simulation in case anything has changed. */
 	UE_API void ReregisterComponents() const;
 
+	/** Lock properties that should not be modified/accessed during async build. */
+	template <typename EnumType UE_REQUIRES(TIsEnum<EnumType>::Value)>
+	void AcquireAsyncProperty(const EnumType AsyncProperties, ESkinnedAssetAsyncPropertyLockType LockType = ESkinnedAssetAsyncPropertyLockType::ReadWrite)
+	{
+		checkf((uint64)AsyncProperties <= (uint64)TNumericLimits<uint32>::Max(), TEXT("Cannot have more than 32 async properties per derived class."));
+		Super::AcquireAsyncProperty((uint64)AsyncProperties << 32ull, LockType);
+	}
+
+	/** Release properties that should not be modified/accessed during async build. */
+	template <typename EnumType UE_REQUIRES(TIsEnum<EnumType>::Value)>
+	void ReleaseAsyncProperty(const EnumType AsyncProperties, ESkinnedAssetAsyncPropertyLockType LockType = ESkinnedAssetAsyncPropertyLockType::ReadWrite)
+	{
+		checkf((uint64)AsyncProperties <= (uint64)TNumericLimits<uint32>::Max(), TEXT("Cannot have more than 32 async properties per derived class."));
+		Super::ReleaseAsyncProperty((uint64)AsyncProperties << 32ull, LockType);
+	}
+
 	/**
 	 * Wait for the asset to finish compilation to protect internal skinned asset data from race conditions during async build.
 	 * This should be called before accessing all async accessible properties.
@@ -377,6 +393,20 @@ private:
 	enum class EAsyncProperties : uint32;
 	FRIEND_ENUM_CLASS_FLAGS(EAsyncProperties);
 
+	/** Lock properties that should not be modified/accessed during async build. */
+	template <>
+	void AcquireAsyncProperty<EAsyncProperties>(const EAsyncProperties AsyncProperties, ESkinnedAssetAsyncPropertyLockType LockType)
+	{
+		Super::AcquireAsyncProperty((uint64)AsyncProperties, LockType);
+	}
+
+	/** Release properties that should not be modified/accessed during async build. */
+	template <>
+	void ReleaseAsyncProperty<EAsyncProperties>(const EAsyncProperties  AsyncProperties, ESkinnedAssetAsyncPropertyLockType LockType)
+	{
+		Super::ReleaseAsyncProperty((uint64)AsyncProperties, LockType);
+	}
+
 	/**
 	 * Wait for the asset to finish compilation to protect internal skinned asset data from race conditions during async build.
 	 * This should be called before accessing all async accessible properties.
