Chaos Outfit Asset:
- Fixed a crash in building a new Outfit when releasing rendering resources.
- Improved the SKM component refresh whenever a simulation model is modified.
- Added some early exits to RestoreClothingFromBackup().

[RN] Chaos Outfit Asset - Fixed crash in building a new Outfit when releasing rendering resources.

diff --git a/Engine/Plugins/ChaosClothAsset/Source/ChaosClothAssetEngine/Private/ChaosClothAsset/ClothAsset.cpp b/Engine/Plugins/ChaosClothAsset/Source/ChaosClothAssetEngine/Private/ChaosClothAsset/ClothAsset.cpp
--- a/Engine/Plugins/ChaosClothAsset/Source/ChaosClothAssetEngine/Private/ChaosClothAsset/ClothAsset.cpp
+++ b/Engine/Plugins/ChaosClothAsset/Source/ChaosClothAssetEngine/Private/ChaosClothAsset/ClothAsset.cpp
@@ -327,7 +327,7 @@ void UChaosClothAsset::PostEditChangeProperty(FPropertyChangedEvent& PropertyCha
 {
 	if (PropertyChangedEvent.GetPropertyName() == GET_MEMBER_NAME_CHECKED(UChaosClothAsset, PhysicsAsset))
 	{
-		ReregisterComponents();
+		OnAssetChanged();
 	}
 	Super::PostEditChangeProperty(PropertyChangedEvent);
 }
@@ -668,6 +668,9 @@ void UChaosClothAsset::Build(TArray<FChaosClothAssetLodTransitionDataCache>* InO
 {
 	using namespace UE::Chaos::ClothAsset;
 
+	// Unregister dependent components, the context will reregister them at the end of the scope
+	const FMultiComponentReregisterContext MultiComponentReregisterContext(GetDependentComponents());
+
 #if WITH_EDITOR
 	FSkinnedAssetAsyncBuildScope AsyncBuildScope(this);
 
@@ -710,8 +713,9 @@ void UChaosClothAsset::Build(TArray<FChaosClothAssetLodTransitionDataCache>* InO
 	FinishBuildInternal(Context);
 #endif
 
-	// Re-register any components using this asset to restart the simulation with the updated asset
-	ReregisterComponents();
+	// Update any components using this asset
+	constexpr bool bReregisterComponents = false;  // Do not reregister twice, this is already done at the function scope
+	OnAssetChanged(bReregisterComponents);
 }
 
 #if WITH_EDITOR
diff --git a/Engine/Plugins/ChaosClothAsset/Source/ChaosClothAssetEngine/Private/ChaosClothAsset/ClothAssetBase.cpp b/Engine/Plugins/ChaosClothAsset/Source/ChaosClothAssetEngine/Private/ChaosClothAsset/ClothAssetBase.cpp
--- a/Engine/Plugins/ChaosClothAsset/Source/ChaosClothAssetEngine/Private/ChaosClothAsset/ClothAssetBase.cpp
+++ b/Engine/Plugins/ChaosClothAsset/Source/ChaosClothAssetEngine/Private/ChaosClothAsset/ClothAssetBase.cpp
@@ -647,30 +647,54 @@ void UChaosClothAssetBase::OnPropertyChanged() const
 	}
 }
 
+void UChaosClothAssetBase::OnAssetChanged(const bool bReregisterComponents) const
+{
+	// Context will go out of scope, causing the components to be re-registered, but only when bReregisterComponents is true
+	const FMultiComponentReregisterContext MultiComponentReregisterContext(bReregisterComponents ? GetDependentComponents() : TArray<UActorComponent*>());
+
+#if WITH_EDITOR
+	for (TObjectIterator<UChaosClothAssetSKMClothingAsset> ObjectIterator; ObjectIterator; ++ObjectIterator)
+	{
+		if (UChaosClothAssetSKMClothingAsset* const ClothingAsset = *ObjectIterator)
+		{
+			if (ClothingAsset->GetAsset() == this)
+			{
+				constexpr bool bReregisterComponentsFromclothingAsset = false;  // Reregistration is done at function scope instead
+				ClothingAsset->OnAssetChanged(bReregisterComponentsFromclothingAsset);
+			}
+		}
+	}
+#endif
+}
+
 void UChaosClothAssetBase::ReregisterComponents() const
 {
-	// Reregister cloth components
+	// Context goes out of scope, causing the components to be re-registered
+	FMultiComponentReregisterContext MultiComponentReregisterContext(GetDependentComponents());
+}
+
+TArray<UActorComponent*> UChaosClothAssetBase::GetDependentComponents() const
+{
+	TArray<UActorComponent*> DependentComponents;
+
+	// Find cloth components
 	for (TObjectIterator<UChaosClothComponent> ObjectIterator; ObjectIterator; ++ObjectIterator)
 	{
 		if (UChaosClothComponent* const Component = *ObjectIterator)
 		{
 			if (Component->GetAsset() == this)
 			{
-				const FComponentReregisterContext Context(Component);  // Context goes out of scope, causing the Component to be re-registered
+				DependentComponents.Emplace(Component);
 			}
 		}
 	}
-	// Reregister skeletal mesh components
+	// Find skeletal mesh components
 	for (TObjectIterator<UChaosClothAssetSKMClothingAsset> ObjectIterator; ObjectIterator; ++ObjectIterator)
 	{
 		if (UChaosClothAssetSKMClothingAsset* const ClothingAsset = *ObjectIterator)
 		{
 			if (ClothingAsset->GetAsset() == this)
 			{
-#if WITH_EDITOR
-				ClothingAsset->OnAssetChanged();  // ReregisterComponents is called, therefore the asset must have changed, this will update the Clothing Data and redo the bindings
-#endif
-
 				if (USkeletalMesh* const OwnerMesh = Cast<USkeletalMesh>(ClothingAsset->GetOuter()))
 				{
 					for (TObjectIterator<USkeletalMeshComponent> It; It; ++It)
@@ -679,8 +703,7 @@ void UChaosClothAssetBase::ReregisterComponents() const
 						{
 							if (Component->GetSkeletalMeshAsset() == OwnerMesh)
 							{
-								FComponentReregisterContext Context(Component);
-								// Context goes out of scope, causing Component to be re-registered
+								DependentComponents.AddUnique(Component);  // Using AddUnique here since multiple SKMClothingAssets can have the same owner asset
 							}
 						}
 					}
@@ -688,4 +711,5 @@ void UChaosClothAssetBase::ReregisterComponents() const
 			}
 		}
 	}
+	return DependentComponents;
 }
diff --git a/Engine/Plugins/ChaosClothAsset/Source/ChaosClothAssetEngine/Private/ChaosClothAsset/ClothAssetSKMClothingAsset.cpp b/Engine/Plugins/ChaosClothAsset/Source/ChaosClothAssetEngine/Private/ChaosClothAsset/ClothAssetSKMClothingAsset.cpp
--- a/Engine/Plugins/ChaosClothAsset/Source/ChaosClothAssetEngine/Private/ChaosClothAsset/ClothAssetSKMClothingAsset.cpp
+++ b/Engine/Plugins/ChaosClothAsset/Source/ChaosClothAssetEngine/Private/ChaosClothAsset/ClothAssetSKMClothingAsset.cpp
@@ -77,7 +77,7 @@ namespace UE::Chaos::ClothAsset::Private
 	// Modified version of FLODUtilities::RestoreClothingFromBackup to only rebind specified assets
 	static void RestoreClothingFromBackup(USkeletalMesh& SkeletalMesh, const UChaosClothAssetSKMClothingAsset& ClothingAsset, TArray<ClothingAssetUtils::FClothingAssetMeshBinding>& ClothingBindings)
 	{
-		if (SkeletalMesh.GetImportedModel())
+		if (SkeletalMesh.GetImportedModel() && ClothingAsset.GetClothSimulationModelIndex() != INDEX_NONE)
 		{
 			auto RestoreClothingLODFromBackup = [&SkeletalMesh, &ClothingAsset](const int32 LODIndex, TArray<ClothingAssetUtils::FClothingAssetMeshBinding>& ClothingBindings)
 				{
@@ -87,11 +87,11 @@ namespace UE::Chaos::ClothAsset::Private
 						FSkeletalMeshLODModel& LODModel = SkeletalMesh.GetImportedModel()->LODModels[LODIndex];
 						for (ClothingAssetUtils::FClothingAssetMeshBinding& Binding : ClothingBindings)
 						{
-							for (int32 SectionIndex = 0; SectionIndex < LODModel.Sections.Num(); ++SectionIndex)
+							if (Binding.LODIndex == LODIndex && Binding.Asset == &ClothingAsset)
 							{
-								if (LODModel.Sections[SectionIndex].OriginalDataSectionIndex == Binding.SectionIndex)
+								for (int32 SectionIndex = 0; SectionIndex < LODModel.Sections.Num(); ++SectionIndex)
 								{
-									if (Binding.LODIndex == LODIndex && Binding.Asset == &ClothingAsset)
+									if (LODModel.Sections[SectionIndex].OriginalDataSectionIndex == Binding.SectionIndex)
 									{
 										if (Binding.Asset->BindToSkeletalMesh(&SkeletalMesh, Binding.LODIndex, SectionIndex, Binding.AssetInternalLodIndex))
 										{
@@ -100,8 +100,8 @@ namespace UE::Chaos::ClothAsset::Private
 											SectionUserData.CorrespondClothAssetIndex = LODModel.Sections[SectionIndex].CorrespondClothAssetIndex;
 											SectionUserData.ClothingData = LODModel.Sections[SectionIndex].ClothingData;
 										}
+										break;  // Found the section, next binding
 									}
-									break;  // Found the section, next binding
 								}
 							}
 						}
@@ -140,9 +140,11 @@ void UChaosClothAssetSKMClothingAsset::SetAsset(const UChaosClothAssetBase* InAs
 	}
 }
 
-void UChaosClothAssetSKMClothingAsset::OnAssetChanged()
+void UChaosClothAssetSKMClothingAsset::OnAssetChanged(const bool bReregisterComponents)
 {
 	// If the asset has changed, check whether the current specified model still exists in this new asset, otherwise update the model guid
+	const FGuid ClothSimulationModelGuid = GetClothSimulationModelGuid();
+	FGuid ModelGuid;
 	int32 ModelIndex = INDEX_NONE;
 
 	if (Asset)
@@ -150,10 +152,9 @@ void UChaosClothAssetSKMClothingAsset::OnAssetChanged()
 		if (const int32 NumClothSimulationModels = Asset->GetNumClothSimulationModels())
 		{
 			// Find a matching model GUID
-			const FGuid ClothSimulationModelGuid = GetClothSimulationModelGuid();
 			for (ModelIndex = 0; ModelIndex < NumClothSimulationModels; ++ModelIndex)
 			{
-				if (ClothSimulationModelGuid == Asset->GetAssetGuid(ModelIndex))
+				if (Asset->GetAssetGuid(ModelIndex) == ClothSimulationModelGuid)
 				{
 					break;
 				}
@@ -178,11 +179,16 @@ void UChaosClothAssetSKMClothingAsset::OnAssetChanged()
 					ModelIndex = 0;
 				}
 			}
+
+			if (ModelIndex != INDEX_NONE)
+			{
+				ModelGuid = Asset->GetAssetGuid(ModelIndex);
+			}
 		}
 	}
 
 	// Notify the model has changed
-	if (ClothSimulationModelIndex != ModelIndex)
+	if (ClothSimulationModelIndex != ModelIndex || ClothSimulationModelGuid != ModelGuid)
 	{
 		ClothSimulationModelIndex = ModelIndex;
 
@@ -193,21 +199,23 @@ void UChaosClothAssetSKMClothingAsset::OnAssetChanged()
 		CalculateLODHasAnySimulationMeshData();
 	}
 
-	// Reregister components
-	if (const USkeletalMesh* const OwnerMesh = Cast<USkeletalMesh>(GetOuter()))
+	if (bReregisterComponents)
 	{
-		for (TObjectIterator<USkeletalMeshComponent> It; It; ++It)
+		if (const USkeletalMesh* const OwnerMesh = Cast<USkeletalMesh>(GetOuter()))
 		{
-			if (USkeletalMeshComponent* const Component = *It)
+			for (TObjectIterator<USkeletalMeshComponent> It; It; ++It)
 			{
-				if (Component->GetSkeletalMeshAsset() == OwnerMesh)
+				if (USkeletalMeshComponent* const Component = *It)
 				{
-					FComponentReregisterContext Context(Component);
-					// Context goes out of scope, causing Component to be re-registered
+					if (Component->GetSkeletalMeshAsset() == OwnerMesh)
+					{
+						FComponentReregisterContext Context(Component);
+						// Context goes out of scope, causing Component to be re-registered
+					}
 				}
 			}
-		}
-	}  // TODO: Move UClothingAssetCommon::ReregisterComponentsUsingClothing() to a common place (in Utils/ClothingMeshUtils?)
+		}  // TODO: Move UClothingAssetCommon::ReregisterComponentsUsingClothing() to a common place (in Utils/ClothingMeshUtils?)
+	}
 }
 
 void UChaosClothAssetSKMClothingAsset::PostEditChangeProperty(FPropertyChangedEvent& PropertyChangedEvent)
diff --git a/Engine/Plugins/ChaosClothAsset/Source/ChaosClothAssetEngine/Public/ChaosClothAsset/ClothAssetBase.h b/Engine/Plugins/ChaosClothAsset/Source/ChaosClothAssetEngine/Public/ChaosClothAsset/ClothAssetBase.h
--- a/Engine/Plugins/ChaosClothAsset/Source/ChaosClothAssetEngine/Public/ChaosClothAsset/ClothAssetBase.h
+++ b/Engine/Plugins/ChaosClothAsset/Source/ChaosClothAssetEngine/Public/ChaosClothAsset/ClothAssetBase.h
@@ -12,6 +12,7 @@
 class FSkeletalMeshRenderData;
 class FSkinnedAssetCompilationContext;
 class UDataflow;
+class UActorComponent;
 struct FChaosClothSimulationModel;
 struct FManagedArrayCollection;
 
@@ -278,9 +279,16 @@ protected:
 	/** Update all properties for components currently using this asset. */
 	UE_API void OnPropertyChanged() const;
 
+	/** Update all components currently using this asset. */
+	UE_API void OnAssetChanged(const bool bReregisterComponents = true) const;
+
 	/** Reregister all components using this asset to reset the simulation in case anything has changed. */
+	UE_DEPRECATED(5.7, "Use OnAssetChanged instead.")
 	UE_API void ReregisterComponents() const;
 
+	/** Return all components currently using this asset. */
+	UE_API TArray<UActorComponent*> GetDependentComponents() const;
+
 	/** Lock properties that should not be modified/accessed during async build. */
 	template <typename EnumType UE_REQUIRES(TIsEnum<EnumType>::Value)>
 	void AcquireAsyncProperty(const EnumType AsyncProperties, ESkinnedAssetAsyncPropertyLockType LockType = ESkinnedAssetAsyncPropertyLockType::ReadWrite)
diff --git a/Engine/Plugins/ChaosClothAsset/Source/ChaosClothAssetEngine/Public/ChaosClothAsset/ClothAssetSKMClothingAsset.h b/Engine/Plugins/ChaosClothAsset/Source/ChaosClothAssetEngine/Public/ChaosClothAsset/ClothAssetSKMClothingAsset.h
--- a/Engine/Plugins/ChaosClothAsset/Source/ChaosClothAssetEngine/Public/ChaosClothAsset/ClothAssetSKMClothingAsset.h
+++ b/Engine/Plugins/ChaosClothAsset/Source/ChaosClothAssetEngine/Public/ChaosClothAsset/ClothAssetSKMClothingAsset.h
@@ -23,12 +23,16 @@ public:
 #if WITH_EDITOR
 	/**
 	 * Assign a Cloth or an Outfit Asset to this clothing data.
-	 * @Note this is only possible WITH_EDITOR since the Skeletal Mesh might need to be rebuild when the asset is changed.
+	 * @note This is only possible WITH_EDITOR since the Skeletal Mesh might need to be rebuild when the asset is changed.
+	 * @param InAsset The new asset to set.
 	 */
 	UE_API void SetAsset(const UChaosClothAssetBase* InAsset);
 
-	/** Call this function whenever the Cloth or Outfit Asset has changed (e.g. once rebuilt in Dataflow). */
-	UE_API void OnAssetChanged();
+	/**
+	 * Call this function whenever the Cloth or Outfit Asset has changed (e.g. once rebuilt in Dataflow).
+	 * @param bReregisterComponents Whether to reregister the dependent components. When set to false, the registration must be carried outside of this function call.
+	 */
+	UE_API void OnAssetChanged(const bool bReregisterComponents = true);
 #endif
 	
 	/** Return the underlaying Cloth or Outfit Asset used in this clothing data. */
diff --git a/Engine/Plugins/Experimental/ChaosOutfitAsset/Source/ChaosOutfitAssetEngine/Private/ChaosOutfitAsset/OutfitAsset.cpp b/Engine/Plugins/Experimental/ChaosOutfitAsset/Source/ChaosOutfitAssetEngine/Private/ChaosOutfitAsset/OutfitAsset.cpp
--- a/Engine/Plugins/Experimental/ChaosOutfitAsset/Source/ChaosOutfitAssetEngine/Private/ChaosOutfitAsset/OutfitAsset.cpp
+++ b/Engine/Plugins/Experimental/ChaosOutfitAsset/Source/ChaosOutfitAssetEngine/Private/ChaosOutfitAsset/OutfitAsset.cpp
@@ -52,6 +52,9 @@ void UChaosOutfitAsset::Build(const TObjectPtr<const UChaosOutfit> InOutfit, UE:
 {
 	using namespace UE::Chaos::OutfitAsset;
 
+	// Unregister dependent components, the context will reregister them at the end of the scope
+	const FMultiComponentReregisterContext MultiComponentReregisterContext(GetDependentComponents());
+
 	// Stop the rendering
 	ReleaseResources();
 
@@ -188,6 +191,8 @@ void UChaosOutfitAsset::Build(const TObjectPtr<const UChaosOutfit> InOutfit, UE:
 		TempSkeletalMeshRenderData->LODRenderData[0].StaticVertexBuffers.PositionVertexBuffer.Init(0);  //  Required for serialization
 		TempSkeletalMeshRenderData->LODRenderData[0].StaticVertexBuffers.StaticMeshVertexBuffer.Init(0, 0);  //  Required for serialization
 	}
+
+	ReleaseResourcesFence.Wait();  // Make sure the release resources fence has completed before deleting/replacing the render data
 	SetResourceForRendering(MoveTemp(TempSkeletalMeshRenderData));
 
 	CalculateBounds();
@@ -201,8 +206,9 @@ void UChaosOutfitAsset::Build(const TObjectPtr<const UChaosOutfit> InOutfit, UE:
 		InitResources();
 	}
 
-	// Re-register any components using this asset to restart the simulation and update the leader bone map
-	ReregisterComponents();
+	// Update any components using this asset
+	constexpr bool bReregisterComponents = false;  // Do not reregister twice, this is already done at the function scope
+	OnAssetChanged(bReregisterComponents);
 }
 
 void UChaosOutfitAsset::Serialize(FArchive& Ar)
