Chaos SKM Cloth Asset - Fixed re-binding when rebuilding the cloth asset.

#rnx

diff --git a/Engine/Plugins/ChaosClothAsset/Source/ChaosClothAssetEngine/Private/ChaosClothAsset/ClothAssetSKMClothingAsset.cpp b/Engine/Plugins/ChaosClothAsset/Source/ChaosClothAssetEngine/Private/ChaosClothAsset/ClothAssetSKMClothingAsset.cpp
--- a/Engine/Plugins/ChaosClothAsset/Source/ChaosClothAssetEngine/Private/ChaosClothAsset/ClothAssetSKMClothingAsset.cpp
+++ b/Engine/Plugins/ChaosClothAsset/Source/ChaosClothAssetEngine/Private/ChaosClothAsset/ClothAssetSKMClothingAsset.cpp
@@ -245,23 +245,44 @@ bool UChaosClothAssetSKMClothingAsset::BindToSkeletalMesh(USkeletalMesh* Skeleta
 	{
 		return false;
 	}
+
+	// When FScopedSkeletalMeshPostEditChange goes out of scope it causes the postedit change and components to be re-registered and the mesh to rebuild
+	FScopedSkeletalMeshPostEditChange SkeletalMeshPostEditChange(SkeletalMesh);
+
+	// Get the original render section
+	FSkeletalMeshLODModel& SkeletalMeshLODModel = SkeletalMesh->GetImportedModel()->LODModels[MeshLodIndex];
+	FSkelMeshSection& SkelMeshSection = SkeletalMeshLODModel.Sections[SectionIndex];
+
+	// Clear the proxy deformer mappings data now in case the cloth simulation model is invalid
+	SkelMeshSection.ClothMappingDataLODs.Reset();
+
+	// Set the asset guid
+	SkelMeshSection.ClothingData.AssetGuid = AssetGuid;
+	SkelMeshSection.ClothingData.AssetLodIndex = AssetLodIndex;
+
+	// Set the asset index, used during rendering to pick the correct sim mesh buffer
+	int32 AssetIndex = INDEX_NONE;
+	verify(SkeletalMesh->GetMeshClothingAssets().Find(this, AssetIndex));
+	SkelMeshSection.CorrespondClothAssetIndex = (int16)AssetIndex;
+
+	// Retrieve the cloth simulation model
 	if (!Asset)
 	{
 		const FText Text = FText::Format(
-			LOCTEXT("NoAssetHasBeenSet", "Failed to bind Clothing Data [{0}], because no Asset has been set."),
+			LOCTEXT("NoAssetHasBeenSet", "Clothing Data [{0}]: no Asset has been set."),
 			FText::FromString(GetName()));
 		WarningNotification(Text);
-		return false;
+		return true;  // Must return true, this is to avoid breaking the binding when the cloth simulation model is invalid
 	}
 	if (ClothSimulationModelIndex == INDEX_NONE)
 	{
 		const FText Text = FText::Format(
-			LOCTEXT("NotAClothSimulationModelOfAsset", "Failed to bind Clothing Data [{0}], because [{1}] is not an existing Cloth Simulation Model of Asset [{2}]."),
+			LOCTEXT("NotAClothSimulationModelOfAsset", "Clothing Data [{0}]: [{1}] is not an existing Cloth Simulation Model of Asset [{2}]."),
 			FText::FromString(GetName()),
 			FText::FromName(GetClothSimulationModelName()),
 			FText::FromName(Asset->GetFName()));
 		WarningNotification(Text);
-		return false;
+		return true;  // Must return true, this is to avoid breaking the binding when the cloth simulation model is invalid
 	}
 	const TSharedPtr<const FChaosClothSimulationModel> ClothSimulationModel = Asset->GetClothSimulationModel(ClothSimulationModelIndex);
 	if (!ClothSimulationModel ||
@@ -269,29 +290,25 @@ bool UChaosClothAssetSKMClothingAsset::BindToSkeletalMesh(USkeletalMesh* Skeleta
 		!ClothSimulationModel->GetNumTriangles(AssetLodIndex))
 	{
 		const FText Text = FText::Format(
-			LOCTEXT("EmptyClothSimulationModel", "Failed to bind Clothing Data [{0}], because [{1}] is an empty Cloth Simulation Model of Asset [{2}]."),
+			LOCTEXT("EmptyClothSimulationModel", "Clothing Data [{0}]: [{1}] is an empty Cloth Simulation Model of Asset [{2}]."),
 			FText::FromString(GetName()),
 			FText::FromName(GetClothSimulationModelName()),
 			FText::FromName(Asset->GetFName()));
 		WarningNotification(Text);
-		return false;
+		return true;  // Must return true, this is to avoid breaking the binding when the cloth simulation model is invalid
 	}
 	if (!ClothSimulationModel->IsValidLodIndex(AssetLodIndex))
 	{
 		const FText Text = FText::Format(
-			LOCTEXT("NoLODInClothSimulationModel", "Failed to bind Clothing Data [{0}], because Cloth Simulation Model [{1}] has no LOD{2} in Asset [{3}]."),
+			LOCTEXT("NoLODInClothSimulationModel", "Clothing Data [{0}]: Cloth Simulation Model [{1}] has no LOD{2} in Asset [{3}]."),
 			FText::FromString(GetName()),
 			FText::FromName(GetClothSimulationModelName()),
 			AssetLodIndex,
 			FText::FromName(Asset->GetFName()));
 		WarningNotification(Text);
-		return false;
+		return true;  // Must return true, this is to avoid breaking the binding when the cloth simulation model is invalid
 	}
 
-	// Get the original render section
-	FSkeletalMeshLODModel& SkeletalMeshLODModel = SkeletalMesh->GetImportedModel()->LODModels[MeshLodIndex];
-	FSkelMeshSection& SkelMeshSection = SkeletalMeshLODModel.Sections[SectionIndex];
-
 	// Set the new section bone map but first verify that it doesn't exceed the maximum when adding the cloth asset bones
 	TArray<FBoneIndexType> BoneMap = SkelMeshSection.BoneMap;
 	for (const FName BoneName : ClothSimulationModel->UsedBoneNames)
@@ -316,7 +333,6 @@ bool UChaosClothAssetSKMClothingAsset::BindToSkeletalMesh(USkeletalMesh* Skeleta
 		return false;
 	}
 
-	FScopedSkeletalMeshPostEditChange SkeletalMeshPostEditChange(SkeletalMesh);  // When FScopedSkeletalMeshPostEditChange goes out of scope it causes the postedit change and components to be re-registered
 	SkelMeshSection.BoneMap = MoveTemp(BoneMap);
 
 	bool bRequireBoneChange = false;
@@ -338,11 +354,6 @@ bool UChaosClothAssetSKMClothingAsset::BindToSkeletalMesh(USkeletalMesh* Skeleta
 		SkeletalMesh->GetRefSkeleton().EnsureParentsExistAndSort(SkeletalMeshLODModel.ActiveBoneIndices);
 	}
 
-	// Set the asset guid
-	SkelMeshSection.CorrespondClothAssetIndex = ClothSimulationModelIndex;
-	SkelMeshSection.ClothingData.AssetGuid = AssetGuid;
-	SkelMeshSection.ClothingData.AssetLodIndex = AssetLodIndex;
-
 	// Calculate the deformer mappings
 	const int32 NumIndices = SkelMeshSection.NumTriangles * 3;
 	const int32 NumVertices = SkelMeshSection.SoftVertices.Num();
@@ -391,11 +402,6 @@ bool UChaosClothAssetSKMClothingAsset::BindToSkeletalMesh(USkeletalMesh* Skeleta
 	SkelMeshSection.ClothMappingDataLODs.SetNum(1);
 	SkelMeshSection.ClothMappingDataLODs[0] = MoveTemp(MeshToMeshData);
 
-	// Set the asset index, used during rendering to pick the correct sim mesh buffer
-	int32 AssetIndex = INDEX_NONE;
-	verify(SkeletalMesh->GetMeshClothingAssets().Find(this, AssetIndex));  // A cloth asset can be added multiple times, but it will need a different clothing data each time
-	SkelMeshSection.CorrespondClothAssetIndex = AssetIndex;
-
 	// Update the extra cloth deformer mapping LOD bias using this cloth entry
 	// TODO: Implement a generic form of LOD Bias mapping in ClothingAssetUtils (need to rework the clothing asset interface for this)
 	//UpdateLODBiasMappings(SkeletalMesh, MeshLodIndex, SectionIndex);
@@ -435,6 +441,8 @@ void UChaosClothAssetSKMClothingAsset::OnModelChanged()
 	// Unbind/rebind clothing with the new GUID
 	if (USkeletalMesh* const SkeletalMesh = Cast<USkeletalMesh>(GetOuter()))
 	{
+		FScopedSkeletalMeshPostEditChange ScopedPostEditChange(SkeletalMesh);  // Prevent the skeletal mesh from rebuilding until the unbind/rebind operation is complete
+
 		TArray<ClothingAssetUtils::FClothingAssetMeshBinding> ClothingBindings;
 		UnbindClothingAndBackup(*SkeletalMesh, *this, ClothingBindings);
 		RestoreClothingFromBackup(*SkeletalMesh, *this, ClothingBindings);
