Update Sim Accessory Mesh creation to handle compatible, but not identical, skeletal meshes.

Refactored MergeClothCollectionsNode code so it could be shared.

diff --git a/Engine/Plugins/ChaosClothAsset/Source/ChaosClothAsset/Private/ChaosClothAsset/ClothGeometryTools.cpp b/Engine/Plugins/ChaosClothAsset/Source/ChaosClothAsset/Private/ChaosClothAsset/ClothGeometryTools.cpp
--- a/Engine/Plugins/ChaosClothAsset/Source/ChaosClothAsset/Private/ChaosClothAsset/ClothGeometryTools.cpp
+++ b/Engine/Plugins/ChaosClothAsset/Source/ChaosClothAsset/Private/ChaosClothAsset/ClothGeometryTools.cpp
@@ -1888,137 +1888,4 @@ PRAGMA_ENABLE_DEPRECATION_WARNINGS
 			}
 		}
 	}
-
-	int32 FClothGeometryTools::CopySimMeshToSimAccessoryMesh(const FName& AccessoryMeshName, FCollectionClothFacade& ToCloth, const FCollectionClothConstFacade& FromCloth, bool bUseSimImportVertexID)
-	{
-		auto GenerateUniqueAccessoryMeshName = [&AccessoryMeshName, &ToCloth]()
-			{
-				int32 MaxExistingNumber = -1;
-				bool bFoundMatchWithIndex = false;
-				for (const FName& Name : ToCloth.GetSimAccessoryMeshName())
-				{
-					if (Name.GetComparisonIndex() == AccessoryMeshName.GetComparisonIndex())
-					{
-						MaxExistingNumber = FMath::Max(MaxExistingNumber, Name.GetNumber());
-					}
-					bFoundMatchWithIndex = bFoundMatchWithIndex || Name == AccessoryMeshName;
-				}
-
-				if (!bFoundMatchWithIndex)
-				{
-					return AccessoryMeshName;
-				}
-				FName UniqueName = AccessoryMeshName;
-				UniqueName.SetNumber(MaxExistingNumber + 1);
-				return UniqueName;
-			};
-
-		if (bUseSimImportVertexID) 
-		{
-			if (ToCloth.IsValid(EClothCollectionExtendedSchemas::Import) && FromCloth.IsValid(EClothCollectionExtendedSchemas::Import))
-			{
-				// Start with existing sim data
-				TArray<FVector3f> Positions(ToCloth.GetSimPosition3D());
-				TArray<FVector3f> Normals(ToCloth.GetSimNormal());
-				TArray<TArray<int32>> BoneIndices(ToCloth.GetSimBoneIndices());
-				TArray<TArray<float>> BoneWeights(ToCloth.GetSimBoneWeights());
-
-				TConstArrayView<int32> ToImportVertex = ToCloth.GetSimImportVertexID();
-				TConstArrayView<int32> FromImportVertex = FromCloth.GetSimImportVertexID();
-
-				TMap<int32, TArray<int32>> ToImportToSim2DLookup;
-				ToImportToSim2DLookup.Reserve(ToImportVertex.Num());
-				for (int32 ToSim2DIndex = 0; ToSim2DIndex < ToImportVertex.Num(); ++ToSim2DIndex)
-				{
-					ToImportToSim2DLookup.FindOrAdd(ToImportVertex[ToSim2DIndex]).Add(ToSim2DIndex);
-				}
-
-				TConstArrayView<int32> ToSim3DLookup = ToCloth.GetSimVertex3DLookup();
-				TConstArrayView<int32> FromSim3DLookup = FromCloth.GetSimVertex3DLookup();
-				TConstArrayView<FVector3f> FromPositions = FromCloth.GetSimPosition3D();
-				TConstArrayView<FVector3f> FromNormals = FromCloth.GetSimNormal();
-				TConstArrayView<TArray<int32>> FromBoneIndices = FromCloth.GetSimBoneIndices();
-				TConstArrayView<TArray<float>> FromBoneWeights = FromCloth.GetSimBoneWeights();
-				for (int32 FromSim2DIndex = 0; FromSim2DIndex < FromImportVertex.Num(); ++FromSim2DIndex)
-				{
-					const int32 FromSim3DIndex = FromSim3DLookup[FromSim2DIndex];
-					if (FromPositions.IsValidIndex(FromSim3DIndex))
-					{
-						if (const TArray<int32>* const ToSim2DIndices = ToImportToSim2DLookup.Find(FromImportVertex[FromSim2DIndex]))
-						{
-							// Found matching index on ToCloth.
-							// Copy data between associated sim3d vertices.
-							// NOTE: this will just last one wins for vertices that did not seam the same way between the two meshes. Hopefully that's good enough.
-							for (const int32 ToSim2DIndex : *ToSim2DIndices)
-							{
-								const int32 ToSim3DIndex = ToSim3DLookup[ToSim2DIndex];
-								if (Positions.IsValidIndex(ToSim3DIndex))
-								{
-									Positions[ToSim3DIndex] = FromPositions[FromSim3DIndex];
-									Normals[ToSim3DIndex] = FromNormals[FromSim3DIndex];
-									BoneIndices[ToSim3DIndex] = FromBoneIndices[FromSim3DIndex];
-									BoneWeights[ToSim3DIndex] = FromBoneWeights[FromSim3DIndex];
-								}
-							}
-						}
-					}
-				}
-
-				FCollectionClothSimAccessoryMeshFacade AccessoryFacade = ToCloth.AddGetSimAccessoryMesh();
-				AccessoryFacade.Initialize(GenerateUniqueAccessoryMeshName(), TConstArrayView<FVector3f>(Positions), TConstArrayView<FVector3f>(Normals), TConstArrayView<TArray<int32>>(BoneIndices), TConstArrayView<TArray<float>>(BoneWeights));
-				return AccessoryFacade.GetSimAccessoryMeshIndex();
-			}
-
-			// no valid import data
-			return INDEX_NONE;
-		}
-
-		if (FromCloth.GetNumSimVertices3D() >= ToCloth.GetNumSimVertices3D())
-		{
-			// Can just initialize directly 
-			const int32 NumToSimVertices3D = ToCloth.GetNumSimVertices3D();
-			FCollectionClothSimAccessoryMeshFacade AccessoryFacade = ToCloth.AddGetSimAccessoryMesh();
-			AccessoryFacade.Initialize(GenerateUniqueAccessoryMeshName(), FromCloth.GetSimPosition3D().Left(NumToSimVertices3D), FromCloth.GetSimNormal().Left(NumToSimVertices3D), FromCloth.GetSimBoneIndices().Left(NumToSimVertices3D), FromCloth.GetSimBoneWeights().Left(NumToSimVertices3D));
-			return AccessoryFacade.GetSimAccessoryMeshIndex();
-		}
-
-		// Fill in tail data with ToCloth's sim mesh data.
-		const int32 NumToSimVertices3D = ToCloth.GetNumSimVertices3D();
-		TArray<FVector3f> Positions;
-		TArray<FVector3f> Normals;
-		TArray<TArray<int32>> BoneIndices;
-		TArray<TArray<float>> BoneWeights;
-		Positions.Reserve(NumToSimVertices3D);
-		Normals.Reserve(NumToSimVertices3D);
-		BoneIndices.Reserve(NumToSimVertices3D);
-		BoneWeights.Reserve(NumToSimVertices3D);
-
-		TConstArrayView<FVector3f> FromPositions = FromCloth.GetSimPosition3D();
-		TConstArrayView<FVector3f> FromNormals = FromCloth.GetSimNormal();
-		TConstArrayView<TArray<int32>> FromBoneIndices = FromCloth.GetSimBoneIndices();
-		TConstArrayView<TArray<float>> FromBoneWeights = FromCloth.GetSimBoneWeights();
-		for (int32 Index = 0; Index < FromCloth.GetNumSimVertices3D(); ++Index)
-		{
-			Positions.Emplace(FromPositions[Index]);
-			Normals.Emplace(FromNormals[Index]);
-			BoneIndices.Emplace(FromBoneIndices[Index]);
-			BoneWeights.Emplace(FromBoneWeights[Index]);
-		}
-		TConstArrayView<FVector3f> ToPositions = ToCloth.GetSimPosition3D();
-		TConstArrayView<FVector3f> ToNormals = ToCloth.GetSimNormal();
-		TConstArrayView<TArray<int32>> ToBoneIndices = ToCloth.GetSimBoneIndices();
-		TConstArrayView<TArray<float>> ToBoneWeights = ToCloth.GetSimBoneWeights();
-		for (int32 Index = FromCloth.GetNumSimVertices3D(); Index < NumToSimVertices3D; ++Index)
-		{
-			Positions.Emplace(ToPositions[Index]);
-			Normals.Emplace(ToNormals[Index]);
-			BoneIndices.Emplace(ToBoneIndices[Index]);
-			BoneWeights.Emplace(ToBoneWeights[Index]);
-		}
-
-		FCollectionClothSimAccessoryMeshFacade AccessoryFacade = ToCloth.AddGetSimAccessoryMesh();
-		AccessoryFacade.Initialize(GenerateUniqueAccessoryMeshName(), TConstArrayView<FVector3f>(Positions), TConstArrayView<FVector3f>(Normals), TConstArrayView<TArray<int32>>(BoneIndices), TConstArrayView<TArray<float>>(BoneWeights));
-		return AccessoryFacade.GetSimAccessoryMeshIndex();
-
-	}
 }  // End namespace UE::Chaos::ClothAsset
diff --git a/Engine/Plugins/ChaosClothAsset/Source/ChaosClothAsset/Public/ChaosClothAsset/ClothGeometryTools.h b/Engine/Plugins/ChaosClothAsset/Source/ChaosClothAsset/Public/ChaosClothAsset/ClothGeometryTools.h
--- a/Engine/Plugins/ChaosClothAsset/Source/ChaosClothAsset/Public/ChaosClothAsset/ClothGeometryTools.h
+++ b/Engine/Plugins/ChaosClothAsset/Source/ChaosClothAsset/Public/ChaosClothAsset/ClothGeometryTools.h
@@ -180,12 +180,6 @@ namespace UE::Chaos::ClothAsset
 		 */
 		static UE_API void ApplyProxyDeformer(const TSharedRef<FManagedArrayCollection>& ClothCollection, bool bIgnoreSkinningBlend);
 		static UE_API void ApplyProxyDeformer(FCollectionClothFacade& ClothFacade, bool bIgnoreSkinningBlend, const TConstArrayView<FVector3f>& SimPositions, const TConstArrayView<FVector3f>& SimNormals);
-
-		/** Create a new SimAccessoryMesh by copying the SimMesh data from one collection into another.
-		 *  Optionally use SimImportVertexID to match data rather than a straight copy of SimVertices3D.
-		 *  Any unmatched vertices will copy the data from the existing Sim Mesh.
-		 *  Returns the newly created accessory mesh's index */
-		static UE_API int32 CopySimMeshToSimAccessoryMesh(const FName& AccessoryMeshName, FCollectionClothFacade& ToCloth, const FCollectionClothConstFacade& FromCloth, bool bUseSimImportVertexID);
 	};
 }  // End namespace UE::Chaos::ClothAsset
 
diff --git a/Engine/Plugins/ChaosClothAsset/Source/ChaosClothAssetEngine/Private/ChaosClothAsset/ClothEngineTools.cpp b/Engine/Plugins/ChaosClothAsset/Source/ChaosClothAssetEngine/Private/ChaosClothAsset/ClothEngineTools.cpp
--- a/Engine/Plugins/ChaosClothAsset/Source/ChaosClothAssetEngine/Private/ChaosClothAsset/ClothEngineTools.cpp
+++ b/Engine/Plugins/ChaosClothAsset/Source/ChaosClothAssetEngine/Private/ChaosClothAsset/ClothEngineTools.cpp
@@ -7,7 +7,11 @@
 #include "ChaosClothAsset/ClothCollectionGroup.h"
 #include "ChaosClothAsset/ClothGeometryTools.h"
 #include "ClothTetherData.h"
+#include "Engine/SkeletalMesh.h"
 #include "ReferenceSkeleton.h"
+#include "UObject/SoftObjectPath.h"
+
+#define LOCTEXT_NAMESPACE "ChaosClothAssetEngineTools"
 
 namespace UE::Chaos::ClothAsset
 {
@@ -34,6 +38,20 @@ static void AppendTetherData(FCollectionClothFacade& ClothFacade, const FClothTe
 		}
 	}
 }
+
+static void RemapBoneIndices(TArrayView<TArray<int32>> BoneIndices, const TArray<int32>& Remap)
+{
+	for (TArray<int32>& Array : BoneIndices)
+	{
+		for (int32& Index : Array)
+		{
+			if (Index != INDEX_NONE && ensure(Remap.IsValidIndex(Index)))
+			{
+				Index = Remap[Index];
+			}
+		}
+	}
+}
 }
 
 void FClothEngineTools::GenerateTethers(const TSharedRef<FManagedArrayCollection>& ClothCollection, const FName& WeightMapName, const bool bGenerateGeodesicTethers, const FVector2f& MaxDistanceValue)
@@ -272,4 +290,277 @@ int32 FClothEngineTools::CalculateReferenceBoneIndex(const TSharedRef<const FMan
 
 	return CalculateReferenceBoneIndex(UsedBones.Array(), ReferenceSkeleton);
 }
+
+bool FClothEngineTools::CalculateRemappedBoneIndicesIfCompatible(const FCollectionClothConstFacade& Cloth1, const FCollectionClothConstFacade& Cloth2, FSoftObjectPath& OutMergedSkeletalMeshPath, TArray<int32>& OutBoneIndicesRemapCloth1, TArray<int32>& OutBoneIndicesRemapCloth2, FText* OutIncompatibleErrorDetails)
+{
+
+	const FSoftObjectPath& SkeletalMeshPathName1 = Cloth1.GetSkeletalMeshSoftObjectPathName();
+	const FSoftObjectPath& SkeletalMeshPathName2 = Cloth2.GetSkeletalMeshSoftObjectPathName();
+	if (SkeletalMeshPathName1.IsNull() || SkeletalMeshPathName2.IsNull() || SkeletalMeshPathName1 == SkeletalMeshPathName2)
+	{
+		OutMergedSkeletalMeshPath = SkeletalMeshPathName1.IsNull() ? SkeletalMeshPathName2 : SkeletalMeshPathName1;
+		OutBoneIndicesRemapCloth1.Reset();
+		OutBoneIndicesRemapCloth2.Reset();
+		return true;
+	}
+
+	const USkeletalMesh* const SkeletalMesh1 = Cast<USkeletalMesh>(SkeletalMeshPathName1.TryLoad());
+	const USkeletalMesh* const SkeletalMesh2 = Cast<USkeletalMesh>(SkeletalMeshPathName2.TryLoad());
+	if (!SkeletalMesh1 || !SkeletalMesh2)
+	{
+		if (OutIncompatibleErrorDetails)
+		{
+			*OutIncompatibleErrorDetails = FText::Format(
+				LOCTEXT(
+					"IncompatibleSkeletalMeshesLoadFailureDetails",
+					"Cloth collections failed to merge due to failing to load SkeletalMesh \"{0}\" to check compatibility."),
+				!SkeletalMesh1 ? FText::FromString(SkeletalMeshPathName1.ToString()) : FText::FromString(SkeletalMeshPathName2.ToString()));
+		}
+		return false;
+	}
+
+	const FReferenceSkeleton& RefSkeleton1 = SkeletalMesh1->GetRefSkeleton();
+	const FReferenceSkeleton& RefSkeleton2 = SkeletalMesh2->GetRefSkeleton();
+
+	const FReferenceSkeleton& MergedRefSkeleton = RefSkeleton1.GetNum() >= RefSkeleton2.GetNum() ? RefSkeleton1 : RefSkeleton2;
+	const FReferenceSkeleton& RemapRefSkeleton = RefSkeleton1.GetNum() >= RefSkeleton2.GetNum() ? RefSkeleton2 : RefSkeleton1;
+	const FSoftObjectPath& MergedSkeletalMeshPath = RefSkeleton1.GetNum() >= RefSkeleton2.GetNum() ? SkeletalMeshPathName1 : SkeletalMeshPathName2;
+
+	const TArray<FMeshBoneInfo>& RemapBoneInfo = RemapRefSkeleton.GetRefBoneInfo();
+	const TArray<FTransform>& MergedBonePose = MergedRefSkeleton.GetRefBonePose();
+	const TArray<FTransform>& RemapBonePose = RemapRefSkeleton.GetRefBonePose();
+	TArray<int32> RemapIndices;
+	RemapIndices.SetNumUninitialized(RemapRefSkeleton.GetNum());
+	bool bAnyRemap = false;
+	for (int32 BoneIndex = 0; BoneIndex < RemapRefSkeleton.GetNum(); ++BoneIndex)
+	{
+		const int32 MergedBoneIndex = MergedRefSkeleton.FindBoneIndex(RemapBoneInfo[BoneIndex].Name);
+		if (MergedBoneIndex == INDEX_NONE)
+		{
+			if (OutIncompatibleErrorDetails)
+			{
+				*OutIncompatibleErrorDetails = FText::Format(
+					LOCTEXT(
+						"IncompatibleSkeletalMeshesRefBoneInfoDetails",
+						"Cloth collections failed to merge due to incompatible Skeletal Meshes, \"{0}\" and \"{1}\". Could not find bone \"{2}\" in \"{3}\"."),
+					FText::FromString(SkeletalMeshPathName1.ToString()),
+					FText::FromString(SkeletalMeshPathName2.ToString()),
+					FText::FromName(RemapBoneInfo[BoneIndex].Name),
+					FText::FromString(MergedSkeletalMeshPath.ToString()));
+			}
+
+			return false;
+		}
+		if (!RemapBonePose[BoneIndex].Equals(MergedBonePose[MergedBoneIndex]))
+		{
+			if (OutIncompatibleErrorDetails)
+			{
+				*OutIncompatibleErrorDetails = FText::Format(
+					LOCTEXT(
+						"IncompatibleSkeletalMeshesRefBonePoseDetails",
+						"Cloth collections failed to merge due to incompatible Skeletal Meshes, \"{0}\" and \"{1}\". RefBonePoses are mismatched for bone \"{2}\"."),
+					FText::FromString(SkeletalMeshPathName1.ToString()),
+					FText::FromString(SkeletalMeshPathName2.ToString()),
+					FText::FromName(RemapBoneInfo[BoneIndex].Name));
+			}
+
+			return false;
+		}
+		RemapIndices[BoneIndex] = MergedBoneIndex;
+		if (BoneIndex != MergedBoneIndex)
+		{
+			bAnyRemap = true;
+		}
+	}
+
+	if (bAnyRemap)
+	{
+		if (&RemapRefSkeleton == &RefSkeleton1)
+		{
+			OutMergedSkeletalMeshPath = SkeletalMeshPathName1;
+			OutBoneIndicesRemapCloth1 = MoveTemp(RemapIndices);
+			OutBoneIndicesRemapCloth2.Reset();
+		}
+		else
+		{
+			check(&RemapRefSkeleton == &RefSkeleton2);
+			OutMergedSkeletalMeshPath = SkeletalMeshPathName2;
+			OutBoneIndicesRemapCloth1.Reset();
+			OutBoneIndicesRemapCloth2 = MoveTemp(RemapIndices);
+		}
+	}
+	else
+	{
+		OutMergedSkeletalMeshPath.Reset();
+		OutBoneIndicesRemapCloth1.Reset();
+		OutBoneIndicesRemapCloth2.Reset();
+	}
+
+	if (OutIncompatibleErrorDetails)
+	{
+		*OutIncompatibleErrorDetails = FText();
+	}
+	return true;
+}
+
+void FClothEngineTools::RemapBoneIndices(FCollectionClothFacade& Cloth, const TArray<int32>& BoneIndicesRemap, const int32 SimVertex3DOffset, const int32 RenderVertexOffset)
+{
+	Private::RemapBoneIndices(Cloth.GetSimBoneIndices().RightChop(SimVertex3DOffset), BoneIndicesRemap);
+	Private::RemapBoneIndices(Cloth.GetRenderBoneIndices().RightChop(RenderVertexOffset), BoneIndicesRemap);
+	for (int32 AccessoryMeshIndex = 0; AccessoryMeshIndex < Cloth.GetNumSimAccessoryMeshes(); ++AccessoryMeshIndex)
+	{
+		FCollectionClothSimAccessoryMeshFacade AccessoryMesh = Cloth.GetSimAccessoryMesh(AccessoryMeshIndex);
+		Private::RemapBoneIndices(AccessoryMesh.GetSimAccessoryMeshBoneIndices().RightChop(SimVertex3DOffset), BoneIndicesRemap);
+	}
+}
+
+int32 FClothEngineTools::CopySimMeshToSimAccessoryMesh(const FName& AccessoryMeshName, FCollectionClothFacade& ToCloth, const FCollectionClothConstFacade& FromCloth, bool bUseSimImportVertexID, FText* OutIncompatibleErrorDetails)
+{
+	auto GenerateUniqueAccessoryMeshName = [&AccessoryMeshName, &ToCloth]()
+		{
+			int32 MaxExistingNumber = -1;
+			bool bFoundMatchWithIndex = false;
+			for (const FName& Name : ToCloth.GetSimAccessoryMeshName())
+			{
+				if (Name.GetComparisonIndex() == AccessoryMeshName.GetComparisonIndex())
+				{
+					MaxExistingNumber = FMath::Max(MaxExistingNumber, Name.GetNumber());
+				}
+				bFoundMatchWithIndex = bFoundMatchWithIndex || Name == AccessoryMeshName;
+			}
+
+			if (!bFoundMatchWithIndex)
+			{
+				return AccessoryMeshName;
+			}
+			FName UniqueName = AccessoryMeshName;
+			UniqueName.SetNumber(MaxExistingNumber + 1);
+			return UniqueName;
+		};
+
+	FSoftObjectPath RemapSkeletalMeshPath;
+	TArray<int32> ToClothBoneRemap, FromClothBoneRemap;
+	if (!CalculateRemappedBoneIndicesIfCompatible(ToCloth, FromCloth, RemapSkeletalMeshPath, ToClothBoneRemap, FromClothBoneRemap, OutIncompatibleErrorDetails))
+	{
+		return INDEX_NONE;
+	}
+
+	if (!ToClothBoneRemap.IsEmpty())
+	{
+		ToCloth.SetSkeletalMeshSoftObjectPathName(RemapSkeletalMeshPath);
+		RemapBoneIndices(ToCloth, ToClothBoneRemap);
+	}
+
+	if (bUseSimImportVertexID)
+	{
+		if (ToCloth.IsValid(EClothCollectionExtendedSchemas::Import) && FromCloth.IsValid(EClothCollectionExtendedSchemas::Import))
+		{
+			// Start with existing sim data
+			TArray<FVector3f> Positions(ToCloth.GetSimPosition3D());
+			TArray<FVector3f> Normals(ToCloth.GetSimNormal());
+			TArray<TArray<int32>> BoneIndices(ToCloth.GetSimBoneIndices());
+			TArray<TArray<float>> BoneWeights(ToCloth.GetSimBoneWeights());
+
+			TConstArrayView<int32> ToImportVertex = ToCloth.GetSimImportVertexID();
+			TConstArrayView<int32> FromImportVertex = FromCloth.GetSimImportVertexID();
+
+			TMap<int32, TArray<int32>> ToImportToSim2DLookup;
+			ToImportToSim2DLookup.Reserve(ToImportVertex.Num());
+			for (int32 ToSim2DIndex = 0; ToSim2DIndex < ToImportVertex.Num(); ++ToSim2DIndex)
+			{
+				ToImportToSim2DLookup.FindOrAdd(ToImportVertex[ToSim2DIndex]).Add(ToSim2DIndex);
+			}
+
+			TConstArrayView<int32> ToSim3DLookup = ToCloth.GetSimVertex3DLookup();
+			TConstArrayView<int32> FromSim3DLookup = FromCloth.GetSimVertex3DLookup();
+			TConstArrayView<FVector3f> FromPositions = FromCloth.GetSimPosition3D();
+			TConstArrayView<FVector3f> FromNormals = FromCloth.GetSimNormal();
+			TConstArrayView<TArray<int32>> FromBoneIndices = FromCloth.GetSimBoneIndices();
+			TConstArrayView<TArray<float>> FromBoneWeights = FromCloth.GetSimBoneWeights();
+			for (int32 FromSim2DIndex = 0; FromSim2DIndex < FromImportVertex.Num(); ++FromSim2DIndex)
+			{
+				const int32 FromSim3DIndex = FromSim3DLookup[FromSim2DIndex];
+				if (FromPositions.IsValidIndex(FromSim3DIndex))
+				{
+					if (const TArray<int32>* const ToSim2DIndices = ToImportToSim2DLookup.Find(FromImportVertex[FromSim2DIndex]))
+					{
+						// Found matching index on ToCloth.
+						// Copy data between associated sim3d vertices.
+						// NOTE: this will just last one wins for vertices that did not seam the same way between the two meshes. Hopefully that's good enough.
+						for (const int32 ToSim2DIndex : *ToSim2DIndices)
+						{
+							const int32 ToSim3DIndex = ToSim3DLookup[ToSim2DIndex];
+							if (Positions.IsValidIndex(ToSim3DIndex))
+							{
+								Positions[ToSim3DIndex] = FromPositions[FromSim3DIndex];
+								Normals[ToSim3DIndex] = FromNormals[FromSim3DIndex];
+								BoneIndices[ToSim3DIndex] = FromBoneIndices[FromSim3DIndex];
+								if (!FromClothBoneRemap.IsEmpty())
+								{
+									Private::RemapBoneIndices(TArrayView<TArray<int32>>(&BoneIndices[ToSim3DIndex], 1), FromClothBoneRemap);
+								}
+								BoneWeights[ToSim3DIndex] = FromBoneWeights[FromSim3DIndex];
+							}
+						}
+					}
+				}
+			}
+
+			FCollectionClothSimAccessoryMeshFacade AccessoryFacade = ToCloth.AddGetSimAccessoryMesh();
+			AccessoryFacade.Initialize(GenerateUniqueAccessoryMeshName(), TConstArrayView<FVector3f>(Positions), TConstArrayView<FVector3f>(Normals), TConstArrayView<TArray<int32>>(BoneIndices), TConstArrayView<TArray<float>>(BoneWeights));
+			return AccessoryFacade.GetSimAccessoryMeshIndex();
+		}
+
+		// no valid import data
+		return INDEX_NONE;
+	}
+
+	if (FromCloth.GetNumSimVertices3D() >= ToCloth.GetNumSimVertices3D())
+	{
+		// Can just initialize directly 
+		const int32 NumToSimVertices3D = ToCloth.GetNumSimVertices3D();
+		FCollectionClothSimAccessoryMeshFacade AccessoryFacade = ToCloth.AddGetSimAccessoryMesh();
+		AccessoryFacade.Initialize(GenerateUniqueAccessoryMeshName(), FromCloth.GetSimPosition3D().Left(NumToSimVertices3D), FromCloth.GetSimNormal().Left(NumToSimVertices3D), FromCloth.GetSimBoneIndices().Left(NumToSimVertices3D), FromCloth.GetSimBoneWeights().Left(NumToSimVertices3D));
+		return AccessoryFacade.GetSimAccessoryMeshIndex();
+	}
+
+	// Fill in tail data with ToCloth's sim mesh data.
+	const int32 NumToSimVertices3D = ToCloth.GetNumSimVertices3D();
+	TArray<FVector3f> Positions;
+	TArray<FVector3f> Normals;
+	TArray<TArray<int32>> BoneIndices;
+	TArray<TArray<float>> BoneWeights;
+	Positions.Reserve(NumToSimVertices3D);
+	Normals.Reserve(NumToSimVertices3D);
+	BoneIndices.Reserve(NumToSimVertices3D);
+	BoneWeights.Reserve(NumToSimVertices3D);
+
+	TConstArrayView<FVector3f> FromPositions = FromCloth.GetSimPosition3D();
+	TConstArrayView<FVector3f> FromNormals = FromCloth.GetSimNormal();
+	TConstArrayView<TArray<int32>> FromBoneIndices = FromCloth.GetSimBoneIndices();
+	TConstArrayView<TArray<float>> FromBoneWeights = FromCloth.GetSimBoneWeights();
+	for (int32 Index = 0; Index < FromCloth.GetNumSimVertices3D(); ++Index)
+	{
+		Positions.Emplace(FromPositions[Index]);
+		Normals.Emplace(FromNormals[Index]);
+		BoneIndices.Emplace(FromBoneIndices[Index]);
+		BoneWeights.Emplace(FromBoneWeights[Index]);
+	}
+	TConstArrayView<FVector3f> ToPositions = ToCloth.GetSimPosition3D();
+	TConstArrayView<FVector3f> ToNormals = ToCloth.GetSimNormal();
+	TConstArrayView<TArray<int32>> ToBoneIndices = ToCloth.GetSimBoneIndices();
+	TConstArrayView<TArray<float>> ToBoneWeights = ToCloth.GetSimBoneWeights();
+	for (int32 Index = FromCloth.GetNumSimVertices3D(); Index < NumToSimVertices3D; ++Index)
+	{
+		Positions.Emplace(ToPositions[Index]);
+		Normals.Emplace(ToNormals[Index]);
+		BoneIndices.Emplace(ToBoneIndices[Index]);
+		BoneWeights.Emplace(ToBoneWeights[Index]);
+	}
+
+	FCollectionClothSimAccessoryMeshFacade AccessoryFacade = ToCloth.AddGetSimAccessoryMesh();
+	AccessoryFacade.Initialize(GenerateUniqueAccessoryMeshName(), TConstArrayView<FVector3f>(Positions), TConstArrayView<FVector3f>(Normals), TConstArrayView<TArray<int32>>(BoneIndices), TConstArrayView<TArray<float>>(BoneWeights));
+	return AccessoryFacade.GetSimAccessoryMeshIndex();
+
+}
 }  // End namespace UE::Chaos::ClothAsset
diff --git a/Engine/Plugins/ChaosClothAsset/Source/ChaosClothAssetEngine/Public/ChaosClothAsset/ClothEngineTools.h b/Engine/Plugins/ChaosClothAsset/Source/ChaosClothAssetEngine/Public/ChaosClothAsset/ClothEngineTools.h
--- a/Engine/Plugins/ChaosClothAsset/Source/ChaosClothAssetEngine/Public/ChaosClothAsset/ClothEngineTools.h
+++ b/Engine/Plugins/ChaosClothAsset/Source/ChaosClothAssetEngine/Public/ChaosClothAsset/ClothEngineTools.h
@@ -12,6 +12,7 @@ struct FManagedArrayCollection;
 class FName;
 class FSkeletalMeshLODModel;
 struct FReferenceSkeleton;
+struct FSoftObjectPath;
 namespace Chaos
 {
 	namespace Softs
@@ -22,7 +23,9 @@ namespace Chaos
 
 namespace UE::Chaos::ClothAsset
 {
+	class FCollectionClothFacade;
 	class FCollectionClothConstFacade;
+	class FCollectionClothFacade;
 
 	/**
 	 *  Tools operating on cloth collections with Engine dependency
@@ -44,6 +47,27 @@ namespace UE::Chaos::ClothAsset
 		static UE_API int32 CalculateReferenceBoneIndex(const TArray<int32>& UsedBones, const FReferenceSkeleton& ReferenceSkeleton);
 		/** Calculate ReferenceBone from bone closest to the root of all used weighted sim bones in a collection */
 		static UE_API int32 CalculateReferenceBoneIndex(const TSharedRef<const FManagedArrayCollection>& ClothCollection, const FReferenceSkeleton& ReferenceSkeleton);
+
+		/** Determine if skeletal meshes on two different cloth collections are compatible (one's reference skeleton is a subset of the other). If they are, determine which skeletal mesh is the superset,
+		 *  and calculate remap indices that can be used to remap bone indices from the subset collection to the superset skeletal mesh. 
+		 *  Use RemapBoneIndices to apply remapping.
+		 *  @param OutBoneIndicesRemapCloth1 is remap indices to convert Cloth1 BoneIndices to OutMergedSkeletalMeshPath's ref skeleton. Empty if no remapping is required.
+		 *  @param OutBoneIndicesRemapCloth2 is remap indices to convert Cloth2 BoneIndices to OutMergedSkeletalMeshPath's ref skeleton. Empty if no remapping is required.
+		 *  @param OutIncompatibleErrorDetails optional text describing why the skeletal meshes aren't compatible (empty if no error).
+		 *  @return Whether or not the collections are compatible
+		 */
+		static UE_API bool CalculateRemappedBoneIndicesIfCompatible(const FCollectionClothConstFacade& Cloth1, const FCollectionClothConstFacade& Cloth2, FSoftObjectPath& OutMergedSkeletalMeshPath, TArray<int32>& OutBoneIndicesRemapCloth1, TArray<int32>& OutBoneIndicesRemapCloth2, FText* OutIncompatibleErrorDetails = nullptr);
+
+		/**
+		 * Remap Bone indices for a cloth collection using remap array calculated by CalculateRemappedBoneIndicesIfCompatible. 
+		 * Only vertices beginning with the provided Offsets will be remapped (can be set to non-zero to remap after merging collections) */
+		static UE_API void RemapBoneIndices(FCollectionClothFacade& Cloth, const TArray<int32>& BoneIndicesRemap, const int32 SimVertex3DOffset = 0, const int32 RenderVertexOffset = 0);
+
+		/** Create a new SimAccessoryMesh by copying the SimMesh data from one collection into another.
+		 *  Optionally use SimImportVertexID to match data rather than a straight copy of SimVertices3D.
+		 *  Any unmatched vertices will copy the data from the existing Sim Mesh.
+		 *  Returns the newly created accessory mesh's index */
+		static UE_API int32 CopySimMeshToSimAccessoryMesh(const FName& AccessoryMeshName, FCollectionClothFacade& ToCloth, const FCollectionClothConstFacade& FromCloth, bool bUseSimImportVertexID, FText* OutIncompatibleErrorDetails = nullptr);
 	};
 }  // End namespace UE::Chaos::ClothAsset
 
diff --git a/Engine/Plugins/ChaosClothAssetEditor/Source/ChaosClothAssetDataflowNodes/Private/ChaosClothAsset/MergeClothCollectionsNode.cpp b/Engine/Plugins/ChaosClothAssetEditor/Source/ChaosClothAssetDataflowNodes/Private/ChaosClothAsset/MergeClothCollectionsNode.cpp
--- a/Engine/Plugins/ChaosClothAssetEditor/Source/ChaosClothAssetDataflowNodes/Private/ChaosClothAsset/MergeClothCollectionsNode.cpp
+++ b/Engine/Plugins/ChaosClothAssetEditor/Source/ChaosClothAssetDataflowNodes/Private/ChaosClothAsset/MergeClothCollectionsNode.cpp
@@ -5,6 +5,7 @@
 #include "ChaosClothAsset/CollectionClothFacade.h"
 #include "ChaosClothAsset/CollectionClothSelectionFacade.h"
 #include "ChaosClothAsset/ClothDataflowTools.h"
+#include "ChaosClothAsset/ClothEngineTools.h"
 #include "Chaos/CollectionPropertyFacade.h"
 #include "Dataflow/DataflowInputOutput.h"
 #include "Engine/SkeletalMesh.h"
@@ -418,125 +419,6 @@ namespace UE::Chaos::ClothAsset::Private
 			}
 		}
 	}
-
-	static void RemapBoneIndices(TArrayView<TArray<int32>> BoneIndices, const TArray<int32>& Remap)
-	{
-		for (TArray<int32>& Array : BoneIndices)
-		{
-			for (int32& Index : Array)
-			{
-				if (Index == INDEX_NONE)
-				{
-					continue;
-				}
-				if (ensure(Remap.IsValidIndex(Index)))
-				{
-					Index = Remap[Index];
-				}
-			}
-		}
-
-	}
-
-	static void RemapBones(FCollectionClothFacade& Cloth, const TArray<int32>& Remap, int32 SimVertex3DOffset, int32 RenderVertexOffset)
-	{
-		RemapBoneIndices(Cloth.GetSimBoneIndices().RightChop(SimVertex3DOffset), Remap);
-		RemapBoneIndices(Cloth.GetRenderBoneIndices().RightChop(RenderVertexOffset), Remap);
-	}
-
-	static bool AreSkeletalMeshesCompatible(const FDataflowNode& DataflowNode, FCollectionClothFacade& Cloth1, const FCollectionClothConstFacade& Cloth2, TArray<int32>& OtherBoneRemap)
-	{
-		OtherBoneRemap.Reset();
-
-		/** Disallow merging cloth facades with incompatible ref skeletons. */
-		const FSoftObjectPath& SkeletalMeshPathName1 = Cloth1.GetSkeletalMeshSoftObjectPathName();
-		const FSoftObjectPath& SkeletalMeshPathName2 = Cloth2.GetSkeletalMeshSoftObjectPathName();
-		if (SkeletalMeshPathName1.IsNull() || SkeletalMeshPathName2.IsNull() || SkeletalMeshPathName1 == SkeletalMeshPathName2)
-		{
-			return true;
-		}
-
-		static const FText ErrorHeadline = LOCTEXT("IncompatibleSkeletalMeshesHeadline", "Incompatible Skeletal Meshes.");
-		const USkeletalMesh* const SkeletalMesh1 = Cast<USkeletalMesh>(SkeletalMeshPathName1.TryLoad());
-		const USkeletalMesh* const SkeletalMesh2 = Cast<USkeletalMesh>(SkeletalMeshPathName2.TryLoad());
-		if (!SkeletalMesh1 || !SkeletalMesh2)
-		{
-			const FText Details = FText::Format(
-				LOCTEXT(
-					"IncompatibleSkeletalMeshesLoadFailureDetails",
-					"Cloth collections failed to merge due to failing to load SkeletalMesh \"{0}\" to check compatibility."),
-				!SkeletalMesh1 ? FText::FromString(SkeletalMeshPathName1.ToString()) : FText::FromString(SkeletalMeshPathName2.ToString()));
-
-			FClothDataflowTools::LogAndToastWarning(DataflowNode, ErrorHeadline, Details);
-			return false;
-		}
-
-		const FReferenceSkeleton& RefSkeleton1 = SkeletalMesh1->GetRefSkeleton();
-		const FReferenceSkeleton& RefSkeleton2 = SkeletalMesh2->GetRefSkeleton();
-
-		const FReferenceSkeleton& MergedRefSkeleton = RefSkeleton1.GetNum() >= RefSkeleton2.GetNum() ? RefSkeleton1 : RefSkeleton2;
-		const FReferenceSkeleton& RemapRefSkeleton = RefSkeleton1.GetNum() >= RefSkeleton2.GetNum() ? RefSkeleton2 : RefSkeleton1;
-		const FSoftObjectPath& MergedSkeletalMeshPath = RefSkeleton1.GetNum() >= RefSkeleton2.GetNum() ? SkeletalMeshPathName1 : SkeletalMeshPathName2;
-
-		const TArray<FMeshBoneInfo>& RemapBoneInfo = RemapRefSkeleton.GetRefBoneInfo();
-		const TArray<FTransform>& MergedBonePose = MergedRefSkeleton.GetRefBonePose();
-		const TArray<FTransform>& RemapBonePose = RemapRefSkeleton.GetRefBonePose();
-		TArray<int32> RemapIndices;
-		RemapIndices.SetNumUninitialized(RemapRefSkeleton.GetNum());
-		bool bAnyRemap = false;
-		for (int32 BoneIndex = 0; BoneIndex < RemapRefSkeleton.GetNum(); ++BoneIndex)
-		{
-			const int32 MergedBoneIndex = MergedRefSkeleton.FindBoneIndex(RemapBoneInfo[BoneIndex].Name);
-			if(MergedBoneIndex == INDEX_NONE)
-			{
-				const FText Details = FText::Format(
-					LOCTEXT(
-						"IncompatibleSkeletalMeshesRefBoneInfoDetails",
-						"Cloth collections failed to merge due to incompatible Skeletal Meshes, \"{0}\" and \"{1}\". Could not find bone \"{2}\" in \"{3}\"."),
-					FText::FromString(SkeletalMeshPathName1.ToString()),
-					FText::FromString(SkeletalMeshPathName2.ToString()),
-					FText::FromName(RemapBoneInfo[BoneIndex].Name),
-					FText::FromString(MergedSkeletalMeshPath.ToString()));
-
-				FClothDataflowTools::LogAndToastWarning(DataflowNode, ErrorHeadline, Details);
-				return false;
-			}
-			if (!RemapBonePose[BoneIndex].Equals(MergedBonePose[MergedBoneIndex]))
-			{
-				const FText Details = FText::Format(
-					LOCTEXT(
-						"IncompatibleSkeletalMeshesRefBonePoseDetails",
-						"Cloth collections failed to merge due to incompatible Skeletal Meshes, \"{0}\" and \"{1}\". RefBonePoses are mismatched for bone \"{2}\"."),
-					FText::FromString(SkeletalMeshPathName1.ToString()),
-					FText::FromString(SkeletalMeshPathName2.ToString()),
-					FText::FromName(RemapBoneInfo[BoneIndex].Name));
-
-				FClothDataflowTools::LogAndToastWarning(DataflowNode, ErrorHeadline, Details);
-				return false;
-			}
-			RemapIndices[BoneIndex] = MergedBoneIndex;
-			if (BoneIndex != MergedBoneIndex)
-			{
-				bAnyRemap = true;
-			}
-		}
-
-		if (bAnyRemap)
-		{
-			if (&RemapRefSkeleton == &RefSkeleton1)
-			{
-				// Remap here since Cloth1 is writable.
-				RemapBones(Cloth1, RemapIndices, 0, 0);
-			}
-			else
-			{
-				OtherBoneRemap = MoveTemp(RemapIndices);
-			}
-		}
-		Cloth1.SetSkeletalMeshSoftObjectPathName(MergedSkeletalMeshPath);
-
-		return true;
-	}
 }
 
 FChaosClothAssetMergeClothCollectionsNode_v2::FChaosClothAssetMergeClothCollectionsNode_v2(const UE::Dataflow::FNodeParameters& InParam, FGuid InGuid)
@@ -659,10 +541,17 @@ void FChaosClothAssetMergeClothCollectionsNode_v2::Evaluate(UE::Dataflow::FConte
 			bool bAppendedCloth = false;
 			if (OtherClothFacade.IsValid())
 			{
-
-				TArray<int32> OtherBoneRemap;
-				if (Private::AreSkeletalMeshesCompatible(*this, ClothFacade, OtherClothFacade, OtherBoneRemap))
+				FSoftObjectPath RemapSkeletalMeshPath;
+				TArray<int32> ThisBoneRemap, OtherBoneRemap;
+				FText IncompatibleErrorDetails;
+				if(FClothEngineTools::CalculateRemappedBoneIndicesIfCompatible(ClothFacade, OtherClothFacade, RemapSkeletalMeshPath, ThisBoneRemap, OtherBoneRemap, &IncompatibleErrorDetails))
 				{
+					if (!ThisBoneRemap.IsEmpty())
+					{
+						ClothFacade.SetSkeletalMeshSoftObjectPathName(RemapSkeletalMeshPath);
+						FClothEngineTools::RemapBoneIndices(ClothFacade, ThisBoneRemap);
+					}
+
 					if (!GroupNameOffsets.Find(ClothCollectionGroup::SimVertices3D))
 					{
 						GroupNameOffsets.Add(ClothCollectionGroup::SimVertices3D) = ClothCollection->NumElements(ClothCollectionGroup::SimVertices3D);
@@ -678,9 +567,14 @@ void FChaosClothAssetMergeClothCollectionsNode_v2::Evaluate(UE::Dataflow::FConte
 
 					if (!OtherBoneRemap.IsEmpty())
 					{
-						Private::RemapBones(ClothFacade, OtherBoneRemap, GroupNameOffsets[ClothCollectionGroup::SimVertices3D], GroupNameOffsets[ClothCollectionGroup::RenderVertices]);
+						FClothEngineTools::RemapBoneIndices(ClothFacade, OtherBoneRemap, GroupNameOffsets[ClothCollectionGroup::SimVertices3D], GroupNameOffsets[ClothCollectionGroup::RenderVertices]);
 					}
 				}
+				else
+				{
+					static const FText ErrorHeadline = LOCTEXT("IncompatibleSkeletalMeshesHeadline", "Incompatible Skeletal Meshes.");
+					FClothDataflowTools::LogAndToastWarning(*this, ErrorHeadline, IncompatibleErrorDetails);
+				}
 			}
 
 			// Append selections (with offsets)
@@ -997,9 +891,17 @@ void FChaosClothAssetMergeClothCollectionsNode::Evaluate(UE::Dataflow::FContext&
 			if (OtherClothFacade.IsValid())
 			{
 
-				TArray<int32> OtherBoneRemap;
-				if (Private::AreSkeletalMeshesCompatible(*this, ClothFacade, OtherClothFacade, OtherBoneRemap))
+				FSoftObjectPath RemapSkeletalMeshPath;
+				TArray<int32> ThisBoneRemap, OtherBoneRemap;
+				FText IncompatibleErrorDetails;
+				if (FClothEngineTools::CalculateRemappedBoneIndicesIfCompatible(ClothFacade, OtherClothFacade, RemapSkeletalMeshPath, ThisBoneRemap, OtherBoneRemap, &IncompatibleErrorDetails))
 				{
+					if (!ThisBoneRemap.IsEmpty())
+					{
+						ClothFacade.SetSkeletalMeshSoftObjectPathName(RemapSkeletalMeshPath);
+						FClothEngineTools::RemapBoneIndices(ClothFacade, ThisBoneRemap);
+					}
+
 					if (!GroupNameOffsets.Find(ClothCollectionGroup::SimVertices3D))
 					{
 						GroupNameOffsets.Add(ClothCollectionGroup::SimVertices3D) = ClothCollection->NumElements(ClothCollectionGroup::SimVertices3D);
@@ -1015,9 +917,14 @@ void FChaosClothAssetMergeClothCollectionsNode::Evaluate(UE::Dataflow::FContext&
 
 					if (!OtherBoneRemap.IsEmpty())
 					{
-						Private::RemapBones(ClothFacade, OtherBoneRemap, GroupNameOffsets[ClothCollectionGroup::SimVertices3D], GroupNameOffsets[ClothCollectionGroup::RenderVertices]);
+						FClothEngineTools::RemapBoneIndices(ClothFacade, OtherBoneRemap, GroupNameOffsets[ClothCollectionGroup::SimVertices3D], GroupNameOffsets[ClothCollectionGroup::RenderVertices]);
 					}
 				}
+				else
+				{
+					static const FText ErrorHeadline = LOCTEXT("IncompatibleSkeletalMeshesHeadline", "Incompatible Skeletal Meshes.");
+					FClothDataflowTools::LogAndToastWarning(*this, ErrorHeadline, IncompatibleErrorDetails);
+				}
 			}
 
 			// Append selections (with offsets)
diff --git a/Engine/Plugins/ChaosClothAssetEditor/Source/ChaosClothAssetDataflowNodes/Private/ChaosClothAsset/SimAccessoryMeshNode.cpp b/Engine/Plugins/ChaosClothAssetEditor/Source/ChaosClothAssetDataflowNodes/Private/ChaosClothAsset/SimAccessoryMeshNode.cpp
--- a/Engine/Plugins/ChaosClothAssetEditor/Source/ChaosClothAssetDataflowNodes/Private/ChaosClothAsset/SimAccessoryMeshNode.cpp
+++ b/Engine/Plugins/ChaosClothAssetEditor/Source/ChaosClothAssetDataflowNodes/Private/ChaosClothAsset/SimAccessoryMeshNode.cpp
@@ -1,7 +1,7 @@
 // Copyright Epic Games, Inc. All Rights Reserved.
 
 #include "ChaosClothAsset/SimAccessoryMeshNode.h"
-#include "ChaosClothAsset/ClothGeometryTools.h"
+#include "ChaosClothAsset/ClothEngineTools.h"
 #include "ChaosClothAsset/CollectionClothFacade.h"
 #include "Dataflow/DataflowInputOutput.h"
 #if WITH_EDITOR
@@ -86,11 +86,16 @@ void FChaosClothAssetSimAccessoryMeshNode::Evaluate(UE::Dataflow::FContext& Cont
 
 			if (AccessoryClothFacade.IsValid())
 			{
-				const int32 AccessoryMeshIndex = FClothGeometryTools::CopySimMeshToSimAccessoryMesh(FName(OutAccessoryMeshName), ClothFacade, AccessoryClothFacade, bUseSimImportVertexID);
+				FText WarningText;
+				const int32 AccessoryMeshIndex = FClothEngineTools::CopySimMeshToSimAccessoryMesh(FName(OutAccessoryMeshName), ClothFacade, AccessoryClothFacade, bUseSimImportVertexID, &WarningText);
 				if (ClothFacade.GetSimAccessoryMeshName().IsValidIndex(AccessoryMeshIndex))
 				{
 					OutAccessoryMeshName = ClothFacade.GetSimAccessoryMeshName()[AccessoryMeshIndex].ToString();
 				}
+				else
+				{
+					Context.Warning(WarningText.ToString());
+				}
 			}
 		}
 
