Fix unbalanced LOCTEXT_NAMESPACE define found by CIS from previous commit

diff --git a/Engine/Plugins/ChaosClothAsset/Source/ChaosClothAssetEngine/Private/ChaosClothAsset/ClothEngineTools.cpp b/Engine/Plugins/ChaosClothAsset/Source/ChaosClothAssetEngine/Private/ChaosClothAsset/ClothEngineTools.cpp
--- a/Engine/Plugins/ChaosClothAsset/Source/ChaosClothAssetEngine/Private/ChaosClothAsset/ClothEngineTools.cpp
+++ b/Engine/Plugins/ChaosClothAsset/Source/ChaosClothAssetEngine/Private/ChaosClothAsset/ClothEngineTools.cpp
@@ -564,3 +564,5 @@ int32 FClothEngineTools::CopySimMeshToSimAccessoryMesh(const FName& AccessoryMes
 
 }
 }  // End namespace UE::Chaos::ClothAsset
+
+#undef LOCTEXT_NAMESPACE
