Chaos SKM Cloth Asset - Fixed ResetConfigProperties() that was losing the clothing Interactors.

#rnx

diff --git a/Engine/Plugins/ChaosClothAsset/Source/ChaosClothAssetEngine/Private/ChaosClothAsset/ClothComponent.cpp b/Engine/Plugins/ChaosClothAsset/Source/ChaosClothAssetEngine/Private/ChaosClothAsset/ClothComponent.cpp
--- a/Engine/Plugins/ChaosClothAsset/Source/ChaosClothAssetEngine/Private/ChaosClothAsset/ClothComponent.cpp
+++ b/Engine/Plugins/ChaosClothAsset/Source/ChaosClothAssetEngine/Private/ChaosClothAsset/ClothComponent.cpp
@@ -35,7 +35,10 @@ void FChaosClothSimulationProperties::Initialize(const TArray<TSharedRef<const F
 		CollectionPropertyFacades.Emplace(StaticCastSharedRef<FCollectionPropertyFacade>(CollectionPropertyMutableFacade));  // Mutable facades conveniently inherits from non mutable ones
 		PropertyCollections.Emplace(MoveTemp(PropertyCollection));
 	}
-	ClothOutfitInteractor = NewObject<UChaosClothAssetInteractor>();
+	if (!ClothOutfitInteractor)  // No need to recreate the interactor, it only needs new property facades, and it is also required for config updates on the SKM implementation to keep the same pointer
+	{
+		ClothOutfitInteractor = NewObject<UChaosClothAssetInteractor>();
+	}
 	ClothOutfitInteractor->SetProperties(CollectionPropertyFacades);
 }
 
