Chaos Cloth Asset - Fixed multiple debug draw visualizations inside the PIE window when an Editor Preview window is also active.

#rnx

diff --git a/Engine/Plugins/ChaosClothAsset/Source/ChaosClothAssetEngine/Private/ChaosClothAsset/ClothSimulationProxy.cpp b/Engine/Plugins/ChaosClothAsset/Source/ChaosClothAssetEngine/Private/ChaosClothAsset/ClothSimulationProxy.cpp
--- a/Engine/Plugins/ChaosClothAsset/Source/ChaosClothAssetEngine/Private/ChaosClothAsset/ClothSimulationProxy.cpp
+++ b/Engine/Plugins/ChaosClothAsset/Source/ChaosClothAssetEngine/Private/ChaosClothAsset/ClothSimulationProxy.cpp
@@ -631,7 +631,7 @@ namespace UE::Chaos::ClothAsset
 			NumDynamicParticles += Cloth->GetNumActiveDynamicParticles();
 			if (FirstActiveClothParticleRangeId == INDEX_NONE && Cloth->GetNumActiveDynamicParticles() > 0)
 			{
-				FirstActiveClothParticleRangeId = Cloth->GetParticleRangeId(Solver.Get ());
+				FirstActiveClothParticleRangeId = Cloth->GetParticleRangeId(Solver.Get());
 			}
 		}
 		if (FirstActiveClothParticleRangeId != INDEX_NONE)
@@ -649,67 +649,71 @@ namespace UE::Chaos::ClothAsset
 
 		// Visualization
 #if CHAOS_DEBUG_DRAW
-		static const TConsoleVariableData<bool>* const DebugDrawLocalSpaceCVar = IConsoleManager::Get().FindTConsoleVariableDataBool(TEXT("p.ChaosCloth.DebugDrawLocalSpace"));
-		static const TConsoleVariableData<bool>* const DebugDrawBoundsCVar = IConsoleManager::Get().FindTConsoleVariableDataBool(TEXT("p.ChaosCloth.DebugDrawBounds"));
-		static const TConsoleVariableData<bool>* const DebugDrawGravityCVar = IConsoleManager::Get().FindTConsoleVariableDataBool(TEXT("p.ChaosCloth.DebugDrawGravity"));
-		static const TConsoleVariableData<bool>* const DebugDrawPhysMeshWiredCVar = IConsoleManager::Get().FindTConsoleVariableDataBool(TEXT("p.ChaosCloth.DebugDrawPhysMeshWired"));
-		static const TConsoleVariableData<bool>* const DebugDrawAnimMeshWiredCVar = IConsoleManager::Get().FindTConsoleVariableDataBool(TEXT("p.ChaosCloth.DebugDrawAnimMeshWired"));
-		static const TConsoleVariableData<bool>* const DebugDrawPointVelocitiesCVar = IConsoleManager::Get().FindTConsoleVariableDataBool(TEXT("p.ChaosCloth.DebugDrawPointVelocities"));
-		static const TConsoleVariableData<bool>* const DebugDrawAnimNormalsCVar = IConsoleManager::Get().FindTConsoleVariableDataBool(TEXT("p.ChaosCloth.DebugDrawAnimNormals"));
-		static const TConsoleVariableData<bool>* const DebugDrawPointNormalsCVar = IConsoleManager::Get().FindTConsoleVariableDataBool(TEXT("p.ChaosCloth.DebugDrawPointNormals"));
-		static const TConsoleVariableData<bool>* const DebugDrawCollisionCVar = IConsoleManager::Get().FindTConsoleVariableDataBool(TEXT("p.ChaosCloth.DebugDrawCollision"));
-		static const TConsoleVariableData<bool>* const DebugDrawBackstopsCVar = IConsoleManager::Get().FindTConsoleVariableDataBool(TEXT("p.ChaosCloth.DebugDrawBackstops"));
-		static const TConsoleVariableData<bool>* const DebugDrawBackstopDistancesCVar = IConsoleManager::Get().FindTConsoleVariableDataBool(TEXT("p.ChaosCloth.DebugDrawBackstopDistances"));
-		static const TConsoleVariableData<bool>* const DebugDrawMaxDistancesCVar = IConsoleManager::Get().FindTConsoleVariableDataBool(TEXT("p.ChaosCloth.DebugDrawMaxDistances"));
-		static const TConsoleVariableData<bool>* const DebugDrawMaxDistanceValuesCVar = IConsoleManager::Get().FindTConsoleVariableDataBool(TEXT("p.ChaosCloth.DebugDrawMaxDistanceValues"));
-		static const TConsoleVariableData<bool>* const DebugDrawAnimDriveCVar = IConsoleManager::Get().FindTConsoleVariableDataBool(TEXT("p.ChaosCloth.DebugDrawAnimDrive"));
-		static const TConsoleVariableData<bool>* const DebugDrawEdgeConstraintCVar = IConsoleManager::Get().FindTConsoleVariableDataBool(TEXT("p.ChaosCloth.DebugDrawEdgeConstraint"));
-		static const TConsoleVariableData<bool>* const DebugDrawBendingConstraintCVar = IConsoleManager::Get().FindTConsoleVariableDataBool(TEXT("p.ChaosCloth.DebugDrawBendingConstraint"));
-		static const TConsoleVariableData<bool>* const DebugDrawLongRangeConstraintCVar = IConsoleManager::Get().FindTConsoleVariableDataBool(TEXT("p.ChaosCloth.DebugDrawLongRangeConstraint"));
-		static const TConsoleVariableData<bool>* const DebugDrawWindForcesCVar = IConsoleManager::Get().FindTConsoleVariableDataBool(TEXT("p.ChaosCloth.DebugDrawWindForces"));
-		static const TConsoleVariableData<bool>* const DebugDrawWindVelocityCVar = IConsoleManager::Get().FindTConsoleVariableDataBool(TEXT("p.ChaosCloth.DebugDrawWindVelocity"));
-		static const TConsoleVariableData<bool>* const DebugDrawSelfCollisionCVar = IConsoleManager::Get().FindTConsoleVariableDataBool(TEXT("p.ChaosCloth.DebugDrawSelfCollision"));
-		static const TConsoleVariableData<bool>* const DebugDrawSelfIntersectionCVar = IConsoleManager::Get().FindTConsoleVariableDataBool(TEXT("p.ChaosCloth.DebugDrawSelfIntersection"));
-		static const TConsoleVariableData<bool>* const DebugDrawParticleIndicesCVar = IConsoleManager::Get().FindTConsoleVariableDataBool(TEXT("p.ChaosCloth.DebugDrawParticleIndices"));
-		static const TConsoleVariableData<bool>* const DebugDrawElementIndicesCVar = IConsoleManager::Get().FindTConsoleVariableDataBool(TEXT("p.ChaosCloth.DebugDrawElementIndices"));
-		static const TConsoleVariableData<bool>* const DebugDrawClothClothConstraintsCVar = IConsoleManager::Get().FindTConsoleVariableDataBool(TEXT("p.ChaosCloth.DebugDrawClothClothConstraints"));
-		static const TConsoleVariableData<bool>* const DebugDrawTeleportResetCVar = IConsoleManager::Get().FindTConsoleVariableDataBool(TEXT("p.ChaosCloth.DebugDrawTeleportReset"));
-		static const TConsoleVariableData<bool>* const DebugDrawExtremlyDeformedEdgesCVar = IConsoleManager::Get().FindTConsoleVariableDataBool(TEXT("p.ChaosCloth.DebugDrawExtremlyDeformedEdges"));
-
-		if (DebugDrawLocalSpaceCVar && DebugDrawLocalSpaceCVar->GetValueOnAnyThread()) { Visualization->DrawLocalSpace(); }
-		if (DebugDrawBoundsCVar && DebugDrawBoundsCVar->GetValueOnAnyThread()) { Visualization->DrawBounds(); }
-		if (DebugDrawGravityCVar && DebugDrawGravityCVar->GetValueOnAnyThread()) { Visualization->DrawGravity(); }
-		if (DebugDrawPhysMeshWiredCVar && DebugDrawPhysMeshWiredCVar->GetValueOnAnyThread()) { Visualization->DrawPhysMeshWired(); }
-		if (DebugDrawAnimMeshWiredCVar && DebugDrawAnimMeshWiredCVar->GetValueOnAnyThread()) { Visualization->DrawAnimMeshWired(); }
-		if (DebugDrawPointVelocitiesCVar && DebugDrawPointVelocitiesCVar->GetValueOnAnyThread()) { Visualization->DrawPointVelocities(); }
-		if (DebugDrawAnimNormalsCVar && DebugDrawAnimNormalsCVar->GetValueOnAnyThread()) { Visualization->DrawAnimNormals(); }
-		if (DebugDrawPointNormalsCVar && DebugDrawPointNormalsCVar->GetValueOnAnyThread()) { Visualization->DrawPointNormals(); }
-		if (DebugDrawCollisionCVar && DebugDrawCollisionCVar->GetValueOnAnyThread()) { Visualization->DrawCollision(); }
-		if (DebugDrawBackstopsCVar && DebugDrawBackstopsCVar->GetValueOnAnyThread()) { Visualization->DrawBackstops(); }
-		if (DebugDrawBackstopDistancesCVar && DebugDrawBackstopDistancesCVar->GetValueOnAnyThread()) { Visualization->DrawBackstopDistances(); }
-		if (DebugDrawMaxDistancesCVar && DebugDrawMaxDistancesCVar->GetValueOnAnyThread()) { Visualization->DrawMaxDistances(); }
-		if (DebugDrawMaxDistanceValuesCVar && DebugDrawMaxDistanceValuesCVar->GetValueOnAnyThread()) { Visualization->DrawMaxDistanceValues(); }
-		if (DebugDrawAnimDriveCVar && DebugDrawAnimDriveCVar->GetValueOnAnyThread()) { Visualization->DrawAnimDrive(); }
-		if (DebugDrawEdgeConstraintCVar && DebugDrawEdgeConstraintCVar->GetValueOnAnyThread()) { Visualization->DrawEdgeConstraint(); }
-		if (DebugDrawBendingConstraintCVar && DebugDrawBendingConstraintCVar->GetValueOnAnyThread()) { Visualization->DrawBendingConstraint(); }
-		if (DebugDrawLongRangeConstraintCVar && DebugDrawLongRangeConstraintCVar->GetValueOnAnyThread()) { Visualization->DrawLongRangeConstraint(); }
-		if (DebugDrawWindForcesCVar && DebugDrawWindForcesCVar->GetValueOnAnyThread()) { Visualization->DrawWindAndPressureForces(); }
-		if (DebugDrawWindVelocityCVar && DebugDrawWindVelocityCVar->GetValueOnAnyThread()) { Visualization->DrawWindVelocity(); }
-		if (DebugDrawSelfCollisionCVar && DebugDrawSelfCollisionCVar->GetValueOnAnyThread()) { Visualization->DrawSelfCollision(); }
-		if (DebugDrawSelfIntersectionCVar && DebugDrawSelfIntersectionCVar->GetValueOnAnyThread()) { Visualization->DrawSelfIntersection(); }
-		if (DebugDrawParticleIndicesCVar && DebugDrawParticleIndicesCVar->GetValueOnAnyThread()) { Visualization->DrawParticleIndices(); }
-		if (DebugDrawElementIndicesCVar && DebugDrawElementIndicesCVar->GetValueOnAnyThread()) { Visualization->DrawElementIndices(); }
-		if (DebugDrawClothClothConstraintsCVar && DebugDrawClothClothConstraintsCVar->GetValueOnAnyThread())
+		const UWorld* const World = ClothComponentAdapter.GetOwnerComponent().GetWorld();
+		if (World && World->WorldType != EWorldType::EditorPreview)  // Avoid drawing any editor preview debug draw inside the PIE window
 		{
-			Visualization->DrawClothClothConstraints();
-		}
-		if (DebugDrawTeleportResetCVar  && DebugDrawTeleportResetCVar->GetValueOnAnyThread())
-		{
-			Visualization->DrawTeleportReset();
-		}
-		if (DebugDrawExtremlyDeformedEdgesCVar && DebugDrawExtremlyDeformedEdgesCVar->GetValueOnAnyThread())
-		{
-			Visualization->DrawExtremlyDeformedEdges();
+			static const TConsoleVariableData<bool>* const DebugDrawLocalSpaceCVar = IConsoleManager::Get().FindTConsoleVariableDataBool(TEXT("p.ChaosCloth.DebugDrawLocalSpace"));
+			static const TConsoleVariableData<bool>* const DebugDrawBoundsCVar = IConsoleManager::Get().FindTConsoleVariableDataBool(TEXT("p.ChaosCloth.DebugDrawBounds"));
+			static const TConsoleVariableData<bool>* const DebugDrawGravityCVar = IConsoleManager::Get().FindTConsoleVariableDataBool(TEXT("p.ChaosCloth.DebugDrawGravity"));
+			static const TConsoleVariableData<bool>* const DebugDrawPhysMeshWiredCVar = IConsoleManager::Get().FindTConsoleVariableDataBool(TEXT("p.ChaosCloth.DebugDrawPhysMeshWired"));
+			static const TConsoleVariableData<bool>* const DebugDrawAnimMeshWiredCVar = IConsoleManager::Get().FindTConsoleVariableDataBool(TEXT("p.ChaosCloth.DebugDrawAnimMeshWired"));
+			static const TConsoleVariableData<bool>* const DebugDrawPointVelocitiesCVar = IConsoleManager::Get().FindTConsoleVariableDataBool(TEXT("p.ChaosCloth.DebugDrawPointVelocities"));
+			static const TConsoleVariableData<bool>* const DebugDrawAnimNormalsCVar = IConsoleManager::Get().FindTConsoleVariableDataBool(TEXT("p.ChaosCloth.DebugDrawAnimNormals"));
+			static const TConsoleVariableData<bool>* const DebugDrawPointNormalsCVar = IConsoleManager::Get().FindTConsoleVariableDataBool(TEXT("p.ChaosCloth.DebugDrawPointNormals"));
+			static const TConsoleVariableData<bool>* const DebugDrawCollisionCVar = IConsoleManager::Get().FindTConsoleVariableDataBool(TEXT("p.ChaosCloth.DebugDrawCollision"));
+			static const TConsoleVariableData<bool>* const DebugDrawBackstopsCVar = IConsoleManager::Get().FindTConsoleVariableDataBool(TEXT("p.ChaosCloth.DebugDrawBackstops"));
+			static const TConsoleVariableData<bool>* const DebugDrawBackstopDistancesCVar = IConsoleManager::Get().FindTConsoleVariableDataBool(TEXT("p.ChaosCloth.DebugDrawBackstopDistances"));
+			static const TConsoleVariableData<bool>* const DebugDrawMaxDistancesCVar = IConsoleManager::Get().FindTConsoleVariableDataBool(TEXT("p.ChaosCloth.DebugDrawMaxDistances"));
+			static const TConsoleVariableData<bool>* const DebugDrawMaxDistanceValuesCVar = IConsoleManager::Get().FindTConsoleVariableDataBool(TEXT("p.ChaosCloth.DebugDrawMaxDistanceValues"));
+			static const TConsoleVariableData<bool>* const DebugDrawAnimDriveCVar = IConsoleManager::Get().FindTConsoleVariableDataBool(TEXT("p.ChaosCloth.DebugDrawAnimDrive"));
+			static const TConsoleVariableData<bool>* const DebugDrawEdgeConstraintCVar = IConsoleManager::Get().FindTConsoleVariableDataBool(TEXT("p.ChaosCloth.DebugDrawEdgeConstraint"));
+			static const TConsoleVariableData<bool>* const DebugDrawBendingConstraintCVar = IConsoleManager::Get().FindTConsoleVariableDataBool(TEXT("p.ChaosCloth.DebugDrawBendingConstraint"));
+			static const TConsoleVariableData<bool>* const DebugDrawLongRangeConstraintCVar = IConsoleManager::Get().FindTConsoleVariableDataBool(TEXT("p.ChaosCloth.DebugDrawLongRangeConstraint"));
+			static const TConsoleVariableData<bool>* const DebugDrawWindForcesCVar = IConsoleManager::Get().FindTConsoleVariableDataBool(TEXT("p.ChaosCloth.DebugDrawWindForces"));
+			static const TConsoleVariableData<bool>* const DebugDrawWindVelocityCVar = IConsoleManager::Get().FindTConsoleVariableDataBool(TEXT("p.ChaosCloth.DebugDrawWindVelocity"));
+			static const TConsoleVariableData<bool>* const DebugDrawSelfCollisionCVar = IConsoleManager::Get().FindTConsoleVariableDataBool(TEXT("p.ChaosCloth.DebugDrawSelfCollision"));
+			static const TConsoleVariableData<bool>* const DebugDrawSelfIntersectionCVar = IConsoleManager::Get().FindTConsoleVariableDataBool(TEXT("p.ChaosCloth.DebugDrawSelfIntersection"));
+			static const TConsoleVariableData<bool>* const DebugDrawParticleIndicesCVar = IConsoleManager::Get().FindTConsoleVariableDataBool(TEXT("p.ChaosCloth.DebugDrawParticleIndices"));
+			static const TConsoleVariableData<bool>* const DebugDrawElementIndicesCVar = IConsoleManager::Get().FindTConsoleVariableDataBool(TEXT("p.ChaosCloth.DebugDrawElementIndices"));
+			static const TConsoleVariableData<bool>* const DebugDrawClothClothConstraintsCVar = IConsoleManager::Get().FindTConsoleVariableDataBool(TEXT("p.ChaosCloth.DebugDrawClothClothConstraints"));
+			static const TConsoleVariableData<bool>* const DebugDrawTeleportResetCVar = IConsoleManager::Get().FindTConsoleVariableDataBool(TEXT("p.ChaosCloth.DebugDrawTeleportReset"));
+			static const TConsoleVariableData<bool>* const DebugDrawExtremlyDeformedEdgesCVar = IConsoleManager::Get().FindTConsoleVariableDataBool(TEXT("p.ChaosCloth.DebugDrawExtremlyDeformedEdges"));
+
+			if (DebugDrawLocalSpaceCVar && DebugDrawLocalSpaceCVar->GetValueOnAnyThread()) { Visualization->DrawLocalSpace(); }
+			if (DebugDrawBoundsCVar && DebugDrawBoundsCVar->GetValueOnAnyThread()) { Visualization->DrawBounds(); }
+			if (DebugDrawGravityCVar && DebugDrawGravityCVar->GetValueOnAnyThread()) { Visualization->DrawGravity(); }
+			if (DebugDrawPhysMeshWiredCVar && DebugDrawPhysMeshWiredCVar->GetValueOnAnyThread()) { Visualization->DrawPhysMeshWired(); }
+			if (DebugDrawAnimMeshWiredCVar && DebugDrawAnimMeshWiredCVar->GetValueOnAnyThread()) { Visualization->DrawAnimMeshWired(); }
+			if (DebugDrawPointVelocitiesCVar && DebugDrawPointVelocitiesCVar->GetValueOnAnyThread()) { Visualization->DrawPointVelocities(); }
+			if (DebugDrawAnimNormalsCVar && DebugDrawAnimNormalsCVar->GetValueOnAnyThread()) { Visualization->DrawAnimNormals(); }
+			if (DebugDrawPointNormalsCVar && DebugDrawPointNormalsCVar->GetValueOnAnyThread()) { Visualization->DrawPointNormals(); }
+			if (DebugDrawCollisionCVar && DebugDrawCollisionCVar->GetValueOnAnyThread()) { Visualization->DrawCollision(); }
+			if (DebugDrawBackstopsCVar && DebugDrawBackstopsCVar->GetValueOnAnyThread()) { Visualization->DrawBackstops(); }
+			if (DebugDrawBackstopDistancesCVar && DebugDrawBackstopDistancesCVar->GetValueOnAnyThread()) { Visualization->DrawBackstopDistances(); }
+			if (DebugDrawMaxDistancesCVar && DebugDrawMaxDistancesCVar->GetValueOnAnyThread()) { Visualization->DrawMaxDistances(); }
+			if (DebugDrawMaxDistanceValuesCVar && DebugDrawMaxDistanceValuesCVar->GetValueOnAnyThread()) { Visualization->DrawMaxDistanceValues(); }
+			if (DebugDrawAnimDriveCVar && DebugDrawAnimDriveCVar->GetValueOnAnyThread()) { Visualization->DrawAnimDrive(); }
+			if (DebugDrawEdgeConstraintCVar && DebugDrawEdgeConstraintCVar->GetValueOnAnyThread()) { Visualization->DrawEdgeConstraint(); }
+			if (DebugDrawBendingConstraintCVar && DebugDrawBendingConstraintCVar->GetValueOnAnyThread()) { Visualization->DrawBendingConstraint(); }
+			if (DebugDrawLongRangeConstraintCVar && DebugDrawLongRangeConstraintCVar->GetValueOnAnyThread()) { Visualization->DrawLongRangeConstraint(); }
+			if (DebugDrawWindForcesCVar && DebugDrawWindForcesCVar->GetValueOnAnyThread()) { Visualization->DrawWindAndPressureForces(); }
+			if (DebugDrawWindVelocityCVar && DebugDrawWindVelocityCVar->GetValueOnAnyThread()) { Visualization->DrawWindVelocity(); }
+			if (DebugDrawSelfCollisionCVar && DebugDrawSelfCollisionCVar->GetValueOnAnyThread()) { Visualization->DrawSelfCollision(); }
+			if (DebugDrawSelfIntersectionCVar && DebugDrawSelfIntersectionCVar->GetValueOnAnyThread()) { Visualization->DrawSelfIntersection(); }
+			if (DebugDrawParticleIndicesCVar && DebugDrawParticleIndicesCVar->GetValueOnAnyThread()) { Visualization->DrawParticleIndices(); }
+			if (DebugDrawElementIndicesCVar && DebugDrawElementIndicesCVar->GetValueOnAnyThread()) { Visualization->DrawElementIndices(); }
+			if (DebugDrawClothClothConstraintsCVar && DebugDrawClothClothConstraintsCVar->GetValueOnAnyThread())
+			{
+				Visualization->DrawClothClothConstraints();
+			}
+			if (DebugDrawTeleportResetCVar && DebugDrawTeleportResetCVar->GetValueOnAnyThread())
+			{
+				Visualization->DrawTeleportReset();
+			}
+			if (DebugDrawExtremlyDeformedEdgesCVar && DebugDrawExtremlyDeformedEdgesCVar->GetValueOnAnyThread())
+			{
+				Visualization->DrawExtremlyDeformedEdges();
+			}
 		}
 #endif  // #if CHAOS_DEBUG_DRAW
 	}
