Mark string-based ClothAssetInteractor methods as deprecated.

diff --git a/Engine/Plugins/ChaosClothAsset/Source/ChaosClothAssetEngine/Public/ChaosClothAsset/ClothAssetInteractor.h b/Engine/Plugins/ChaosClothAsset/Source/ChaosClothAssetEngine/Public/ChaosClothAsset/ClothAssetInteractor.h
--- a/Engine/Plugins/ChaosClothAsset/Source/ChaosClothAssetEngine/Public/ChaosClothAsset/ClothAssetInteractor.h
+++ b/Engine/Plugins/ChaosClothAsset/Source/ChaosClothAssetEngine/Public/ChaosClothAsset/ClothAssetInteractor.h
@@ -90,6 +90,7 @@ public:
 
 	/** Get the value for a property cast to float. DefaultValue will be returned if the property is not found.*/
 	UFUNCTION(BlueprintCallable, Category = "ClothProperty|Deprecated", meta = (DeprecatedFunction = "Note", DeprecationMessage = "Use GetFloatPropertyValue instead"))
+	UE_DEPRECATED(5.7, "Use GetFloatPropertyValue instead")
 	float GetFloatValue(const FString& PropertyName, int32 LODIndex = 0, float DefaultValue = 0.f) const
 	{
 		return GetFloatPropertyValue(FName(PropertyName), LODIndex, DefaultValue);
@@ -97,6 +98,7 @@ public:
 
 	/** Get the low value for a weighted property value (same as GetFloatValue). DefaultValue will be returned if the property is not found.*/
 	UFUNCTION(BlueprintCallable, Category = "ClothProperty|Deprecated", meta = (DeprecatedFunction = "Note", DeprecationMessage = "Use GetLowFloatPropertyValue instead"))
+	UE_DEPRECATED(5.7, "Use GetLowFloatPropertyValue instead")
 	float GetLowFloatValue(const FString& PropertyName, int32 LODIndex = 0, float DefaultValue = 0.f) const
 	{
 		return GetLowFloatPropertyValue(FName(PropertyName), LODIndex, DefaultValue);
@@ -104,6 +106,7 @@ public:
 
 	/** Get the high value for a weighted property value. DefaultValue will be returned if the property is not found.*/
 	UFUNCTION(BlueprintCallable, Category = "ClothProperty|Deprecated", meta = (DeprecatedFunction = "Note", DeprecationMessage = "Use GetHighFloatPropertyValue instead"))
+	UE_DEPRECATED(5.7, "Use GetHighFloatPropertyValue instead")
 	float GetHighFloatValue(const FString& PropertyName, int32 LODIndex = 0, float DefaultValue = 0.f) const
 	{
 		return GetHighFloatPropertyValue(FName(PropertyName), LODIndex, DefaultValue);
@@ -111,6 +114,7 @@ public:
 
 	/** Get the low and high values for a weighted property value. DefaultValue will be returned if the property is not found.*/
 	UFUNCTION(BlueprintCallable, Category = "ClothProperty|Deprecated", meta = (DeprecatedFunction = "Note", DeprecationMessage = "Use GetWeightedFloatPropertyValue instead"))
+	UE_DEPRECATED(5.7, "Use GetWeightedFloatPropertyValue instead")
 	FVector2D GetWeightedFloatValue(const FString& PropertyName, int32 LODIndex = 0, FVector2D DefaultValue = FVector2D(0., 0.)) const
 	{
 		return GetWeightedFloatPropertyValue(FName(PropertyName), LODIndex, DefaultValue);
@@ -118,6 +122,7 @@ public:
 
 	/** Get the value for a property cast to int. DefaultValue will be returned if the property is not found.*/
 	UFUNCTION(BlueprintCallable, Category = "ClothProperty|Deprecated", meta = (DeprecatedFunction = "Note", DeprecationMessage = "Use GetIntPropertyValue instead"))
+	UE_DEPRECATED(5.7, "Use GetIntPropertyValue instead")
 	int32 GetIntValue(const FString& PropertyName, int32 LODIndex = 0, int32 DefaultValue = 0) const
 	{
 		return GetIntPropertyValue(FName(PropertyName), LODIndex, DefaultValue);
@@ -125,6 +130,7 @@ public:
 
 	/** Get the value for a property cast to vector. DefaultValue will be returned if the property is not found.*/
 	UFUNCTION(BlueprintCallable, Category = "ClothProperty|Deprecated", meta = (DeprecatedFunction = "Note", DeprecationMessage = "Use GetVectorPropertyValue instead"))
+	UE_DEPRECATED(5.7, "Use GetVectorPropertyValue instead")
 	FVector GetVectorValue(const FString& PropertyName, int32 LODIndex = 0, FVector DefaultValue = FVector(0.)) const
 	{
 		return GetVectorPropertyValue(FName(PropertyName), LODIndex, DefaultValue);
@@ -132,6 +138,7 @@ public:
 
 	/** Get the string value for a property (typically the associated map name for weighted values). DefaultValue will be returned if the property is not found.*/
 	UFUNCTION(BlueprintCallable, Category = "ClothProperty|Deprecated", meta = (DeprecatedFunction = "Note", DeprecationMessage = "Use GetStringPropertyValue instead"))
+	UE_DEPRECATED(5.7, "Use GetStringPropertyValue instead")
 	FString GetStringValue(const FString& PropertyName, int32 LODIndex = 0, const FString& DefaultValue = "") const
 	{
 		return GetStringPropertyValue(FName(PropertyName), LODIndex, DefaultValue);
@@ -139,6 +146,7 @@ public:
 
 	/**Set the value for a property (if it exists). This sets the Low and High values for weighted values. All LODs will be set when LODIndex = -1.*/
 	UFUNCTION(BlueprintCallable, Category = "ClothProperty|Deprecated", meta = (DeprecatedFunction = "Note", DeprecationMessage = "Use SetFloatPropertyValue instead"))
+	UE_DEPRECATED(5.7, "Use SetFloatPropertyValue instead")
 	void SetFloatValue(const FString& PropertyName, int32 LODIndex = -1, float Value = 0.f)
 	{
 		SetFloatPropertyValue(FName(PropertyName), LODIndex, Value);
@@ -146,6 +154,7 @@ public:
 
 	/**Set the low value for a weighted property (if it exists). All LODs will be set when LODIndex = -1.*/
 	UFUNCTION(BlueprintCallable, Category = "ClothProperty|Deprecated", meta = (DeprecatedFunction = "Note", DeprecationMessage = "Use SetLowFloatPropertyValue instead"))
+	UE_DEPRECATED(5.7, "Use SetLowFloatPropertyValue instead")
 	void SetLowFloatValue(const FString& PropertyName, int32 LODIndex = -1, float Value = 0.f)
 	{
 		SetLowFloatPropertyValue(FName(PropertyName), LODIndex, Value);
@@ -153,6 +162,7 @@ public:
 
 	/**Set the high value for a weighted property (if it exists). All LODs will be set when LODIndex = -1.*/
 	UFUNCTION(BlueprintCallable, Category = "ClothProperty|Deprecated", meta = (DeprecatedFunction = "Note", DeprecationMessage = "Use SetHighFloatPropertyValue instead"))
+	UE_DEPRECATED(5.7, "Use SetHighFloatPropertyValue instead")
 	void SetHighFloatValue(const FString& PropertyName, int32 LODIndex = -1, float Value = 0.f)
 	{
 		SetHighFloatPropertyValue(FName(PropertyName), LODIndex, Value);
@@ -160,6 +170,7 @@ public:
 
 	/**Set the low and high values for a weighted property (if it exists). All LODs will be set when LODIndex = -1.*/
 	UFUNCTION(BlueprintCallable, Category = "ClothProperty|Deprecated", meta = (DeprecatedFunction = "Note", DeprecationMessage = "Use SetWeightedFloatPropertyValue instead"))
+	UE_DEPRECATED(5.7, "Use SetWeightedFloatPropertyValue instead")
 	void SetWeightedFloatValue(const FString& PropertyName, int32 LODIndex = -1, FVector2D Value = FVector2D(0., 0.))
 	{
 		SetWeightedFloatPropertyValue(FName(PropertyName), LODIndex, Value);
@@ -167,6 +178,7 @@ public:
 
 	/**Set the value for a property (if it exists). All LODs will be set when LODIndex = -1.*/
 	UFUNCTION(BlueprintCallable, Category = "ClothProperty|Deprecated", meta = (DeprecatedFunction = "Note", DeprecationMessage = "Use SetIntPropertyValue instead"))
+	UE_DEPRECATED(5.7, "Use SetIntPropertyValue instead")
 	void SetIntValue(const FString& PropertyName, int32 LODIndex = -1, int32 Value = 0)
 	{
 		SetIntPropertyValue(FName(PropertyName), LODIndex, Value);
@@ -174,6 +186,7 @@ public:
 
 	/**Set the value for a property (if it exists). All LODs will be set when LODIndex = -1.*/
 	UFUNCTION(BlueprintCallable, Category = "ClothProperty|Deprecated", meta = (DeprecatedFunction = "Note", DeprecationMessage = "Use SetVectorPropertyValue instead"))
+	UE_DEPRECATED(5.7, "Use SetVectorPropertyValue instead")
 	void SetVectorValue(const FString& PropertyName, int32 LODIndex = -1, FVector Value = FVector(0.))
 	{
 		SetVectorPropertyValue(FName(PropertyName), LODIndex, Value);
@@ -181,6 +194,7 @@ public:
 
 	/**Set the string value for a property (if it exists). This is typically the map name associated with a property. All LODs will be set when LODIndex = -1.*/
 	UFUNCTION(BlueprintCallable, Category = "ClothProperty|Deprecated", meta = (DeprecatedFunction = "Note", DeprecationMessage = "Use SetStringPropertyValue instead"))
+	UE_DEPRECATED(5.7, "Use SetStringPropertyValue instead")
 	void SetStringValue(const FString& PropertyName, int32 LODIndex = -1, const FString& Value = "")
 	{
 		SetStringPropertyValue(FName(PropertyName), LODIndex, Value);
diff --git a/Engine/Plugins/MovieScene/MovieRenderPipeline/Source/MovieRenderPipelineCore/Private/MoviePipelineUtils.cpp b/Engine/Plugins/MovieScene/MovieRenderPipeline/Source/MovieRenderPipelineCore/Private/MoviePipelineUtils.cpp
--- a/Engine/Plugins/MovieScene/MovieRenderPipeline/Source/MovieRenderPipelineCore/Private/MoviePipelineUtils.cpp
+++ b/Engine/Plugins/MovieScene/MovieRenderPipeline/Source/MovieRenderPipelineCore/Private/MoviePipelineUtils.cpp
@@ -1319,8 +1319,8 @@ namespace UE
 								for (int32 LODIndex = 0; LODIndex < NumLODs; ++LODIndex)
 								{
 									constexpr int32 MinNumSubsteps = 1;
-									const int32 NumSubsteps = FMath::Max(ClothAssetInteractor->GetIntValue(TEXT("NumSubsteps"), LODIndex, MinNumSubsteps), MinNumSubsteps);
-									const int32 DynamicSubstepDeltaTime = ClothAssetInteractor->GetFloatValue(TEXT("DynamicSubstepDeltaTime"), LODIndex, MinNumSubsteps);
+									const int32 NumSubsteps = FMath::Max(ClothAssetInteractor->GetIntPropertyValue(TEXT("NumSubsteps"), LODIndex, MinNumSubsteps), MinNumSubsteps);
+									const int32 DynamicSubstepDeltaTime = ClothAssetInteractor->GetFloatPropertyValue(TEXT("DynamicSubstepDeltaTime"), LODIndex, MinNumSubsteps);
 									(*ExistingCacheEntry)[LODIndex].NumSubSteps = NumSubsteps;
 									(*ExistingCacheEntry)[LODIndex].DynamicSubstepDeltaTime = DynamicSubstepDeltaTime;
 								}
@@ -1330,8 +1330,8 @@ namespace UE
 							{
 								if (ExistingCacheEntry->IsValidIndex(LODIndex))
 								{
-									ClothAssetInteractor->SetIntValue(TEXT("NumSubsteps"), LODIndex, (*ExistingCacheEntry)[LODIndex].NumSubSteps * InSubdivisionCount);
-									ClothAssetInteractor->SetFloatValue(TEXT("DynamicSubstepDeltaTime"), LODIndex, 0.f);
+									ClothAssetInteractor->SetIntPropertyValue(TEXT("NumSubsteps"), LODIndex, (*ExistingCacheEntry)[LODIndex].NumSubSteps * InSubdivisionCount);
+									ClothAssetInteractor->SetFloatPropertyValue(TEXT("DynamicSubstepDeltaTime"), LODIndex, 0.f);
 								}
 							}
 						}
@@ -1355,8 +1355,8 @@ namespace UE
 					{
 						for (int32 LODIndex = 0; LODIndex < Pair.Value.Num(); ++LODIndex)
 						{
-							ClothAssetInteractor->SetIntValue(TEXT("NumSubsteps"), LODIndex, Pair.Value[LODIndex].NumSubSteps);
-							ClothAssetInteractor->SetIntValue(TEXT("DynamicSubstepDeltaTime"), LODIndex, Pair.Value[LODIndex].DynamicSubstepDeltaTime);
+							ClothAssetInteractor->SetIntPropertyValue(TEXT("NumSubsteps"), LODIndex, Pair.Value[LODIndex].NumSubSteps);
+							ClothAssetInteractor->SetIntPropertyValue(TEXT("DynamicSubstepDeltaTime"), LODIndex, Pair.Value[LODIndex].DynamicSubstepDeltaTime);
 						}
 					}
 				}
