Chaos Cloth Asset - Fixed crash when importing sim morph targets.

[RN] Chaos Cloth Asset - Fixed crash when importing sim morph targets.

diff --git a/Engine/Plugins/ChaosClothAsset/Source/ChaosClothAsset/Private/ChaosClothAsset/CollectionClothSimMorphTargetFacade.cpp b/Engine/Plugins/ChaosClothAsset/Source/ChaosClothAsset/Private/ChaosClothAsset/CollectionClothSimMorphTargetFacade.cpp
--- a/Engine/Plugins/ChaosClothAsset/Source/ChaosClothAsset/Private/ChaosClothAsset/CollectionClothSimMorphTargetFacade.cpp
+++ b/Engine/Plugins/ChaosClothAsset/Source/ChaosClothAsset/Private/ChaosClothAsset/CollectionClothSimMorphTargetFacade.cpp
@@ -65,7 +65,7 @@ namespace UE::Chaos::ClothAsset
 		: ClothCollection(ClothCollection)
 		, MorphTargetIndex(InMorphTargetIndex)
 	{
-		check(ClothCollection->IsValid(EClothCollectionExtendedSchemas::CookedOnly | EClothCollectionExtendedSchemas::SimMorphTargets));
+		check(ClothCollection->IsValid(EClothCollectionExtendedSchemas::CookedOnly) || ClothCollection->IsValid(EClothCollectionExtendedSchemas::SimMorphTargets));
 		check(MorphTargetIndex >= 0 && MorphTargetIndex < ClothCollection->GetNumElements(ClothCollectionGroup::SimMorphTargets));
 	}
 
