Chaos Cloth Asset - Fixed crash when proxy is simulating while a sim model changes.

[RN] Chaos Cloth Asset - Fixed crash when proxy is simulating while a sim model changes.

diff --git a/Engine/Plugins/ChaosClothAsset/Source/ChaosClothAssetEngine/Private/ChaosClothAsset/ClothSimulationProxy.cpp b/Engine/Plugins/ChaosClothAsset/Source/ChaosClothAssetEngine/Private/ChaosClothAsset/ClothSimulationProxy.cpp
--- a/Engine/Plugins/ChaosClothAsset/Source/ChaosClothAssetEngine/Private/ChaosClothAsset/ClothSimulationProxy.cpp
+++ b/Engine/Plugins/ChaosClothAsset/Source/ChaosClothAssetEngine/Private/ChaosClothAsset/ClothSimulationProxy.cpp
@@ -1,6 +1,7 @@
 // Copyright Epic Games, Inc. All Rights Reserved.
 
 #include "ChaosClothAsset/ClothSimulationProxy.h"
+#include "Algo/NoneOf.h"
 #include "ChaosClothAsset/ClothAsset.h"
 #include "ChaosClothAsset/ClothAssetPrivate.h"
 #include "ChaosClothAsset/ClothComponent.h"
@@ -168,6 +169,7 @@ namespace UE::Chaos::ClothAsset
 		Meshes.Reset();
 		Colliders.Reset();
 		Cloths.Reset();
+		ClothSimulationModels.Reset();
 
 		// Create solver config simulation thread object first. Need to know which solver type we're creating.
 		const int32 SolverConfigIndex = Configs.Emplace(MakeUnique<FClothingSimulationConfig>(ClothComponentAdapter.GetSolverPropertyCollections()));
@@ -201,6 +203,13 @@ namespace UE::Chaos::ClothAsset
 				if (SimulationGroupId != INDEX_NONE)  // INDEX_NONE means this model shouldn't be added to the simulation
 				{
 					const TSharedPtr<const FChaosClothSimulationModel> ClothSimulationModel = Asset->GetClothSimulationModel(ModelIndex);
+					if (Algo::NoneOf(ClothSimulationModels, [&ClothSimulationModel](const TSharedPtr<const FChaosClothSimulationModel>& InClothSimulationModel)
+						{
+							return InClothSimulationModel.Get() == ClothSimulationModel.Get();
+						}))
+					{
+						ClothSimulationModels.Emplace(ClothSimulationModel);  // Ownership has to transfer to this proxy when the model changes during the simulation
+					}
 
 					FString DebugName;
 #if !UE_BUILD_SHIPPING
