Chaos Cloth Asset - Fixed the Merge and SimAccessoryMesh nodes skeleton remaping code that was causing render data corruption and crashes.
- The code omitted cases where the second reference skeleton was extended from the first one.
- Also fixed the returned superset skeletal mesh path which was the worng one.

[RN] Chaos Cloth Asset - Fixed the Merge and SimAccessoryMesh nodes skeleton remaping code that was causing render data corruption and crashes.

#Synced-CL 50258746

diff --git a/Engine/Plugins/ChaosClothAsset/Source/ChaosClothAssetEngine/Private/ChaosClothAsset/ClothEngineTools.cpp b/Engine/Plugins/ChaosClothAsset/Source/ChaosClothAssetEngine/Private/ChaosClothAsset/ClothEngineTools.cpp
--- a/Engine/Plugins/ChaosClothAsset/Source/ChaosClothAssetEngine/Private/ChaosClothAsset/ClothEngineTools.cpp
+++ b/Engine/Plugins/ChaosClothAsset/Source/ChaosClothAssetEngine/Private/ChaosClothAsset/ClothEngineTools.cpp
@@ -2,6 +2,7 @@
 
 #include "ChaosClothAsset/ClothEngineTools.h"
 #include "Chaos/CollectionPropertyFacade.h"
+#include "ChaosClothAsset/ClothAssetPrivate.h"
 #include "ChaosClothAsset/CollectionClothFacade.h"
 #include "ChaosClothAsset/CollectionClothSelectionFacade.h"
 #include "ChaosClothAsset/ClothCollectionGroup.h"
@@ -291,9 +292,14 @@ int32 FClothEngineTools::CalculateReferenceBoneIndex(const TSharedRef<const FMan
 	return CalculateReferenceBoneIndex(UsedBones.Array(), ReferenceSkeleton);
 }
 
-bool FClothEngineTools::CalculateRemappedBoneIndicesIfCompatible(const FCollectionClothConstFacade& Cloth1, const FCollectionClothConstFacade& Cloth2, FSoftObjectPath& OutMergedSkeletalMeshPath, TArray<int32>& OutBoneIndicesRemapCloth1, TArray<int32>& OutBoneIndicesRemapCloth2, FText* OutIncompatibleErrorDetails)
+bool FClothEngineTools::CalculateRemappedBoneIndicesIfCompatible(
+	const FCollectionClothConstFacade& Cloth1,
+	const FCollectionClothConstFacade& Cloth2,
+	FSoftObjectPath& OutMergedSkeletalMeshPath,
+	TArray<int32>& OutBoneIndicesRemapCloth1,
+	TArray<int32>& OutBoneIndicesRemapCloth2,
+	FText* OutIncompatibleErrorDetails)
 {
-
 	const FSoftObjectPath& SkeletalMeshPathName1 = Cloth1.GetSkeletalMeshSoftObjectPathName();
 	const FSoftObjectPath& SkeletalMeshPathName2 = Cloth2.GetSkeletalMeshSoftObjectPathName();
 	if (SkeletalMeshPathName1.IsNull() || SkeletalMeshPathName2.IsNull() || SkeletalMeshPathName1 == SkeletalMeshPathName2)
@@ -373,25 +379,24 @@ bool FClothEngineTools::CalculateRemappedBoneIndicesIfCompatible(const FCollecti
 		}
 	}
 
+	OutMergedSkeletalMeshPath = MergedSkeletalMeshPath;
+
 	if (bAnyRemap)
 	{
 		if (&RemapRefSkeleton == &RefSkeleton1)
 		{
-			OutMergedSkeletalMeshPath = SkeletalMeshPathName1;
 			OutBoneIndicesRemapCloth1 = MoveTemp(RemapIndices);
 			OutBoneIndicesRemapCloth2.Reset();
 		}
 		else
 		{
 			check(&RemapRefSkeleton == &RefSkeleton2);
-			OutMergedSkeletalMeshPath = SkeletalMeshPathName2;
 			OutBoneIndicesRemapCloth1.Reset();
 			OutBoneIndicesRemapCloth2 = MoveTemp(RemapIndices);
 		}
 	}
 	else
 	{
-		OutMergedSkeletalMeshPath.Reset();
 		OutBoneIndicesRemapCloth1.Reset();
 		OutBoneIndicesRemapCloth2.Reset();
 	}
@@ -400,6 +405,28 @@ bool FClothEngineTools::CalculateRemappedBoneIndicesIfCompatible(const FCollecti
 	{
 		*OutIncompatibleErrorDetails = FText();
 	}
+
+	if (UE_LOG_ACTIVE(LogChaosClothAsset, VeryVerbose))
+	{
+		UE_LOG(LogChaosClothAsset, VeryVerbose, TEXT("--------- Calculate Remapped Bone Indices ---------"));
+		const TArray<FMeshBoneInfo>& MeshBoneInfos1 = RefSkeleton1.GetRefBoneInfo();
+		const FString BoneNames1 = FString::JoinBy(MeshBoneInfos1, TEXT(", "), [](const FMeshBoneInfo& MeshBoneInfo) { return MeshBoneInfo.Name.ToString(); });
+		UE_LOG(LogChaosClothAsset, VeryVerbose, TEXT("Skeleton 1: %s\n  %s"), *SkeletalMeshPathName1.GetAssetPathString(), *BoneNames1);
+
+		const TArray<FMeshBoneInfo>& MeshBoneInfos2 = RefSkeleton2.GetRefBoneInfo();
+		const FString BoneNames2 = FString::JoinBy(MeshBoneInfos2, TEXT(", "), [](const FMeshBoneInfo& MeshBoneInfo) { return MeshBoneInfo.Name.ToString(); });
+		UE_LOG(LogChaosClothAsset, VeryVerbose, TEXT("Skeleton 2: %s\n  %s"), *SkeletalMeshPathName2.GetAssetPathString(), *BoneNames2);
+
+		UE_LOG(LogChaosClothAsset, VeryVerbose, TEXT("Out Skeleton: %s"), *OutMergedSkeletalMeshPath.GetAssetPathString());
+
+		const FString RemapIndices1 = FString::JoinBy(OutBoneIndicesRemapCloth1, TEXT(", "), [](const int32 Index) { return FString::FromInt(Index); });
+		UE_LOG(LogChaosClothAsset, VeryVerbose, TEXT("Remap Indices 1:\n  %s"), *RemapIndices1);
+
+		const FString RemapIndices2 = FString::JoinBy(OutBoneIndicesRemapCloth2, TEXT(", "), [](const int32 Index) { return FString::FromInt(Index); });
+		UE_LOG(LogChaosClothAsset, VeryVerbose, TEXT("Remap Indices 2:\n  %s"), *RemapIndices2);
+		UE_LOG(LogChaosClothAsset, VeryVerbose, TEXT("---------------------------------------------------"));
+	}
+
 	return true;
 }
 
@@ -445,9 +472,9 @@ int32 FClothEngineTools::CopySimMeshToSimAccessoryMesh(const FName& AccessoryMes
 		return INDEX_NONE;
 	}
 
+	ToCloth.SetSkeletalMeshSoftObjectPathName(RemapSkeletalMeshPath);  // The path may change even with no remaps (e.g. when the second skeleton is extended from the first one)
 	if (!ToClothBoneRemap.IsEmpty())
 	{
-		ToCloth.SetSkeletalMeshSoftObjectPathName(RemapSkeletalMeshPath);
 		RemapBoneIndices(ToCloth, ToClothBoneRemap);
 	}
 
diff --git a/Engine/Plugins/ChaosClothAsset/Source/ChaosClothAssetEngine/Public/ChaosClothAsset/ClothEngineTools.h b/Engine/Plugins/ChaosClothAsset/Source/ChaosClothAssetEngine/Public/ChaosClothAsset/ClothEngineTools.h
--- a/Engine/Plugins/ChaosClothAsset/Source/ChaosClothAssetEngine/Public/ChaosClothAsset/ClothEngineTools.h
+++ b/Engine/Plugins/ChaosClothAsset/Source/ChaosClothAssetEngine/Public/ChaosClothAsset/ClothEngineTools.h
@@ -48,15 +48,24 @@ namespace UE::Chaos::ClothAsset
 		/** Calculate ReferenceBone from bone closest to the root of all used weighted sim bones in a collection */
 		static UE_API int32 CalculateReferenceBoneIndex(const TSharedRef<const FManagedArrayCollection>& ClothCollection, const FReferenceSkeleton& ReferenceSkeleton);
 
-		/** Determine if skeletal meshes on two different cloth collections are compatible (one's reference skeleton is a subset of the other). If they are, determine which skeletal mesh is the superset,
-		 *  and calculate remap indices that can be used to remap bone indices from the subset collection to the superset skeletal mesh. 
-		 *  Use RemapBoneIndices to apply remapping.
-		 *  @param OutBoneIndicesRemapCloth1 is remap indices to convert Cloth1 BoneIndices to OutMergedSkeletalMeshPath's ref skeleton. Empty if no remapping is required.
-		 *  @param OutBoneIndicesRemapCloth2 is remap indices to convert Cloth2 BoneIndices to OutMergedSkeletalMeshPath's ref skeleton. Empty if no remapping is required.
+		/** Determine if skeletal meshes on two different cloth collections are compatible (one's reference skeleton is a subset of the other).
+		 *  If they are, determine which skeletal mesh is the superset, and calculate remap indices that can be used to remap bone indices from
+		 *  the subset collection to the superset skeletal mesh.
+		 *  @param Cloth1 The cloth facade for the first cloth collection.
+		 *  @param Cloth2 The cloth facade for the second cloth collection.
+		 *  @param OutMergedSkeletalMeshPath The superset skeletal mesh used as reference skeleton, either Cloth1 or Cloth2's skeletal mesh path.
+		 *  @param OutBoneIndicesRemapCloth1 Remap indices to convert Cloth1 BoneIndices to OutMergedSkeletalMeshPath's ref skeleton. Empty if no remapping is required.
+		 *  @param OutBoneIndicesRemapCloth2 Remap indices to convert Cloth2 BoneIndices to OutMergedSkeletalMeshPath's ref skeleton. Empty if no remapping is required.
 		 *  @param OutIncompatibleErrorDetails optional text describing why the skeletal meshes aren't compatible (empty if no error).
-		 *  @return Whether or not the collections are compatible
+		 *  @return Whether or not the collections are compatible.
 		 */
-		static UE_API bool CalculateRemappedBoneIndicesIfCompatible(const FCollectionClothConstFacade& Cloth1, const FCollectionClothConstFacade& Cloth2, FSoftObjectPath& OutMergedSkeletalMeshPath, TArray<int32>& OutBoneIndicesRemapCloth1, TArray<int32>& OutBoneIndicesRemapCloth2, FText* OutIncompatibleErrorDetails = nullptr);
+		static UE_API bool CalculateRemappedBoneIndicesIfCompatible(
+			const FCollectionClothConstFacade& Cloth1,
+			const FCollectionClothConstFacade& Cloth2,
+			FSoftObjectPath& OutMergedSkeletalMeshPath,
+			TArray<int32>& OutBoneIndicesRemapCloth1,
+			TArray<int32>& OutBoneIndicesRemapCloth2,
+			FText* OutIncompatibleErrorDetails = nullptr);
 
 		/**
 		 * Remap Bone indices for a cloth collection using remap array calculated by CalculateRemappedBoneIndicesIfCompatible. 
diff --git a/Engine/Plugins/ChaosClothAssetEditor/Source/ChaosClothAssetDataflowNodes/Private/ChaosClothAsset/MergeClothCollectionsNode.cpp b/Engine/Plugins/ChaosClothAssetEditor/Source/ChaosClothAssetDataflowNodes/Private/ChaosClothAsset/MergeClothCollectionsNode.cpp
--- a/Engine/Plugins/ChaosClothAssetEditor/Source/ChaosClothAssetDataflowNodes/Private/ChaosClothAsset/MergeClothCollectionsNode.cpp
+++ b/Engine/Plugins/ChaosClothAssetEditor/Source/ChaosClothAssetDataflowNodes/Private/ChaosClothAsset/MergeClothCollectionsNode.cpp
@@ -544,11 +544,11 @@ void FChaosClothAssetMergeClothCollectionsNode_v2::Evaluate(UE::Dataflow::FConte
 				FSoftObjectPath RemapSkeletalMeshPath;
 				TArray<int32> ThisBoneRemap, OtherBoneRemap;
 				FText IncompatibleErrorDetails;
-				if(FClothEngineTools::CalculateRemappedBoneIndicesIfCompatible(ClothFacade, OtherClothFacade, RemapSkeletalMeshPath, ThisBoneRemap, OtherBoneRemap, &IncompatibleErrorDetails))
+				if (FClothEngineTools::CalculateRemappedBoneIndicesIfCompatible(ClothFacade, OtherClothFacade, RemapSkeletalMeshPath, ThisBoneRemap, OtherBoneRemap, &IncompatibleErrorDetails))
 				{
+					ClothFacade.SetSkeletalMeshSoftObjectPathName(RemapSkeletalMeshPath);  // The path may change even with no remaps (e.g. when the second skeleton is extended from the first one)
 					if (!ThisBoneRemap.IsEmpty())
 					{
-						ClothFacade.SetSkeletalMeshSoftObjectPathName(RemapSkeletalMeshPath);
 						FClothEngineTools::RemapBoneIndices(ClothFacade, ThisBoneRemap);
 					}
 
@@ -896,9 +896,9 @@ void FChaosClothAssetMergeClothCollectionsNode::Evaluate(UE::Dataflow::FContext&
 				FText IncompatibleErrorDetails;
 				if (FClothEngineTools::CalculateRemappedBoneIndicesIfCompatible(ClothFacade, OtherClothFacade, RemapSkeletalMeshPath, ThisBoneRemap, OtherBoneRemap, &IncompatibleErrorDetails))
 				{
+					ClothFacade.SetSkeletalMeshSoftObjectPathName(RemapSkeletalMeshPath);  // The path may change even with no remaps (e.g. when the second skeleton is extended from the first one)
 					if (!ThisBoneRemap.IsEmpty())
 					{
-						ClothFacade.SetSkeletalMeshSoftObjectPathName(RemapSkeletalMeshPath);
 						FClothEngineTools::RemapBoneIndices(ClothFacade, ThisBoneRemap);
 					}
 
