Add new ClothAssetInteractorDataAsset which allows creating sets of cloth properties via a DataAsset which can be set from a single function call on the ClothAssetInteractor

diff --git a/Engine/Plugins/ChaosClothAsset/Source/ChaosClothAssetEngine/ChaosClothAssetEngine.Build.cs b/Engine/Plugins/ChaosClothAsset/Source/ChaosClothAssetEngine/ChaosClothAssetEngine.Build.cs
--- a/Engine/Plugins/ChaosClothAsset/Source/ChaosClothAssetEngine/ChaosClothAssetEngine.Build.cs
+++ b/Engine/Plugins/ChaosClothAsset/Source/ChaosClothAssetEngine/ChaosClothAssetEngine.Build.cs
@@ -36,7 +36,13 @@ public class ChaosClothAssetEngine : ModuleRules
 
 		if (Target.bBuildEditor || Target.bCompileAgainstEditor)
 		{
-			PrivateDependencyModuleNames.Add("PropertyEditor");  // For adding the Cloth Component "Cloth Sim" section to the Details panel UI
+			PrivateDependencyModuleNames.AddRange(
+				new string[]
+				{
+					"PropertyEditor",   // For adding the Cloth Component "Cloth Sim" section to the Details panel UI
+					"StructUtilsEditor"
+				}
+			);
 		}
 
 		PrivateIncludePathModuleNames.AddRange(
diff --git a/Engine/Plugins/ChaosClothAsset/Source/ChaosClothAssetEngine/Private/ChaosClothAsset/ChaosClothAssetEngineModule.cpp b/Engine/Plugins/ChaosClothAsset/Source/ChaosClothAssetEngine/Private/ChaosClothAsset/ChaosClothAssetEngineModule.cpp
--- a/Engine/Plugins/ChaosClothAsset/Source/ChaosClothAssetEngine/Private/ChaosClothAsset/ChaosClothAssetEngineModule.cpp
+++ b/Engine/Plugins/ChaosClothAsset/Source/ChaosClothAssetEngine/Private/ChaosClothAsset/ChaosClothAssetEngineModule.cpp
@@ -8,6 +8,8 @@
 #include "Modules/ModuleManager.h"
 #if WITH_EDITOR
 #include "ChaosClothAsset/ClothComponent.h"
+#include "ChaosClothAsset/ClothAssetInteractorDataAsset.h"
+#include "PropertyBagDetails.h"
 #include "PropertyEditorModule.h"
 #endif
 
@@ -33,6 +35,9 @@ namespace UE::Chaos::ClothAsset
 			FPropertyEditorModule& PropertyModule = FModuleManager::GetModuleChecked<FPropertyEditorModule>("PropertyEditor");
 			const TSharedPtr<FPropertySection> Section = PropertyModule.FindOrCreateSection(UChaosClothComponent::StaticClass()->GetFName(), TEXT("ClothSim"), LOCTEXT("ClothSim", "Cloth Sim"));
 			Section->AddCategory("Cloth Component");
+
+			PropertyModule.RegisterCustomPropertyTypeLayout(FClothAssetInteractorPropertyBag::StaticStruct()->GetFName(), FOnGetPropertyTypeCustomizationInstance::CreateStatic(&FPropertyBagDetails::MakeInstance));
+
 			PropertyModule.NotifyCustomizationModuleChanged();
 #endif
 			IModularFeatures::Get().RegisterModularFeature(IClothingSimulationFactoryClassProvider::FeatureName, static_cast<IClothingSimulationFactoryClassProvider*>(this));
diff --git a/Engine/Plugins/ChaosClothAsset/Source/ChaosClothAssetEngine/Private/ChaosClothAsset/ClothAssetInteractor.cpp b/Engine/Plugins/ChaosClothAsset/Source/ChaosClothAssetEngine/Private/ChaosClothAsset/ClothAssetInteractor.cpp
--- a/Engine/Plugins/ChaosClothAsset/Source/ChaosClothAssetEngine/Private/ChaosClothAsset/ClothAssetInteractor.cpp
+++ b/Engine/Plugins/ChaosClothAsset/Source/ChaosClothAssetEngine/Private/ChaosClothAsset/ClothAssetInteractor.cpp
@@ -1,6 +1,7 @@
 // Copyright Epic Games, Inc. All Rights Reserved.
 
 #include "ChaosClothAsset/ClothAssetInteractor.h"
+#include "ChaosClothAsset/ClothAssetInteractorDataAsset.h"
 #include "Chaos/CollectionPropertyFacade.h"
 
 #include UE_INLINE_GENERATED_CPP_BY_NAME(ClothAssetInteractor)
@@ -388,6 +389,53 @@ FString UChaosClothAssetInteractor::GetStringPropertyValue(const FName PropertyN
 	return DefaultValue;
 }
 
+void UChaosClothAssetInteractor::SetPropertySet(const FClothAssetInteractorPropertyBag& PropertySet, int32 LODIndex)
+{
+	if (const UPropertyBag* const CurrentBagStruct = PropertySet.GetPropertyBagStruct())
+	{
+		for (const FPropertyBagPropertyDesc& Desc : CurrentBagStruct->GetPropertyDescs())
+		{
+			if (Desc.IsNumericFloatType())
+			{
+				if (TValueOrError<float, EPropertyBagResult> Value = PropertySet.GetValueFloat(Desc); Value.IsValid())
+				{
+					SetFloatPropertyValue(Desc.Name, LODIndex, Value.GetValue());
+				}
+			}
+			else if (Desc.IsNumericType())
+			{
+				if (TValueOrError<int32, EPropertyBagResult> Value = PropertySet.GetValueInt32(Desc); Value.IsValid())
+				{
+					SetIntPropertyValue(Desc.Name, LODIndex, Value.GetValue());
+				}
+			}
+			else if (Desc.ValueType == EPropertyBagPropertyType::String)
+			{
+				if (TValueOrError<FString, EPropertyBagResult> Value = PropertySet.GetValueString(Desc); Value.IsValid())
+				{
+					SetStringPropertyValue(Desc.Name, LODIndex, Value.GetValue());
+				}
+			}
+			else if (Desc.ValueType == EPropertyBagPropertyType::Struct)
+			{
+				if (TValueOrError<FStructView, EPropertyBagResult> VectorValue = PropertySet.GetValueStruct(Desc, TBaseStructure<FVector>::Get()); VectorValue.IsValid())
+				{
+					// FVector
+					const FVector& Vector = VectorValue.GetValue().Get<FVector>();
+					SetVectorPropertyValue(Desc.Name, LODIndex, Vector);
+				}
+				else if (TValueOrError<FStructView, EPropertyBagResult> Vector2DValue = PropertySet.GetValueStruct(Desc, TBaseStructure<FVector2D>::Get()); Vector2DValue.IsValid())
+				{
+					// FVector2D (WeightedFloat)
+					const FVector2D& Vector = Vector2DValue.GetValue().Get<FVector2D>();
+					SetWeightedFloatPropertyValue(Desc.Name, LODIndex, Vector);
+				}
+			}
+		}
+	}
+
+}
+
 void UChaosClothAssetInteractor::SetFloatPropertyValue(const FName PropertyName, int32 LODIndex, float Value)
 {
 	using namespace UE::Chaos::ClothAsset::Private;
diff --git a/Engine/Plugins/ChaosClothAsset/Source/ChaosClothAssetEngine/Private/ChaosClothAsset/ClothAssetInteractorDataAsset.cpp b/Engine/Plugins/ChaosClothAsset/Source/ChaosClothAssetEngine/Private/ChaosClothAsset/ClothAssetInteractorDataAsset.cpp
new file mode 100644
--- /dev/null
+++ b/Engine/Plugins/ChaosClothAsset/Source/ChaosClothAssetEngine/Private/ChaosClothAsset/ClothAssetInteractorDataAsset.cpp
@@ -0,0 +1,50 @@
+// Copyright Epic Games, Inc. All Rights Reserved.
+
+#include "ChaosClothAsset/ClothAssetInteractorDataAsset.h"
+
+#include UE_INLINE_GENERATED_CPP_BY_NAME(ClothAssetInteractorDataAsset)
+
+#if WITH_EDITOR
+void UClothAssetInteractorDataAsset::PostEditChangeProperty(struct FPropertyChangedEvent& PropertyChangedEvent)
+{
+	if (PropertyChangedEvent.GetPropertyName() == GET_MEMBER_NAME_CHECKED(UClothAssetInteractorDataAsset, PropertySets) && (PropertyChangedEvent.ChangeType & EPropertyChangeType::ArrayAdd))
+	{
+		// Initialize new property set with default values
+		if (FMapProperty* Map = CastField<FMapProperty>(PropertyChangedEvent.Property))
+		{
+			FScriptMapHelper MapHelper(Map, &PropertySets);
+			const int32 ArrayIndex = PropertyChangedEvent.GetArrayIndex(GET_MEMBER_NAME_STRING_CHECKED(UClothAssetInteractorDataAsset, PropertySets));
+			if (MapHelper.IsValidIndex(ArrayIndex))
+			{
+				FClothAssetInteractorPropertyBag* const ValuePtr = (FClothAssetInteractorPropertyBag*)MapHelper.GetValuePtr(ArrayIndex);
+				ValuePtr->Reset();
+				ValuePtr->MigrateToNewBagInstance(DefaultProperties);
+			}
+		}
+	}
+	Super::PostEditChangeProperty(PropertyChangedEvent);
+}
+#endif
+
+void UClothAssetInteractorDataAsset::SynchronizePropertySets()
+{
+	if (const UPropertyBag* const SourceBagStruct = DefaultProperties.GetPropertyBagStruct())
+	{
+		for (TMap<FName, FClothAssetInteractorPropertyBag>::TIterator PropertyIter = PropertySets.CreateIterator(); PropertyIter; ++PropertyIter)
+		{
+			if (const UPropertyBag* const TargetBagStruct = PropertyIter.Value().GetPropertyBagStruct())
+			{
+				constexpr bool bCheckNamesTrue = true;
+				if (FStructUtils::TheSameLayout(SourceBagStruct, TargetBagStruct, bCheckNamesTrue))
+				{
+					continue;
+				}
+
+				const FClothAssetInteractorPropertyBag OrigTarget = PropertyIter.Value();
+				PropertyIter.Value().Reset();
+				PropertyIter.Value().MigrateToNewBagInstance(DefaultProperties);
+				PropertyIter.Value().CopyMatchingValuesByID(OrigTarget);
+			}
+		}
+	}
+}
\ No newline at end of file
diff --git a/Engine/Plugins/ChaosClothAsset/Source/ChaosClothAssetEngine/Private/ChaosClothAsset/ClothAssetInteractorDataAsset.h b/Engine/Plugins/ChaosClothAsset/Source/ChaosClothAssetEngine/Private/ChaosClothAsset/ClothAssetInteractorDataAsset.h
new file mode 100644
--- /dev/null
+++ b/Engine/Plugins/ChaosClothAsset/Source/ChaosClothAssetEngine/Private/ChaosClothAsset/ClothAssetInteractorDataAsset.h
@@ -0,0 +1,62 @@
+// Copyright Epic Games, Inc. All Rights Reserved.
+
+#pragma once
+
+#include "Engine/DataAsset.h"
+#include "StructUtils/PropertyBag.h"
+
+#include "ClothAssetInteractorDataAsset.generated.h"
+
+#define UE_API CHAOSCLOTHASSETENGINE_API
+
+USTRUCT(BlueprintType)
+struct FClothAssetInteractorPropertyBag : public FInstancedPropertyBag
+{
+	GENERATED_BODY()
+};
+
+template<>
+struct TStructOpsTypeTraits<FClothAssetInteractorPropertyBag> : public TStructOpsTypeTraits<FInstancedPropertyBag>
+{
+};
+
+/**
+ * Data asset for setting groups of properties on the Chaos Cloth Asset Interator.
+ */
+UCLASS(BlueprintType)
+class  UClothAssetInteractorDataAsset : public UDataAsset
+{
+	GENERATED_BODY()
+
+public:
+
+	/** Get the Property Set for a given Set Name. If Set Name does not exist, the Default Properties will be returned.*/
+	UFUNCTION(BlueprintCallable, Category = ClothProperty)
+	const FClothAssetInteractorPropertyBag& GetPropertySet(const FName SetName) const
+	{
+		if (const FClothAssetInteractorPropertyBag* const FoundSet = PropertySets.Find(SetName))
+		{
+			return *FoundSet;
+		}
+		return DefaultProperties;
+	}
+
+	/** Synchronize the Property Sets with the Default Properties by adding/removing properties. Existing property values will not be modified.*/
+	UFUNCTION(CallInEditor, Category = "Interactor Properties", meta = (DisplayPriority = 2))
+	void SynchronizePropertySets();
+
+#if WITH_EDITOR
+	virtual void PostEditChangeProperty(struct FPropertyChangedEvent& PropertyChangedEvent) override;
+#endif
+private:
+
+	/** Default Property schema and default values. Property Names should match the Interactor Name to set using the Cloth Asset Interactor. 
+	 * Only property types known to the interactor (float, double, int, string, Vector, Vector2D) will be used to set values. */
+	UPROPERTY(EditAnywhere, Category = "Interactor Properties")
+	FClothAssetInteractorPropertyBag DefaultProperties;
+
+	/** Property sets with matching schema to Default Properties. */
+	UPROPERTY(EditAnywhere, Category = "Interactor Properties", Meta = (FixedLayout))
+	TMap<FName, FClothAssetInteractorPropertyBag> PropertySets;
+};
+#undef UE_API
diff --git a/Engine/Plugins/ChaosClothAsset/Source/ChaosClothAssetEngine/Public/ChaosClothAsset/ClothAssetInteractor.h b/Engine/Plugins/ChaosClothAsset/Source/ChaosClothAssetEngine/Public/ChaosClothAsset/ClothAssetInteractor.h
--- a/Engine/Plugins/ChaosClothAsset/Source/ChaosClothAssetEngine/Public/ChaosClothAsset/ClothAssetInteractor.h
+++ b/Engine/Plugins/ChaosClothAsset/Source/ChaosClothAssetEngine/Public/ChaosClothAsset/ClothAssetInteractor.h
@@ -10,6 +10,7 @@ namespace Chaos::Softs
 {
 class FCollectionPropertyFacade;
 }
+struct FClothAssetInteractorPropertyBag;
 
 UCLASS(BlueprintType, MinimalAPI)
 class UChaosClothAssetInteractor : public UClothingInteractor
@@ -55,6 +56,10 @@ public:
 	UFUNCTION(BlueprintCallable, Category = ClothProperty)
 	CHAOSCLOTHASSETENGINE_API FString GetStringPropertyValue(const FName PropertyName, int32 LODIndex = 0, const FString& DefaultValue = "") const;
 
+	/** Set properties from an interactor data set */
+	UFUNCTION(BlueprintCallable, Category = ClothProperty)
+	CHAOSCLOTHASSETENGINE_API void SetPropertySet(const FClothAssetInteractorPropertyBag& PropertySet, int32 LODIndex = -1);
+
 	/**Set the value for a property (if it exists). This sets the Low and High values for weighted values. All LODs will be set when LODIndex = -1.*/
 	UFUNCTION(BlueprintCallable, Category = ClothProperty)
 	CHAOSCLOTHASSETENGINE_API void SetFloatPropertyValue(const FName PropertyName, int32 LODIndex = -1, float Value = 0.f);
