Minor fixups from last CL.

diff --git a/Engine/Plugins/ChaosClothAsset/Source/ChaosClothAssetEngine/Private/ChaosClothAsset/ChaosClothAssetEngineModule.cpp b/Engine/Plugins/ChaosClothAsset/Source/ChaosClothAssetEngine/Private/ChaosClothAsset/ChaosClothAssetEngineModule.cpp
--- a/Engine/Plugins/ChaosClothAsset/Source/ChaosClothAssetEngine/Private/ChaosClothAsset/ChaosClothAssetEngineModule.cpp
+++ b/Engine/Plugins/ChaosClothAsset/Source/ChaosClothAssetEngine/Private/ChaosClothAsset/ChaosClothAssetEngineModule.cpp
@@ -52,6 +52,8 @@ namespace UE::Chaos::ClothAsset
 				// Remove sections
 				FPropertyEditorModule& PropertyModule = FModuleManager::GetModuleChecked<FPropertyEditorModule>("PropertyEditor");
 				PropertyModule.RemoveSection(UChaosClothComponent::StaticClass()->GetFName(), TEXT("ClothSim"));
+
+				PropertyModule.UnregisterCustomPropertyTypeLayout(FClothAssetInteractorPropertyBag::StaticStruct()->GetFName());
 			}
 #endif
 		}
diff --git a/Engine/Plugins/ChaosClothAsset/Source/ChaosClothAssetEngine/Private/ChaosClothAsset/ClothAssetInteractorDataAsset.h b/Engine/Plugins/ChaosClothAsset/Source/ChaosClothAssetEngine/Private/ChaosClothAsset/ClothAssetInteractorDataAsset.h
--- a/Engine/Plugins/ChaosClothAsset/Source/ChaosClothAssetEngine/Private/ChaosClothAsset/ClothAssetInteractorDataAsset.h
+++ b/Engine/Plugins/ChaosClothAsset/Source/ChaosClothAssetEngine/Private/ChaosClothAsset/ClothAssetInteractorDataAsset.h
@@ -21,7 +21,7 @@ struct TStructOpsTypeTraits<FClothAssetInteractorPropertyBag> : public TStructOp
 };
 
 /**
- * Data asset for setting groups of properties on the Chaos Cloth Asset Interator.
+ * Data asset for setting groups of properties on the Chaos Cloth Asset Interactor.
  */
 UCLASS(BlueprintType)
 class  UClothAssetInteractorDataAsset : public UDataAsset
