Chaos SKM Cloth Asset - Added defomer mapping from the cloth asset own render data.

[RN] Chaos SKM Cloth Asset - Added defomer mapping from the Cloth Asset own render data to the Skeletal Mesh render sections.

diff --git a/Engine/Plugins/ChaosClothAsset/Source/ChaosClothAssetEngine/Private/ChaosClothAsset/ClothAssetSKMClothingAsset.cpp b/Engine/Plugins/ChaosClothAsset/Source/ChaosClothAssetEngine/Private/ChaosClothAsset/ClothAssetSKMClothingAsset.cpp
--- a/Engine/Plugins/ChaosClothAsset/Source/ChaosClothAssetEngine/Private/ChaosClothAsset/ClothAssetSKMClothingAsset.cpp
+++ b/Engine/Plugins/ChaosClothAsset/Source/ChaosClothAssetEngine/Private/ChaosClothAsset/ClothAssetSKMClothingAsset.cpp
@@ -12,6 +12,7 @@
 #include "Components/SkeletalMeshComponent.h"
 #include "Engine/SkeletalMesh.h"
 #include "Rendering/SkeletalMeshModel.h"
+#include "Rendering/SkeletalMeshRenderData.h"
 #include "Utils/ClothingMeshUtils.h"
 #include "ClothingAsset.h"
 #endif  // #if WITH_EDITOR
@@ -354,53 +355,113 @@ bool UChaosClothAssetSKMClothingAsset::BindToSkeletalMesh(USkeletalMesh* Skeleta
 		SkeletalMesh->GetRefSkeleton().EnsureParentsExistAndSort(SkeletalMeshLODModel.ActiveBoneIndices);
 	}
 
-	// Calculate the deformer mappings
-	const int32 NumIndices = SkelMeshSection.NumTriangles * 3;
-	const int32 NumVertices = SkelMeshSection.SoftVertices.Num();
-	const uint32 BaseIndex = SkelMeshSection.BaseIndex;
-	const uint32 BaseVertexIndex = SkelMeshSection.BaseVertexIndex;
-
-	TArray<FVector3f> RenderPositions;
-	TArray<FVector3f> RenderNormals;
-	TArray<FVector3f> RenderTangents;
-	RenderPositions.Reserve(NumVertices);
-	RenderNormals.Reserve(NumVertices);
-	RenderTangents.Reserve(NumVertices);
-	for (const FSoftSkinVertex& SoftSkinVertex : SkelMeshSection.SoftVertices)
+	// Set the deformer mappings
+	int32 MatchingDeformerSectionIndex = INDEX_NONE;
+	if (const FSkeletalMeshRenderData* const AssetRenderData = Asset->GetResourceForRendering())
 	{
-		RenderPositions.Add(SoftSkinVertex.Position);
-		RenderNormals.Add(SoftSkinVertex.TangentZ);
-		RenderTangents.Add(SoftSkinVertex.TangentX);
-	}
+		// Look for a suitable Cloth/Outfit Asset render section to copy the deformer mapping from
+		if (AssetRenderData->LODRenderData.IsValidIndex(AssetLodIndex))
+		{
+			const FSkeletalMeshLODRenderData& AssetLODRenderData = AssetRenderData->LODRenderData[AssetLodIndex];
+			const FPositionVertexBuffer& AssetPositionVertexBuffer = AssetLODRenderData.StaticVertexBuffers.PositionVertexBuffer;
 
-	const TConstArrayView<uint32> SectionRenderIndices(SkeletalMeshLODModel.IndexBuffer.GetData() + BaseIndex, NumIndices);
-	TArray<uint32> RenderIndices;
-	RenderIndices.Reserve(SectionRenderIndices.Num());
-	for (const uint32 SectionRenderIndex : SectionRenderIndices)
-	{
-		RenderIndices.Add(SectionRenderIndex - BaseVertexIndex);
-	}
+			TArray<FSoftSkinVertex> SkelMeshVertices;
+			SkeletalMeshLODModel.GetVertices(SkelMeshVertices);
+			const int32 SkelMeshVertexOffset = SkelMeshSection.BaseVertexIndex;
 
-	const ClothingMeshUtils::ClothMeshDesc TargetMesh(RenderPositions, RenderNormals, RenderTangents, RenderIndices);
-	const ClothingMeshUtils::ClothMeshDesc SourceMesh(
-		ClothSimulationModel->GetPositions(AssetLodIndex),
-		ClothSimulationModel->GetIndices(AssetLodIndex));  // Let it calculate the averaged normals as to match the simulation data output
+			for (int32 DeformerSectionIndex = 0; DeformerSectionIndex < AssetLODRenderData.RenderSections.Num(); ++DeformerSectionIndex)
+			{
+				const FSkelMeshRenderSection& AssetRenderSection = AssetLODRenderData.RenderSections[DeformerSectionIndex];
+				if (AssetRenderSection.CorrespondClothAssetIndex == ClothSimulationModelIndex &&  // The model index must match when binding with an Outfit Asset
+					AssetRenderSection.HasClothingData() &&  // The section must have mapping to copy from
+					AssetRenderSection.GetNumVertices() == SkelMeshSection.GetNumVertices())
+				{
+					const int32 AssetVertexOffset = AssetRenderSection.BaseVertexIndex;
 
-	const TSharedRef<const FManagedArrayCollection>& Collection = Asset->GetCollections(ClothSimulationModelIndex)[AssetLodIndex];
-	FCollectionClothConstFacade ClothFacade(Collection);
-	::Chaos::Softs::FCollectionPropertyConstFacade PropertyFacade(Collection);
-	const FPointWeightMap MaxDistances = FClothEngineTools::GetMaxDistanceWeightMap(ClothFacade, PropertyFacade, ClothSimulationModel->GetNumVertices(AssetLodIndex));
+					bool bVertexPositionsMatch = true;
+					for (int32 VertexIndex = 0; VertexIndex < AssetRenderSection.GetNumVertices(); ++VertexIndex)
+					{
+						const int32 SkelMeshVertexIndex = SkelMeshVertexOffset + VertexIndex;
+						const int32 AssetVertexIndex = AssetVertexOffset + VertexIndex;
+
+						if (UE_SMALL_NUMBER < FVector3f::DistSquared(
+							AssetPositionVertexBuffer.VertexPosition(AssetVertexIndex),
+							SkelMeshVertices[SkelMeshVertexIndex].Position))
+						{
+							bVertexPositionsMatch = false;
+							break;
+						}
+					}
+					if (bVertexPositionsMatch)
+					{
+						SkelMeshSection.ClothMappingDataLODs.SetNum(1);
+						SkelMeshSection.ClothMappingDataLODs[0] = AssetRenderSection.ClothMappingDataLODs[0];
+						MatchingDeformerSectionIndex = DeformerSectionIndex;
+						break;
+					}
+				}
+			}
+			if (MatchingDeformerSectionIndex == INDEX_NONE && AssetLODRenderData.RenderSections.Num())
+			{
+				const FText Text = FText::Format(
+					LOCTEXT("NoMatchingRenderDeformerSection",
+						"No matching render section found in Clothing Data [{0}] Asset [{1}] to copy render deformer data from, even though the Cloth/Outfit Asset contains render data."),
+					FText::FromString(GetName()),
+					FText::FromName(Asset->GetFName()));
+				WarningNotification(Text);
+			}
+		}
+	}
+	if (MatchingDeformerSectionIndex == INDEX_NONE)
+	{
+		// Calculate the deformer mappings instead if no section was found
+		const int32 NumIndices = SkelMeshSection.NumTriangles * 3;
+		const int32 NumVertices = SkelMeshSection.SoftVertices.Num();
+		const uint32 BaseIndex = SkelMeshSection.BaseIndex;
+		const uint32 BaseVertexIndex = SkelMeshSection.BaseVertexIndex;
+
+		TArray<FVector3f> RenderPositions;
+		TArray<FVector3f> RenderNormals;
+		TArray<FVector3f> RenderTangents;
+		RenderPositions.Reserve(NumVertices);
+		RenderNormals.Reserve(NumVertices);
+		RenderTangents.Reserve(NumVertices);
+		for (const FSoftSkinVertex& SoftSkinVertex : SkelMeshSection.SoftVertices)
+		{
+			RenderPositions.Add(SoftSkinVertex.Position);
+			RenderNormals.Add(SoftSkinVertex.TangentZ);
+			RenderTangents.Add(SoftSkinVertex.TangentX);
+		}
 
-	constexpr bool bSmoothTransition = true;
-	constexpr bool bUseMultipleInfluences = false;
-	constexpr float SkinningKernelRadius = 100.f;
-	TArray<FMeshToMeshVertData> MeshToMeshData;
-	ClothingMeshUtils::GenerateMeshToMeshVertData(
-		MeshToMeshData, TargetMesh, SourceMesh, &MaxDistances,
-		bSmoothTransition, bUseMultipleInfluences, SkinningKernelRadius);
+		const TConstArrayView<uint32> SectionRenderIndices(SkeletalMeshLODModel.IndexBuffer.GetData() + BaseIndex, NumIndices);
+		TArray<uint32> RenderIndices;
+		RenderIndices.Reserve(SectionRenderIndices.Num());
+		for (const uint32 SectionRenderIndex : SectionRenderIndices)
+		{
+			RenderIndices.Add(SectionRenderIndex - BaseVertexIndex);
+		}
 
-	SkelMeshSection.ClothMappingDataLODs.SetNum(1);
-	SkelMeshSection.ClothMappingDataLODs[0] = MoveTemp(MeshToMeshData);
+		const ClothingMeshUtils::ClothMeshDesc TargetMesh(RenderPositions, RenderNormals, RenderTangents, RenderIndices);
+		const ClothingMeshUtils::ClothMeshDesc SourceMesh(
+			ClothSimulationModel->GetPositions(AssetLodIndex),
+			ClothSimulationModel->GetIndices(AssetLodIndex));  // Let it calculate the averaged normals as to match the simulation data output
+
+		const TSharedRef<const FManagedArrayCollection>& Collection = Asset->GetCollections(ClothSimulationModelIndex)[AssetLodIndex];
+		FCollectionClothConstFacade ClothFacade(Collection);
+		::Chaos::Softs::FCollectionPropertyConstFacade PropertyFacade(Collection);
+		const FPointWeightMap MaxDistances = FClothEngineTools::GetMaxDistanceWeightMap(ClothFacade, PropertyFacade, ClothSimulationModel->GetNumVertices(AssetLodIndex));
+
+		constexpr bool bSmoothTransition = true;
+		constexpr bool bUseMultipleInfluences = false;
+		constexpr float SkinningKernelRadius = 100.f;
+		TArray<FMeshToMeshVertData> MeshToMeshData;
+		ClothingMeshUtils::GenerateMeshToMeshVertData(
+			MeshToMeshData, TargetMesh, SourceMesh, &MaxDistances,
+			bSmoothTransition, bUseMultipleInfluences, SkinningKernelRadius);
+
+		SkelMeshSection.ClothMappingDataLODs.SetNum(1);
+		SkelMeshSection.ClothMappingDataLODs[0] = MoveTemp(MeshToMeshData);
+	}
 
 	// Update the extra cloth deformer mapping LOD bias using this cloth entry
 	// TODO: Implement a generic form of LOD Bias mapping in ClothingAssetUtils (need to rework the clothing asset interface for this)
@@ -446,6 +507,7 @@ void UChaosClothAssetSKMClothingAsset::OnModelChanged()
 		TArray<ClothingAssetUtils::FClothingAssetMeshBinding> ClothingBindings;
 		UnbindClothingAndBackup(*SkeletalMesh, *this, ClothingBindings);
 		RestoreClothingFromBackup(*SkeletalMesh, *this, ClothingBindings);
+		SkeletalMesh->InvalidateDeriveDataCacheGUID();  // The DDC key only changes with the ClothingAsset GUID, but a model change still requires a mesh rebuild
 	}
 	CalculateLODHasAnySimulationMeshData();
 }
