Chaos Cloth: -fix static FName initialization order

diff --git a/Engine/Plugins/ChaosClothAsset/Source/ChaosClothAsset/Private/ChaosClothAsset/ClothCollection.cpp b/Engine/Plugins/ChaosClothAsset/Source/ChaosClothAsset/Private/ChaosClothAsset/ClothCollection.cpp
--- a/Engine/Plugins/ChaosClothAsset/Source/ChaosClothAsset/Private/ChaosClothAsset/ClothCollection.cpp
+++ b/Engine/Plugins/ChaosClothAsset/Source/ChaosClothAsset/Private/ChaosClothAsset/ClothCollection.cpp
@@ -7,204 +7,286 @@
 
 namespace UE::Chaos::ClothAsset::Private
 {
+	// Note: static variables that depend on other static variables from other translation units must be accessed through getter functions,
+	// otherwise static variables can contain uninitialized values.
+
 	// Lods Group
-	static const TArray<FName> LodsGroupAttributes =
+	static const TArray<FName>& GetLodsGroupAttributes()
 	{
-		ClothCollectionAttribute::PhysicsAssetSoftObjectPathName,
-		ClothCollectionAttribute::SkeletalMeshSoftObjectPathName,
-		ClothCollectionAttribute::ReferenceBoneName,
-	};
+		static const TArray<FName> LodsGroupAttributes =
+		{
+			ClothCollectionAttribute::PhysicsAssetSoftObjectPathName,
+			ClothCollectionAttribute::SkeletalMeshSoftObjectPathName,
+			ClothCollectionAttribute::ReferenceBoneName,
+		};
+		return LodsGroupAttributes;
+	}
 
 	// Solvers Group
-	static const TArray<FName> SolverGroupAttributes =
+	static const TArray<FName>& GetSolverGroupAttributes()
 	{
-		ClothCollectionAttribute::SolverGravity,
-		ClothCollectionAttribute::SolverAirDamping,
-		ClothCollectionAttribute::SolverSubSteps,
-		ClothCollectionAttribute::SolverTimeStep,
-	};
+		static const TArray<FName> SolverGroupAttributes =
+		{
+			ClothCollectionAttribute::SolverGravity,
+			ClothCollectionAttribute::SolverAirDamping,
+			ClothCollectionAttribute::SolverSubSteps,
+			ClothCollectionAttribute::SolverTimeStep,
+		};
+		return SolverGroupAttributes;
+	}
 
 	// Fabrics Group
-	static const TArray<FName> SimFabricGroupAttributes =
+	static const TArray<FName>& GetSimFabricGroupAttributes()
 	{
-		ClothCollectionAttribute::FabricBendingStiffness,
-		ClothCollectionAttribute::FabricBucklingStiffness,
-		ClothCollectionAttribute::FabricStretchStiffness,
-		ClothCollectionAttribute::FabricBucklingRatio,
-		ClothCollectionAttribute::FabricDensity,
-		ClothCollectionAttribute::FabricFriction,
-		ClothCollectionAttribute::FabricDamping,
-		ClothCollectionAttribute::FabricPressure,
-		ClothCollectionAttribute::FabricLayer,
-		ClothCollectionAttribute::FabricCollisionThickness,
-	};
+		static const TArray<FName> SimFabricGroupAttributes =
+		{
+			ClothCollectionAttribute::FabricBendingStiffness,
+			ClothCollectionAttribute::FabricBucklingStiffness,
+			ClothCollectionAttribute::FabricStretchStiffness,
+			ClothCollectionAttribute::FabricBucklingRatio,
+			ClothCollectionAttribute::FabricDensity,
+			ClothCollectionAttribute::FabricFriction,
+			ClothCollectionAttribute::FabricDamping,
+			ClothCollectionAttribute::FabricPressure,
+			ClothCollectionAttribute::FabricLayer,
+			ClothCollectionAttribute::FabricCollisionThickness,
+		};
+		return SimFabricGroupAttributes;
+	}
 
 	// Seam Group
-	static const TArray<FName> SeamsGroupAttributes =
+	static const TArray<FName>& GetSeamsGroupAttributes()
 	{
-		ClothCollectionAttribute::SeamStitchStart,
-		ClothCollectionAttribute::SeamStitchEnd
-	};
+		static const TArray<FName> SeamsGroupAttributes =
+		{
+			ClothCollectionAttribute::SeamStitchStart,
+			ClothCollectionAttribute::SeamStitchEnd
+		};
+		return SeamsGroupAttributes;
+	}
 
 	// Seam Stitches Group
-	static const TArray<FName> SeamStitchesGroupAttributes =
+	static const TArray<FName>& GetSeamStitchesGroupAttributes()
 	{
-		ClothCollectionAttribute::SeamStitch2DEndIndices,
-		ClothCollectionAttribute::SeamStitch3DIndex
-	};
+		static const TArray<FName> SeamStitchesGroupAttributes =
+		{
+			ClothCollectionAttribute::SeamStitch2DEndIndices,
+			ClothCollectionAttribute::SeamStitch3DIndex
+		};
+		return SeamStitchesGroupAttributes;
+	}
 
 	// Sim Patterns Group
-	static const TArray<FName> SimPatternsGroupAttributes =
-	{ 
-		ClothCollectionAttribute::SimVertices2DStart,
-		ClothCollectionAttribute::SimVertices2DEnd,
-		ClothCollectionAttribute::SimFacesStart,
-		ClothCollectionAttribute::SimFacesEnd,
-		ClothCollectionAttribute::SimPatternFabric,
-		
-	};
+	static const TArray<FName>& GetSimPatternsGroupAttributes()
+	{
+		static const TArray<FName> SimPatternsGroupAttributes =
+		{
+			ClothCollectionAttribute::SimVertices2DStart,
+			ClothCollectionAttribute::SimVertices2DEnd,
+			ClothCollectionAttribute::SimFacesStart,
+			ClothCollectionAttribute::SimFacesEnd,
+			ClothCollectionAttribute::SimPatternFabric,
+		};
+		return SimPatternsGroupAttributes;
+	}
 
 	// Render Patterns Group
-	static const TArray<FName> RenderPatternsGroupAttributes =
+	static const TArray<FName>& GetRenderPatternsGroupAttributes()
 	{
-		ClothCollectionAttribute::RenderVerticesStart,
-		ClothCollectionAttribute::RenderVerticesEnd,
-		ClothCollectionAttribute::RenderFacesStart,
-		ClothCollectionAttribute::RenderFacesEnd,
-		ClothCollectionAttribute::RenderDeformerNumInfluences,
-		ClothCollectionAttribute::RenderMaterialSoftObjectPathName
-	};
+		static const TArray<FName> RenderPatternsGroupAttributes =
+		{
+			ClothCollectionAttribute::RenderVerticesStart,
+			ClothCollectionAttribute::RenderVerticesEnd,
+			ClothCollectionAttribute::RenderFacesStart,
+			ClothCollectionAttribute::RenderFacesEnd,
+			ClothCollectionAttribute::RenderDeformerNumInfluences,
+			ClothCollectionAttribute::RenderMaterialSoftObjectPathName
+		};
+		return RenderPatternsGroupAttributes;
+	}
 
 	// Sim Faces Group
-	static const TArray<FName> SimFacesGroupAttributes =
+	static const TArray<FName>& GetSimFacesGroupAttributes()
 	{
-		ClothCollectionAttribute::SimIndices2D,
-		ClothCollectionAttribute::SimIndices3D
-	};
+		static const TArray<FName> SimFacesGroupAttributes =
+		{
+			ClothCollectionAttribute::SimIndices2D,
+			ClothCollectionAttribute::SimIndices3D
+		};
+		return SimFacesGroupAttributes;
+	}
 
 	// Sim Vertices 2D Group
-	static const TArray<FName> SimVertices2DGroupAttributes =
+	static const TArray<FName>& GetSimVertices2DGroupAttributes()
 	{
-		ClothCollectionAttribute::SimPosition2D,
-		ClothCollectionAttribute::SimVertex3DLookup,
-		ClothCollectionAttribute::SimImportVertexID,
-	};
+		static const TArray<FName> SimVertices2DGroupAttributes =
+		{
+			ClothCollectionAttribute::SimPosition2D,
+			ClothCollectionAttribute::SimVertex3DLookup,
+			ClothCollectionAttribute::SimImportVertexID,
+		};
+		return SimVertices2DGroupAttributes;
+	}
 
 	// Sim Vertices 3D Group
 	// NOTE: if you add anything here, you need to implement how to merge it 
 	// in CollectionClothSeamFacade. Otherwise, the data in the lowest vertex index 
 	// will survive and the other data will be lost whenever seams are added.
-	static const TArray<FName> SimVertices3DGroupAttributes =
+	static const TArray<FName>& GetSimVertices3DGroupAttributes()
 	{
-		ClothCollectionAttribute::SimPosition3D,
-		ClothCollectionAttribute::PreResizedSimPosition3D,
-		ClothCollectionAttribute::SimNormal,
-		ClothCollectionAttribute::SimBoneIndices,
-		ClothCollectionAttribute::SimBoneWeights,
-		ClothCollectionAttribute::SimVertex2DLookup,
-		ClothCollectionAttribute::SeamStitchLookup,
-		ClothCollectionAttribute::SimCustomResizingBlend,
-	};
-	static const TArray<FName> SimAccessoryMeshVertices3DGroupAttributePrefixes =
+		static const TArray<FName> SimVertices3DGroupAttributes =
+		{
+			ClothCollectionAttribute::SimPosition3D,
+			ClothCollectionAttribute::PreResizedSimPosition3D,
+			ClothCollectionAttribute::SimNormal,
+			ClothCollectionAttribute::SimBoneIndices,
+			ClothCollectionAttribute::SimBoneWeights,
+			ClothCollectionAttribute::SimVertex2DLookup,
+			ClothCollectionAttribute::SeamStitchLookup,
+			ClothCollectionAttribute::SimCustomResizingBlend,
+		};
+		return SimVertices3DGroupAttributes;
+	}
+
+	static const TArray<FName>& GetSimAccessoryMeshVertices3DGroupAttributePrefixes()
 	{
-		ClothCollectionAttribute::SimAccessoryMeshPosition3DPrefix,
-		ClothCollectionAttribute::SimAccessoryMeshNormalPrefix,
-		ClothCollectionAttribute::SimAccessoryMeshBoneIndicesPrefix,
-		ClothCollectionAttribute::SimAccessoryMeshBoneWeightsPrefix,
-	};
+		static const TArray<FName> SimAccessoryMeshVertices3DGroupAttributePrefixes =
+		{
+			ClothCollectionAttribute::SimAccessoryMeshPosition3DPrefix,
+			ClothCollectionAttribute::SimAccessoryMeshNormalPrefix,
+			ClothCollectionAttribute::SimAccessoryMeshBoneIndicesPrefix,
+			ClothCollectionAttribute::SimAccessoryMeshBoneWeightsPrefix,
+		};
+		return SimAccessoryMeshVertices3DGroupAttributePrefixes;
+	}
 
 	//~ Sim Morph Targets Group
-	static const TArray<FName> SimMorphTargetsGroupAttributes =
+	static const TArray<FName>& GetSimMorphTargetsGroupAttributes()
 	{
-		ClothCollectionAttribute::SimMorphTargetName,
-		ClothCollectionAttribute::SimMorphTargetVerticesStart,
-		ClothCollectionAttribute::SimMorphTargetVerticesEnd,
-	};
+		static const TArray<FName> SimMorphTargetsGroupAttributes =
+		{
+			ClothCollectionAttribute::SimMorphTargetName,
+			ClothCollectionAttribute::SimMorphTargetVerticesStart,
+			ClothCollectionAttribute::SimMorphTargetVerticesEnd,
+		};
+		return SimMorphTargetsGroupAttributes;
+	}
 
 	//~ Sim Morph Targets Vertices Group
-	static const TArray<FName> SimMorphTargetVerticesGroupAttributes =
+	static const TArray<FName>& GetSimMorphTargetVerticesGroupAttributes()
 	{
-		ClothCollectionAttribute::SimMorphTargetPositionDelta,
-		ClothCollectionAttribute::SimMorphTargetTangentZDelta,
-		ClothCollectionAttribute::SimMorphTargetSimVertex3DIndex,
-	};
+		static const TArray<FName> SimMorphTargetVerticesGroupAttributes =
+		{
+			ClothCollectionAttribute::SimMorphTargetPositionDelta,
+			ClothCollectionAttribute::SimMorphTargetTangentZDelta,
+			ClothCollectionAttribute::SimMorphTargetSimVertex3DIndex,
+		};
+		return SimMorphTargetVerticesGroupAttributes;
+	}
 
 	// Render Faces Group
-	static const TArray<FName> RenderFacesGroupAttributes =
+	static const TArray<FName>& GetRenderFacesGroupAttributes()
 	{
-		ClothCollectionAttribute::RenderIndices,
-	};
+		static const TArray<FName> RenderFacesGroupAttributes =
+		{
+			ClothCollectionAttribute::RenderIndices,
+		};
+		return RenderFacesGroupAttributes;
+	}
 
 	// Render Vertices Group
-	static const TArray<FName> RenderVerticesGroupAttributes =
+	static const TArray<FName>& GetRenderVerticesGroupAttributes()
 	{
-		ClothCollectionAttribute::RenderPosition,
-		ClothCollectionAttribute::RenderNormal,
-		ClothCollectionAttribute::RenderTangentU,
-		ClothCollectionAttribute::RenderTangentV,
-		ClothCollectionAttribute::RenderUVs,
-		ClothCollectionAttribute::RenderColor,
-		ClothCollectionAttribute::RenderBoneIndices,
-		ClothCollectionAttribute::RenderBoneWeights,
-		ClothCollectionAttribute::RenderDeformerPositionBaryCoordsAndDist,
-		ClothCollectionAttribute::RenderDeformerNormalBaryCoordsAndDist,
-		ClothCollectionAttribute::RenderDeformerTangentBaryCoordsAndDist,
-		ClothCollectionAttribute::RenderDeformerSimIndices3D,
-		ClothCollectionAttribute::RenderDeformerWeight,
-		ClothCollectionAttribute::RenderDeformerSkinningBlend,
-		ClothCollectionAttribute::RenderCustomResizingBlend,
-	};
+		static const TArray<FName> RenderVerticesGroupAttributes =
+		{
+			ClothCollectionAttribute::RenderPosition,
+			ClothCollectionAttribute::RenderNormal,
+			ClothCollectionAttribute::RenderTangentU,
+			ClothCollectionAttribute::RenderTangentV,
+			ClothCollectionAttribute::RenderUVs,
+			ClothCollectionAttribute::RenderColor,
+			ClothCollectionAttribute::RenderBoneIndices,
+			ClothCollectionAttribute::RenderBoneWeights,
+			ClothCollectionAttribute::RenderDeformerPositionBaryCoordsAndDist,
+			ClothCollectionAttribute::RenderDeformerNormalBaryCoordsAndDist,
+			ClothCollectionAttribute::RenderDeformerTangentBaryCoordsAndDist,
+			ClothCollectionAttribute::RenderDeformerSimIndices3D,
+			ClothCollectionAttribute::RenderDeformerWeight,
+			ClothCollectionAttribute::RenderDeformerSkinningBlend,
+			ClothCollectionAttribute::RenderCustomResizingBlend,
+		};
+		return RenderVerticesGroupAttributes;
+	}
 
 	//~ Group Resizing Group
-	static const TArray<FName> GroupResizingGroupAttributes =
+	static const TArray<FName>& GetGroupResizingGroupAttributes()
 	{
-		ClothCollectionAttribute::CustomResizingRegionSet,
-		ClothCollectionAttribute::CustomResizingRegionType
-	};
+		static const TArray<FName> GroupResizingGroupAttributes =
+		{
+			ClothCollectionAttribute::CustomResizingRegionSet,
+			ClothCollectionAttribute::CustomResizingRegionType
+		};
+		return GroupResizingGroupAttributes;
+	}
 
 	//~ Sim Accessory Meshes Group
-	static const TArray<FName> SimAccessoryMeshesGroupAttributes =
+	static const TArray<FName>& GetSimAccessoryMeshesGroupAttributes()
 	{
-		ClothCollectionAttribute::SimAccessoryMeshName,
-		ClothCollectionAttribute::SimAccessoryMeshPosition3DAttribute,
-		ClothCollectionAttribute::SimAccessoryMeshNormalAttribute,
-		ClothCollectionAttribute::SimAccessoryMeshBoneIndicesAttribute,
-		ClothCollectionAttribute::SimAccessoryMeshBoneWeightsAttribute,
-	};
+		static const TArray<FName> SimAccessoryMeshesGroupAttributes =
+		{
+			ClothCollectionAttribute::SimAccessoryMeshName,
+			ClothCollectionAttribute::SimAccessoryMeshPosition3DAttribute,
+			ClothCollectionAttribute::SimAccessoryMeshNormalAttribute,
+			ClothCollectionAttribute::SimAccessoryMeshBoneIndicesAttribute,
+			ClothCollectionAttribute::SimAccessoryMeshBoneWeightsAttribute,
+		};
+		return SimAccessoryMeshesGroupAttributes;
+	}
 
 	// All fixed attributes for this collection
-	static const TMap<FName, TArray<FName>> FixedAttributeNamesMap =
+	static const TMap<FName, TArray<FName>>& GetFixedAttributeNamesMap()
 	{
-		{ ClothCollectionGroup::Lods, LodsGroupAttributes },
-		{ ClothCollectionGroup::Seams, SeamsGroupAttributes },
-		{ ClothCollectionGroup::SeamStitches, SeamStitchesGroupAttributes },
-		{ ClothCollectionGroup::SimPatterns, SimPatternsGroupAttributes },
-		{ ClothCollectionGroup::RenderPatterns, RenderPatternsGroupAttributes },
-		{ ClothCollectionGroup::SimFaces, SimFacesGroupAttributes },
-		{ ClothCollectionGroup::SimVertices2D, SimVertices2DGroupAttributes },
-		{ ClothCollectionGroup::SimVertices3D, SimVertices3DGroupAttributes },
-		{ ClothCollectionGroup::SimMorphTargets, SimMorphTargetsGroupAttributes },
-		{ ClothCollectionGroup::SimMorphTargetVertices, SimMorphTargetVerticesGroupAttributes },
-		{ ClothCollectionGroup::RenderFaces, RenderFacesGroupAttributes },
-		{ ClothCollectionGroup::RenderVertices, RenderVerticesGroupAttributes },
-		{ ClothCollectionGroup::Fabrics, SimFabricGroupAttributes },
-		{ ClothCollectionGroup::Solvers, SolverGroupAttributes },
-		{ ClothCollectionGroup::CustomResizingRegions, GroupResizingGroupAttributes },
-		{ ClothCollectionGroup::SimAccessoryMeshes, SimAccessoryMeshesGroupAttributes },
-	};
+		static const TMap<FName, TArray<FName>> FixedAttributeNamesMap =
+		{
+			{ ClothCollectionGroup::Lods, GetLodsGroupAttributes() },
+			{ ClothCollectionGroup::Seams, GetSeamsGroupAttributes() },
+			{ ClothCollectionGroup::SeamStitches, GetSeamStitchesGroupAttributes() },
+			{ ClothCollectionGroup::SimPatterns, GetSimPatternsGroupAttributes() },
+			{ ClothCollectionGroup::RenderPatterns, GetRenderPatternsGroupAttributes() },
+			{ ClothCollectionGroup::SimFaces, GetSimFacesGroupAttributes() },
+			{ ClothCollectionGroup::SimVertices2D, GetSimVertices2DGroupAttributes() },
+			{ ClothCollectionGroup::SimVertices3D, GetSimVertices3DGroupAttributes() },
+			{ ClothCollectionGroup::SimMorphTargets, GetSimMorphTargetsGroupAttributes() },
+			{ ClothCollectionGroup::SimMorphTargetVertices, GetSimMorphTargetVerticesGroupAttributes() },
+			{ ClothCollectionGroup::RenderFaces, GetRenderFacesGroupAttributes() },
+			{ ClothCollectionGroup::RenderVertices, GetRenderVerticesGroupAttributes() },
+			{ ClothCollectionGroup::Fabrics, GetSimFabricGroupAttributes() },
+			{ ClothCollectionGroup::Solvers, GetSolverGroupAttributes() },
+			{ ClothCollectionGroup::CustomResizingRegions, GetGroupResizingGroupAttributes() },
+			{ ClothCollectionGroup::SimAccessoryMeshes, GetSimAccessoryMeshesGroupAttributes() },
+		};
+		return FixedAttributeNamesMap;
+	}
 
 	// All fixed attribute prefixes for this collection (see SimAccessoryMesh SimVertex3D)
-	static const TMap<FName, TArray<FName>> FixedAttributePrefixNamesMap =
+	static const TMap<FName, TArray<FName>>& GetFixedAttributePrefixNamesMap()
 	{
-		{ ClothCollectionGroup::SimVertices3D, SimAccessoryMeshVertices3DGroupAttributePrefixes },
-	};
+		static const TMap<FName, TArray<FName>> FixedAttributePrefixNamesMap = {
+			{ ClothCollectionGroup::SimVertices3D, GetSimAccessoryMeshVertices3DGroupAttributePrefixes() },
+		};
+		return FixedAttributePrefixNamesMap;
+	}
 
 	// All paintable attributes for this collection
-	static const TMap<FName, TArray<FName>> UserAccessibleAttributeNamesMap =
+	static const TMap<FName, TArray<FName>>& GetUserAccessibleAttributeNamesMap()
 	{
-		{ ClothCollectionGroup::SimVertices3D, { ClothCollectionAttribute::SimCustomResizingBlend } },  // Can be painted if required
-		{ ClothCollectionGroup::RenderVertices, { ClothCollectionAttribute::RenderDeformerSkinningBlend, ClothCollectionAttribute::RenderCustomResizingBlend } },  // Can be painted if required
-	};
+		static const TMap<FName, TArray<FName>> UserAccessibleAttributeNamesMap =
+		{
+			{ ClothCollectionGroup::SimVertices3D, { ClothCollectionAttribute::SimCustomResizingBlend } },  // Can be painted if required
+			{ ClothCollectionGroup::RenderVertices, { ClothCollectionAttribute::RenderDeformerSkinningBlend, ClothCollectionAttribute::RenderCustomResizingBlend } },  // Can be painted if required
+		};
+		return UserAccessibleAttributeNamesMap;
+	}
 
 	static bool NamesHaveMatchingPrefix(const FName& Name, const FName& Prefix)
 	{
@@ -988,9 +1070,9 @@ namespace UE::Chaos::ClothAsset
 
 		TArray<FName> UserDefinedAttributeNames;
 
-		if (const TArray<FName>* const FixedAttributeNames = FixedAttributeNamesMap.Find(GroupName))  // Also checks that the group name is a recognized group name
+		if (const TArray<FName>* const FixedAttributeNames = GetFixedAttributeNamesMap().Find(GroupName))  // Also checks that the group name is a recognized group name
 		{
-			const TArray<FName>* const FixedAttributePrefixNames = FixedAttributePrefixNamesMap.Find(GroupName);
+			const TArray<FName>* const FixedAttributePrefixNames = GetFixedAttributePrefixNamesMap().Find(GroupName);
 
 			const TArray<FName> AttributeNames = ManagedArrayCollection->AttributeNames(GroupName);
 
@@ -1019,7 +1101,7 @@ namespace UE::Chaos::ClothAsset
 	TManagedArray<T>* FClothCollection::FindOrAddUserDefinedAttribute(const FName& Name, const FName& GroupName, const FName& GroupDependency)
 	{
 		using namespace UE::Chaos::ClothAsset::Private;
-		if (const TArray<FName>* const FixedAttributePrefixes = FixedAttributePrefixNamesMap.Find(GroupName))
+		if (const TArray<FName>* const FixedAttributePrefixes = GetFixedAttributePrefixNamesMap().Find(GroupName))
 		{
 			// Don't allow adding a name reserved as a fixed attribute prefix.
 			if (FixedAttributePrefixes->ContainsByPredicate([&Name](const FName& Prefix)
@@ -1030,7 +1112,7 @@ namespace UE::Chaos::ClothAsset
 				return nullptr;
 			}
 		}
-		if (const TArray<FName>* const FixedAttributeNames = FixedAttributeNamesMap.Find(GroupName))
+		if (const TArray<FName>* const FixedAttributeNames = GetFixedAttributeNamesMap().Find(GroupName))
 		{
 			// Only allow adding to known groups. Do not allow adding a name reserved as a fixed attribute.
 			if (!FixedAttributeNames->Contains(Name))
@@ -1038,7 +1120,7 @@ namespace UE::Chaos::ClothAsset
 				return GetManagedArrayCollection()->FindOrAddAttributeTyped<T>(Name, GroupName, FManagedArrayCollection::FConstructionParameters(GroupDependency));
 			}
 		}
-		if (const TArray<FName>* const UserAccessibleAttributeNames = UserAccessibleAttributeNamesMap.Find(GroupName))
+		if (const TArray<FName>* const UserAccessibleAttributeNames = GetUserAccessibleAttributeNamesMap().Find(GroupName))
 		{
 			// Paintable attribute can only be found, not created (they are part of the schema)
 			if (UserAccessibleAttributeNames->Contains(Name))
@@ -1057,9 +1139,9 @@ namespace UE::Chaos::ClothAsset
 	void FClothCollection::RemoveUserDefinedAttribute(const FName& Name, const FName& GroupName)
 	{
 		using namespace UE::Chaos::ClothAsset::Private;
-		if (const TArray<FName>* const FixedAttributeNames = FixedAttributeNamesMap.Find(GroupName))
+		if (const TArray<FName>* const FixedAttributeNames = GetFixedAttributeNamesMap().Find(GroupName))
 		{
-			if (const TArray<FName>* const FixedAttributePrefixes = FixedAttributePrefixNamesMap.Find(GroupName))
+			if (const TArray<FName>* const FixedAttributePrefixes = GetFixedAttributePrefixNamesMap().Find(GroupName))
 			{
 				// Don't allow removing a reserved attribute
 				if (FixedAttributePrefixes->ContainsByPredicate([&Name](const FName& Prefix)
@@ -1082,7 +1164,7 @@ namespace UE::Chaos::ClothAsset
 	bool FConstClothCollection::HasUserDefinedAttribute(const FName& Name, const FName& GroupName) const
 	{
 		using namespace UE::Chaos::ClothAsset::Private;
-		if (const TArray<FName>* const FixedAttributePrefixes = FixedAttributePrefixNamesMap.Find(GroupName))
+		if (const TArray<FName>* const FixedAttributePrefixes = GetFixedAttributePrefixNamesMap().Find(GroupName))
 		{
 			// Don't consider reserved attribute prefixes as user defined.
 			if (FixedAttributePrefixes->ContainsByPredicate([&Name](const FName& Prefix)
@@ -1093,7 +1175,7 @@ namespace UE::Chaos::ClothAsset
 				return false;
 			}
 		}
-		if (const TArray<FName>* const FixedAttributeNames = FixedAttributeNamesMap.Find(GroupName))
+		if (const TArray<FName>* const FixedAttributeNames = GetFixedAttributeNamesMap().Find(GroupName))
 		{
 			// User defined attributes are only allowed in known group names.
 			if (!FixedAttributeNames->Contains(Name))
@@ -1101,7 +1183,7 @@ namespace UE::Chaos::ClothAsset
 				return ManagedArrayCollection->FindAttributeTyped<T>(Name, GroupName) != nullptr;
 			}
 		}
-		if (const TArray<FName>* const UserAccessibleAttributeNames = UserAccessibleAttributeNamesMap.Find(GroupName))
+		if (const TArray<FName>* const UserAccessibleAttributeNames = GetUserAccessibleAttributeNamesMap().Find(GroupName))
 		{
 			// Paintable attribute can only be found, not created (they are part of the schema)
 			if (UserAccessibleAttributeNames->Contains(Name))
@@ -1121,7 +1203,7 @@ namespace UE::Chaos::ClothAsset
 	const TManagedArray<T>* FConstClothCollection::GetUserDefinedAttribute(const FName& Name, const FName& GroupName) const
 	{
 		using namespace UE::Chaos::ClothAsset::Private;
-		if (const TArray<FName>* const FixedAttributePrefixes = FixedAttributePrefixNamesMap.Find(GroupName))
+		if (const TArray<FName>* const FixedAttributePrefixes = GetFixedAttributePrefixNamesMap().Find(GroupName))
 		{
 			// Don't consider reserved attribute prefixes as user defined.
 			if (FixedAttributePrefixes->ContainsByPredicate([&Name](const FName& Prefix)
@@ -1132,7 +1214,7 @@ namespace UE::Chaos::ClothAsset
 				return nullptr;
 			}
 		}
-		if (const TArray<FName>* const FixedAttributeNames = FixedAttributeNamesMap.Find(GroupName))
+		if (const TArray<FName>* const FixedAttributeNames = GetFixedAttributeNamesMap().Find(GroupName))
 		{
 			// User defined attributes are only allowed in known group names.
 			if (!FixedAttributeNames->Contains(Name))
@@ -1140,7 +1222,7 @@ namespace UE::Chaos::ClothAsset
 				return ManagedArrayCollection->FindAttributeTyped<T>(Name, GroupName);
 			}
 		}
-		if (const TArray<FName>* const UserAccessibleAttributeNames = UserAccessibleAttributeNamesMap.Find(GroupName))
+		if (const TArray<FName>* const UserAccessibleAttributeNames = GetUserAccessibleAttributeNamesMap().Find(GroupName))
 		{
 			// Paintable attribute can only be found, not created (they are part of the schema)
 			if (UserAccessibleAttributeNames->Contains(Name))
@@ -1160,7 +1242,7 @@ namespace UE::Chaos::ClothAsset
 	TManagedArray<T>* FClothCollection::GetUserDefinedAttribute(const FName& Name, const FName& GroupName)
 	{
 		using namespace UE::Chaos::ClothAsset::Private;
-		if (const TArray<FName>* const FixedAttributePrefixes = FixedAttributePrefixNamesMap.Find(GroupName))
+		if (const TArray<FName>* const FixedAttributePrefixes = GetFixedAttributePrefixNamesMap().Find(GroupName))
 		{
 			// Don't consider reserved attribute prefixes as user defined.
 			if (FixedAttributePrefixes->ContainsByPredicate([&Name](const FName& Prefix)
@@ -1171,7 +1253,7 @@ namespace UE::Chaos::ClothAsset
 				return nullptr;
 			}
 		}
-		if (const TArray<FName>* const FixedAttributeNames = FixedAttributeNamesMap.Find(GroupName))
+		if (const TArray<FName>* const FixedAttributeNames = GetFixedAttributeNamesMap().Find(GroupName))
 		{
 			// User defined attributes are only allowed in known group names.
 			if (!FixedAttributeNames->Contains(Name))
@@ -1179,7 +1261,7 @@ namespace UE::Chaos::ClothAsset
 				return GetManagedArrayCollection()->FindAttributeTyped<T>(Name, GroupName);
 			}
 		}
-		if (const TArray<FName>* const UserAccessibleAttributeNames = UserAccessibleAttributeNamesMap.Find(GroupName))
+		if (const TArray<FName>* const UserAccessibleAttributeNames = GetUserAccessibleAttributeNamesMap().Find(GroupName))
 		{
 			// Paintable attribute can only be found, not created (they are part of the schema)
 			if (UserAccessibleAttributeNames->Contains(Name))
@@ -1199,20 +1281,20 @@ namespace UE::Chaos::ClothAsset
 	{
 		using namespace UE::Chaos::ClothAsset::Private;
 		TArray<FName> ValidClothCollectionGroupName;
-		FixedAttributeNamesMap.GetKeys(ValidClothCollectionGroupName);
+		GetFixedAttributeNamesMap().GetKeys(ValidClothCollectionGroupName);
 		return ValidClothCollectionGroupName;
 	}
 
 	/*static*/ bool FConstClothCollection::IsValidClothCollectionGroupName(const FName& GroupName)
 	{
 		using namespace UE::Chaos::ClothAsset::Private;
-		return FixedAttributeNamesMap.Contains(GroupName);
+		return GetFixedAttributeNamesMap().Contains(GroupName);
 	}
 
 	bool FConstClothCollection::IsValidUserDefinedAttributeName(const FName& Name, const FName& GroupName) const
 	{
 		using namespace UE::Chaos::ClothAsset::Private;
-		if (const TArray<FName>* const FixedAttributePrefixes = FixedAttributePrefixNamesMap.Find(GroupName))
+		if (const TArray<FName>* const FixedAttributePrefixes = GetFixedAttributePrefixNamesMap().Find(GroupName))
 		{
 			// Don't consider reserved attribute prefixes as user defined.
 			if (FixedAttributePrefixes->ContainsByPredicate([&Name](const FName& Prefix)
@@ -1223,7 +1305,7 @@ namespace UE::Chaos::ClothAsset
 				return false;
 			}
 		}
-		if (const TArray<FName>* const FixedAttributeNames = FixedAttributeNamesMap.Find(GroupName))
+		if (const TArray<FName>* const FixedAttributeNames = GetFixedAttributeNamesMap().Find(GroupName))
 		{
 			// User defined attributes are only allowed in known group names and cannot be reserved by FixedAttributeNames.
 			if (!FixedAttributeNames->Contains(Name))
@@ -1231,7 +1313,7 @@ namespace UE::Chaos::ClothAsset
 				return true;
 			}
 		}
-		if (const TArray<FName>* const UserAccessibleAttributeNames = UserAccessibleAttributeNamesMap.Find(GroupName))
+		if (const TArray<FName>* const UserAccessibleAttributeNames = GetUserAccessibleAttributeNamesMap().Find(GroupName))
 		{
 			// Paintable attribute can only be found, not created (they are part of the schema)
 			if (UserAccessibleAttributeNames->Contains(Name) && ManagedArrayCollection->HasAttribute(Name, GroupName))
