Chaos Cloth Asset - Fixed missing SKM clothing rebinding when reloading a Cloth Asset.

#rnx
#Synced-CL 48035595

diff --git a/Engine/Plugins/ChaosClothAsset/Source/ChaosClothAssetEngine/Private/ChaosClothAsset/ClothAssetBase.cpp b/Engine/Plugins/ChaosClothAsset/Source/ChaosClothAssetEngine/Private/ChaosClothAsset/ClothAssetBase.cpp
--- a/Engine/Plugins/ChaosClothAsset/Source/ChaosClothAssetEngine/Private/ChaosClothAsset/ClothAssetBase.cpp
+++ b/Engine/Plugins/ChaosClothAsset/Source/ChaosClothAssetEngine/Private/ChaosClothAsset/ClothAssetBase.cpp
@@ -19,6 +19,9 @@
 #endif
 #include "Rendering/SkeletalMeshRenderData.h"
 #include "UObject/FortniteMainBranchObjectVersion.h"
+#if WITH_EDITOR
+#include "UObject/PackageReload.h"
+#endif
 #include "ClothingSimulationFactory.h"
 #include "ClothingSimulationInstance.h"
 #include "ClothingSimulationInteractor.h"
@@ -162,6 +165,31 @@ void UChaosClothAssetBase::UpdateSimulationActor(TObjectPtr<AActor>& SimulationA
 	}
 }
 
+#if WITH_EDITOR
+void UChaosClothAssetBase::RegisterOnPackageReloadedDelegate()
+{
+	UnregisterOnPackageReloadedDelegate();
+	FCoreUObjectDelegates::OnPackageReloaded.AddWeakLambda(this, [this](const EPackageReloadPhase PackageReloadPhase, FPackageReloadedEvent* PackageReloadedEvent)
+		{
+			UObject* ReloadedObject;
+			if (PackageReloadPhase == EPackageReloadPhase::PostPackageFixup &&
+				PackageReloadedEvent && PackageReloadedEvent->GetRepointedObject(this, ReloadedObject))
+			{
+				if (UChaosClothAssetBase* const ReloadedClothAsset = Cast<UChaosClothAssetBase>(ReloadedObject))
+				{
+					constexpr bool bReregisterComponents = true;
+					ReloadedClothAsset->OnAssetChanged(bReregisterComponents);
+				}
+			}
+		});
+}
+
+void UChaosClothAssetBase::UnregisterOnPackageReloadedDelegate()
+{
+	FCoreUObjectDelegates::OnPackageReloaded.RemoveAll(this);
+}
+#endif
+
 const FDataflowInstance& UChaosClothAssetBase::GetDataflowInstance() const
 {
 PRAGMA_DISABLE_DEPRECATION_WARNINGS
@@ -182,6 +210,10 @@ void UChaosClothAssetBase::BeginDestroy()
 
 	Super::BeginDestroy();
 
+#if WITH_EDITOR
+	UnregisterOnPackageReloadedDelegate();
+#endif
+
 	// Release the mesh's render resources now
 	ReleaseResources();
 }
@@ -221,6 +253,18 @@ void UChaosClothAssetBase::PostLoad()
 #endif  // WITH_EDITORONLY_DATA
 }
 
+void UChaosClothAssetBase::PostInitProperties()
+{
+	Super::PostInitProperties();
+
+#if WITH_EDITOR
+	if (!IsTemplate())
+	{
+		RegisterOnPackageReloadedDelegate();
+	}
+#endif
+}
+
 #if WITH_EDITOR
 void UChaosClothAssetBase::PostEditChangeProperty(FPropertyChangedEvent& PropertyChangedEvent)
 {
diff --git a/Engine/Plugins/ChaosClothAsset/Source/ChaosClothAssetEngine/Public/ChaosClothAsset/ClothAssetBase.h b/Engine/Plugins/ChaosClothAsset/Source/ChaosClothAssetEngine/Public/ChaosClothAsset/ClothAssetBase.h
--- a/Engine/Plugins/ChaosClothAsset/Source/ChaosClothAssetEngine/Public/ChaosClothAsset/ClothAssetBase.h
+++ b/Engine/Plugins/ChaosClothAsset/Source/ChaosClothAssetEngine/Public/ChaosClothAsset/ClothAssetBase.h
@@ -85,6 +85,7 @@ public:
 	UE_API virtual bool IsReadyForFinishDestroy() override;
 	UE_API virtual void Serialize(FArchive& Ar) override;
 	UE_API virtual void PostLoad() override;
+	UE_API virtual void PostInitProperties() override;
 #if WITH_EDITOR
 	UE_API virtual void PostEditChangeProperty(FPropertyChangedEvent& PropertyChangedEvent) override;
 #endif
@@ -428,6 +429,13 @@ private:
 	/** Update simulation actor/components from the asset data */
 	void UpdateSimulationActor(TObjectPtr<AActor>& SimulationActor) const;
 
+#if WITH_EDITOR
+	/** Register the OnPackageReloaded delegate to call OnAssetChanged() when reloading this asset. */
+	void RegisterOnPackageReloadedDelegate();
+	/** Unregister the OnPackageReloaded delegate. */
+	void UnregisterOnPackageReloadedDelegate();
+#endif
+
 #if WITH_EDITORONLY_DATA
 	// The following PreviewScene properties are modeled after PreviewSkeletalMesh in USkeleton
 	//	- they are inside WITH_EDITORONLY_DATA because they are not used at game runtime
