Dataflow : Enable dataflow thumbnail and menu options to Cloth, Outfit and groom Asset types

#Synced-CL 48299723

diff --git a/Engine/Plugins/ChaosClothAsset/Source/ChaosClothAssetEngine/Private/ChaosClothAsset/ClothAssetBase.cpp b/Engine/Plugins/ChaosClothAsset/Source/ChaosClothAssetEngine/Private/ChaosClothAsset/ClothAssetBase.cpp
--- a/Engine/Plugins/ChaosClothAsset/Source/ChaosClothAssetEngine/Private/ChaosClothAsset/ClothAssetBase.cpp
+++ b/Engine/Plugins/ChaosClothAsset/Source/ChaosClothAssetEngine/Private/ChaosClothAsset/ClothAssetBase.cpp
@@ -18,6 +18,7 @@
 #include "Interfaces/ITargetPlatform.h"
 #endif
 #include "Rendering/SkeletalMeshRenderData.h"
+#include "UObject/AssetRegistryTagsContext.h"
 #include "UObject/FortniteMainBranchObjectVersion.h"
 #if WITH_EDITOR
 #include "UObject/PackageReload.h"
@@ -265,6 +266,14 @@ void UChaosClothAssetBase::PostInitProperties()
 #endif
 }
 
+void UChaosClothAssetBase::GetAssetRegistryTags(FAssetRegistryTagsContext Context) const
+{
+	// Dataflow specific tags
+	UE::Dataflow::InstanceUtils::GetAssetRegistryTags(this, Context);
+
+	Super::GetAssetRegistryTags(Context);
+}
+
 #if WITH_EDITOR
 void UChaosClothAssetBase::PostEditChangeProperty(FPropertyChangedEvent& PropertyChangedEvent)
 {
diff --git a/Engine/Plugins/ChaosClothAsset/Source/ChaosClothAssetEngine/Public/ChaosClothAsset/ClothAssetBase.h b/Engine/Plugins/ChaosClothAsset/Source/ChaosClothAssetEngine/Public/ChaosClothAsset/ClothAssetBase.h
--- a/Engine/Plugins/ChaosClothAsset/Source/ChaosClothAssetEngine/Public/ChaosClothAsset/ClothAssetBase.h
+++ b/Engine/Plugins/ChaosClothAsset/Source/ChaosClothAssetEngine/Public/ChaosClothAsset/ClothAssetBase.h
@@ -86,6 +86,7 @@ public:
 	UE_API virtual void Serialize(FArchive& Ar) override;
 	UE_API virtual void PostLoad() override;
 	UE_API virtual void PostInitProperties() override;
+	UE_API virtual void GetAssetRegistryTags(FAssetRegistryTagsContext Context) const override;
 #if WITH_EDITOR
 	UE_API virtual void PostEditChangeProperty(FPropertyChangedEvent& PropertyChangedEvent) override;
 #endif
diff --git a/Engine/Plugins/ChaosClothAssetEditor/Source/ChaosClothAssetEditor/Private/ChaosClothAsset/AssetDefinition_ClothAsset.cpp b/Engine/Plugins/ChaosClothAssetEditor/Source/ChaosClothAssetEditor/Private/ChaosClothAsset/AssetDefinition_ClothAsset.cpp
--- a/Engine/Plugins/ChaosClothAssetEditor/Source/ChaosClothAssetEditor/Private/ChaosClothAsset/AssetDefinition_ClothAsset.cpp
+++ b/Engine/Plugins/ChaosClothAssetEditor/Source/ChaosClothAssetEditor/Private/ChaosClothAsset/AssetDefinition_ClothAsset.cpp
@@ -266,4 +266,7 @@ bool UAssetDefinition_ClothAsset::LaunchClothDataflowAssetEditor(UChaosClothAsse
 	return false;
 }
 
+// Register dataflow related right-click menu options
+DATAFLOW_REGISTER_ASSET_MENUS(UChaosClothAsset)
+
 #undef LOCTEXT_NAMESPACE
diff --git a/Engine/Plugins/Experimental/ChaosOutfitAsset/Source/ChaosOutfitAssetEditor/Private/ChaosOutfitAsset/AssetDefinition_OutfitAsset.cpp b/Engine/Plugins/Experimental/ChaosOutfitAsset/Source/ChaosOutfitAssetEditor/Private/ChaosOutfitAsset/AssetDefinition_OutfitAsset.cpp
--- a/Engine/Plugins/Experimental/ChaosOutfitAsset/Source/ChaosOutfitAssetEditor/Private/ChaosOutfitAsset/AssetDefinition_OutfitAsset.cpp
+++ b/Engine/Plugins/Experimental/ChaosOutfitAsset/Source/ChaosOutfitAssetEditor/Private/ChaosOutfitAsset/AssetDefinition_OutfitAsset.cpp
@@ -84,4 +84,7 @@ EAssetCommandResult UAssetDefinition_OutfitAsset::OpenAssets(const FAssetOpenArg
 	return EAssetCommandResult::Unhandled;
 }
 
+// Register dataflow related right-click menu options
+DATAFLOW_REGISTER_ASSET_MENUS(UChaosOutfitAsset)
+
 #undef LOCTEXT_NAMESPACE
diff --git a/Engine/Plugins/Experimental/Dataflow/Source/DataflowEditor/Public/Dataflow/AssetDefinition_DataflowAsset.h b/Engine/Plugins/Experimental/Dataflow/Source/DataflowEditor/Public/Dataflow/AssetDefinition_DataflowAsset.h
--- a/Engine/Plugins/Experimental/Dataflow/Source/DataflowEditor/Public/Dataflow/AssetDefinition_DataflowAsset.h
+++ b/Engine/Plugins/Experimental/Dataflow/Source/DataflowEditor/Public/Dataflow/AssetDefinition_DataflowAsset.h
@@ -45,6 +45,16 @@ namespace UE::DataflowAssetDefinitionHelpers
 	DATAFLOWEDITOR_API void RegisterReEvaluateDataflowAssetMenu(UClass* AssetClass);
 }
 
+#define DATAFLOW_REGISTER_ASSET_MENUS(AssetClassType) \
+	namespace MenuExtension_##AssetClassType \
+	{ \
+		static FDelayedAutoRegisterHelper DelayedAutoRegister(EDelayedRegisterRunPhase::EndOfEngineInit, []() \
+		{ \
+			UE::DataflowAssetDefinitionHelpers::RegisterEmbedDataflowAssetMenu(AssetClassType::StaticClass()); \
+			UE::DataflowAssetDefinitionHelpers::RegisterReEvaluateDataflowAssetMenu(AssetClassType::StaticClass()); \
+		}); \
+	} 
+
 struct FDataflowInput;
 struct FDataflowOutput;
 
diff --git a/Engine/Plugins/Runtime/HairStrands/Source/HairStrandsCore/Private/GroomAsset.cpp b/Engine/Plugins/Runtime/HairStrands/Source/HairStrandsCore/Private/GroomAsset.cpp
--- a/Engine/Plugins/Runtime/HairStrands/Source/HairStrandsCore/Private/GroomAsset.cpp
+++ b/Engine/Plugins/Runtime/HairStrands/Source/HairStrandsCore/Private/GroomAsset.cpp
@@ -1815,6 +1815,9 @@ void UGroomAsset::GetAssetRegistryTags(FAssetRegistryTagsContext Context) const
 		Context.AddTag(FAssetRegistryTag(SourceFileTagName(), AssetImportData->GetSourceData().ToJson(), FAssetRegistryTag::TT_Hidden));
 	}
 
+	// Dataflow specific tags
+	UE::Dataflow::InstanceUtils::GetAssetRegistryTags(this, Context);
+
 	Super::GetAssetRegistryTags(Context);
 }
 
diff --git a/Engine/Plugins/Runtime/HairStrands/Source/HairStrandsEditor/Private/GroomActions.cpp b/Engine/Plugins/Runtime/HairStrands/Source/HairStrandsEditor/Private/GroomActions.cpp
--- a/Engine/Plugins/Runtime/HairStrands/Source/HairStrandsEditor/Private/GroomActions.cpp
+++ b/Engine/Plugins/Runtime/HairStrands/Source/HairStrandsEditor/Private/GroomActions.cpp
@@ -727,5 +727,8 @@ static FDelayedAutoRegisterHelper DelayedAutoRegister(EDelayedRegisterRunPhase::
 	}));
 });
 
+// Register dataflow related right-click menu options
+DATAFLOW_REGISTER_ASSET_MENUS(UGroomAsset)
+
 } // namespace MenuExtension_GroomAsset
 #undef LOCTEXT_NAMESPACE
