Cloth Asset - Made the Cloth Asset preview assets to use the new data driven BP variable system, instead of relying on hard coded preview assets.

[RN] Cloth Asset - Made the Cloth Asset preview assets to use the new data driven BP variable system, instead of relying on hard coded preview assets.
#virtualized
#Synced-CL 48653274

diff --git a/Engine/Plugins/ChaosClothAsset/Source/ChaosClothAssetEngine/Private/ChaosClothAsset/ClothAssetBase.cpp b/Engine/Plugins/ChaosClothAsset/Source/ChaosClothAssetEngine/Private/ChaosClothAsset/ClothAssetBase.cpp
--- a/Engine/Plugins/ChaosClothAsset/Source/ChaosClothAssetEngine/Private/ChaosClothAsset/ClothAssetBase.cpp
+++ b/Engine/Plugins/ChaosClothAsset/Source/ChaosClothAssetEngine/Private/ChaosClothAsset/ClothAssetBase.cpp
@@ -10,6 +10,10 @@
 #include "ChaosClothAsset/ClothSimulationModel.h"
 #include "ChaosClothAsset/SkeletalMeshConverterClassProvider.h"
 #include "Components/SkeletalMeshComponent.h"
+#if WITH_EDITORONLY_DATA
+#include "Engine/Blueprint.h"
+#include "Engine/BlueprintGeneratedClass.h"
+#endif
 #include "Engine/Engine.h"
 #include "Engine/SkeletalMesh.h"
 #include "Engine/SkinnedAssetAsyncCompileUtils.h"
@@ -114,41 +118,29 @@ void UChaosClothAssetBase::SetReferenceSkeleton(const FReferenceSkeleton* Refere
 
 TObjectPtr<UDataflowBaseContent> UChaosClothAssetBase::CreateDataflowContent()
 {
-	TObjectPtr<UDataflowSkeletalContent> SkeletalContent = UE::DataflowContextHelpers::CreateNewDataflowContent<UDataflowSkeletalContent>(this);
+	TObjectPtr<UDataflowBaseContent> DataflowContent = UE::DataflowContextHelpers::CreateNewDataflowContent<UDataflowBaseContent>(this);
 
-	SkeletalContent->SetDataflowOwner(this);
-	SkeletalContent->SetTerminalAsset(this);
+	DataflowContent->SetDataflowOwner(this);
+	DataflowContent->SetTerminalAsset(this);
 
-	WriteDataflowContent(SkeletalContent);
+	WriteDataflowContent(DataflowContent);
 
-	SkeletalContent->OnContentDataChanged.AddUObject(this, &UChaosClothAssetBase::UpdateSimulationActor);
+	DataflowContent->OnContentDataChanged.AddUObject(this, &UChaosClothAssetBase::UpdateSimulationActor);
 
-	return SkeletalContent;
+	return DataflowContent;
 }
 
 void UChaosClothAssetBase::WriteDataflowContent(const TObjectPtr<UDataflowBaseContent>& DataflowContent) const
 {
-	if (const TObjectPtr<UDataflowSkeletalContent> SkeletalContent = Cast<UDataflowSkeletalContent>(DataflowContent))
+	if (DataflowContent)
 	{
-		SkeletalContent->SetDataflowAsset(GetDataflowInstance().GetDataflowAsset());
-		SkeletalContent->SetDataflowTerminal(GetDataflowInstance().GetDataflowTerminal().ToString());
-
-#if WITH_EDITORONLY_DATA
-		SkeletalContent->SetAnimationAsset(GetPreviewSceneAnimation());
-		SkeletalContent->SetSkeletalMesh(GetPreviewSceneSkeletalMesh());
-#endif
+		DataflowContent->SetDataflowAsset(GetDataflowInstance().GetDataflowAsset());
+		DataflowContent->SetDataflowTerminal(GetDataflowInstance().GetDataflowTerminal().ToString());
 	}
 }
 
 void UChaosClothAssetBase::ReadDataflowContent(const TObjectPtr<UDataflowBaseContent>& DataflowContent)
 {
-	if (const TObjectPtr<UDataflowSkeletalContent> SkeletalContent = Cast<UDataflowSkeletalContent>(DataflowContent))
-	{
-#if WITH_EDITORONLY_DATA
-		PreviewSceneAnimation = SkeletalContent->GetAnimationAsset();
-		PreviewSceneSkeletalMesh = SkeletalContent->GetSkeletalMesh();
-#endif
-	}
 }
 
 void UChaosClothAssetBase::UpdateSimulationActor(TObjectPtr<AActor>& SimulationActor) const 
@@ -245,12 +237,62 @@ void UChaosClothAssetBase::Serialize(FArchive& Ar)
 	Ar << GetRefSkeleton();
 }
 
+#if WITH_EDITORONLY_DATA
+template<typename T>
+void UChaosClothAssetBase::SyncBPPreviewAsset(TSoftObjectPtr<T>& PreviewAsset, const TCHAR* const VarName)
+{
+	// Note:
+	// This function allows the Cloth Editor preview assets to keep in sync with the Dataflow Editor preview assets set by the BP until the migration to the new editor is complete.
+	// Note this is technically not a complete sync, the values are only replaced when they are null as to not change current user preferences in the two editors.
+	// Also missing BP variables aren't added, this isn't the purpose of this code.
+	if (const UDataflow* const DataflowAsset = GetDataflowInstance().GetDataflowAsset())
+	{
+		if (const UBlueprintGeneratedClass* const BlueprintGeneratedClass = Cast<const UBlueprintGeneratedClass>(DataflowAsset->PreviewBlueprintClass))
+		{
+			if (const UBlueprint* const Blueprint = Cast<const UBlueprint>(BlueprintGeneratedClass->ClassGeneratedBy))
+			{
+				FInstancedPropertyBag& BlueprintVariables = GetDataflowInstance().GetPreviewBlueprintVariableOverrides();
+
+				if (const FBPVariableDescription* const BPVariableDescription =
+					Blueprint->NewVariables.FindByPredicate([VarName](const FBPVariableDescription& Element)
+						{
+							return Element.VarName == VarName;
+						}))
+				{
+					const FName VarGuidAsName(BPVariableDescription->VarGuid.ToString());  // VarGuid instead of the VarName, because the variable name can change and FInstancedPropertyBag::SanitizePropertyName
+					if (const TValueOrError<UObject*, EPropertyBagResult> ValueOrError =
+						BlueprintVariables.GetValueObject(VarGuidAsName, T::StaticClass());
+						ValueOrError.HasValue())
+					{
+						if (T* const Value = Cast<T>(ValueOrError.GetValue()))
+						{
+							if (PreviewAsset.IsNull())
+							{
+								PreviewAsset = Value;  // Keep the the old preview asset up to date, in case it is reopened from inside the Cloth Editor
+							}
+						}
+						else
+						{
+							BlueprintVariables.SetValueObject(VarGuidAsName, PreviewAsset.LoadSynchronous());  // The variable isn't set, update it with whatever asset we find in the old preview asset
+						}
+					}
+				}
+			}
+		}
+	}
+}
+#endif
+
 void UChaosClothAssetBase::PostLoad()
 {
 	Super::PostLoad();
 
 #if WITH_EDITORONLY_DATA
 	bHasDataflowAsset = (GetDataflowInstance().GetDataflowAsset() != nullptr);
+
+	// Sync preview assets after moving from the old preview assets to the new BP variable system
+	SyncBPPreviewAsset(PreviewSceneSkeletalMesh, TEXT("SkeletalMesh"));  // TODO: Deprecate PreviewSceneSkeletalMesh and PreviewSceneAnimation
+	SyncBPPreviewAsset(PreviewSceneAnimation, TEXT("AnimationAsset"));   //       and remove these calls once the old Cloth Editor is retired
 #endif  // WITH_EDITORONLY_DATA
 }
 
@@ -288,7 +330,6 @@ PRAGMA_DISABLE_DEPRECATION_WARNINGS
 		bHasDataflowAsset = (GetDataflowInstance().GetDataflowAsset() != nullptr);
 	}
 PRAGMA_ENABLE_DEPRECATION_WARNINGS
-
 	// This invalidates the IDataflowContentOwner, not the Dataflow itself
 	InvalidateDataflowContents();
 
diff --git a/Engine/Plugins/ChaosClothAsset/Source/ChaosClothAssetEngine/Public/ChaosClothAsset/ClothAssetBase.h b/Engine/Plugins/ChaosClothAsset/Source/ChaosClothAssetEngine/Public/ChaosClothAsset/ClothAssetBase.h
--- a/Engine/Plugins/ChaosClothAsset/Source/ChaosClothAssetEngine/Public/ChaosClothAsset/ClothAssetBase.h
+++ b/Engine/Plugins/ChaosClothAsset/Source/ChaosClothAssetEngine/Public/ChaosClothAsset/ClothAssetBase.h
@@ -330,7 +330,7 @@ protected:
 	/** Pre-calculate refpose-to-local transforms. */
 	UE_API void CalculateInvRefMatrices();
 
-	UE_DEPRECATED(5.6, "This property isn't deprecated, but the proper getter and setter should be used instead in oreder to preserve correct behavior.")
+	UE_DEPRECATED(5.6, "This property isn't deprecated, but the proper getter and setter should be used instead in order to preserve correct behavior.")
 	UPROPERTY(EditAnywhere, Category = "Dataflow")
 	FDataflowInstance DataflowInstance;
 
@@ -438,6 +438,10 @@ private:
 #endif
 
 #if WITH_EDITORONLY_DATA
+	/** Helper function to move deprecated Cloth Editor preview assets to the new BP variable system. */
+	template<typename T>
+	void SyncBPPreviewAsset(TSoftObjectPtr<T>& PreviewAsset, const TCHAR* const VarName);
+
 	// The following PreviewScene properties are modeled after PreviewSkeletalMesh in USkeleton
 	//	- they are inside WITH_EDITORONLY_DATA because they are not used at game runtime
 	//	- TSoftObjectPtrs since that will make it possible to avoid loading these assets until the PreviewScene asks for them
