Chaos SKM Cloth Asset:
- Reworked WriteSimulationData to deal with multiple clothing data and cloth assets.
- Updated simulation instance head-up display.

#rnx
#Synced-CL 48653274

diff --git a/Engine/Plugins/ChaosClothAsset/Source/ChaosClothAssetEngine/ChaosClothAssetEngine.Build.cs b/Engine/Plugins/ChaosClothAsset/Source/ChaosClothAssetEngine/ChaosClothAssetEngine.Build.cs
--- a/Engine/Plugins/ChaosClothAsset/Source/ChaosClothAssetEngine/ChaosClothAssetEngine.Build.cs
+++ b/Engine/Plugins/ChaosClothAsset/Source/ChaosClothAssetEngine/ChaosClothAssetEngine.Build.cs
@@ -40,6 +40,7 @@ public class ChaosClothAssetEngine : ModuleRules
 				new string[]
 				{
 					"PropertyEditor",   // For adding the Cloth Component "Cloth Sim" section to the Details panel UI
+					"Slate",  // For Slate notifications
 					"StructUtilsEditor"
 				}
 			);
diff --git a/Engine/Plugins/ChaosClothAsset/Source/ChaosClothAssetEngine/Private/ChaosClothAsset/ClothAssetSKMClothingAsset.cpp b/Engine/Plugins/ChaosClothAsset/Source/ChaosClothAssetEngine/Private/ChaosClothAsset/ClothAssetSKMClothingAsset.cpp
--- a/Engine/Plugins/ChaosClothAsset/Source/ChaosClothAssetEngine/Private/ChaosClothAsset/ClothAssetSKMClothingAsset.cpp
+++ b/Engine/Plugins/ChaosClothAsset/Source/ChaosClothAssetEngine/Private/ChaosClothAsset/ClothAssetSKMClothingAsset.cpp
@@ -11,9 +11,11 @@
 #include "ChaosClothAsset/CollectionClothFacade.h"
 #include "Components/SkeletalMeshComponent.h"
 #include "Engine/SkeletalMesh.h"
+#include "Framework/Notifications/NotificationManager.h"
 #include "Rendering/SkeletalMeshModel.h"
 #include "Rendering/SkeletalMeshRenderData.h"
 #include "Utils/ClothingMeshUtils.h"
+#include "Widgets/Notifications/SNotificationList.h"
 #include "ClothingAsset.h"
 #endif  // #if WITH_EDITOR
 
@@ -203,9 +205,9 @@ void UChaosClothAssetSKMClothingAsset::OnAssetChanged(const bool bReregisterComp
 		CalculateLODHasAnySimulationMeshData();
 	}
 
-	if (bReregisterComponents)
+	if (const USkeletalMesh* const OwnerMesh = Cast<USkeletalMesh>(GetOuter()))
 	{
-		if (const USkeletalMesh* const OwnerMesh = Cast<USkeletalMesh>(GetOuter()))
+		if (bReregisterComponents)
 		{
 			for (TObjectIterator<USkeletalMeshComponent> It; It; ++It)
 			{
@@ -219,6 +221,26 @@ void UChaosClothAssetSKMClothingAsset::OnAssetChanged(const bool bReregisterComp
 				}
 			}
 		}  // TODO: Move UClothingAssetCommon::ReregisterComponentsUsingClothing() to a common place (in Utils/ClothingMeshUtils?)
+
+		// Warn asset model already in use
+		for (const UClothingAssetBase* const& MeshClothingAsset : OwnerMesh->GetMeshClothingAssets())
+		{
+			if (const UChaosClothAssetSKMClothingAsset* const Other = Cast<const UChaosClothAssetSKMClothingAsset>(MeshClothingAsset))
+			{
+				if (Other != this &&
+					Other->Asset == Asset &&
+					Other->ClothSimulationModelIndex == ClothSimulationModelIndex)
+				{
+					FNotificationInfo NotificationInfo(LOCTEXT("WarningUsingDuplicateAssetModel",
+						"This asset and simulation model is already set to another Clothing Data.\n"
+						"For performance reason, prefer binding a single Clothing Data to multiple LOD sections instead."));
+					NotificationInfo.ExpireDuration = 10.f;
+					FSlateNotificationManager::Get().AddNotification(NotificationInfo);
+					UE_LOG(LogChaosClothAsset, Warning, TEXT("%s"), *NotificationInfo.Text.Get().ToString());
+					break;
+				}
+			}
+		}
 	}
 }
 
diff --git a/Engine/Plugins/ChaosClothAsset/Source/ChaosClothAssetEngine/Private/ChaosClothAsset/ClothAssetSKMClothingSimulation.cpp b/Engine/Plugins/ChaosClothAsset/Source/ChaosClothAssetEngine/Private/ChaosClothAsset/ClothAssetSKMClothingSimulation.cpp
--- a/Engine/Plugins/ChaosClothAsset/Source/ChaosClothAssetEngine/Private/ChaosClothAsset/ClothAssetSKMClothingSimulation.cpp
+++ b/Engine/Plugins/ChaosClothAsset/Source/ChaosClothAssetEngine/Private/ChaosClothAsset/ClothAssetSKMClothingSimulation.cpp
@@ -215,14 +215,8 @@ namespace UE::Chaos::ClothAsset
 	{
 		if (ClothSimulationProxy.IsValid())
 		{
-			ClothSimulationProxy->PostSimulate_GameThread();
-			ClothSimulationProxy->PostProcess_GameThread();
-
-			// TODO: Add a way to get the simulationdata without the unnecessary copy: GetCurrentSimulationData_GameThread(TMap<int32, FClothSimulData>& OutData)?
-			for (TMap<int32, FClothSimulData>::TConstIterator CurrentDataIter = ClothSimulationProxy->GetCurrentSimulationData_AnyThread().CreateConstIterator(); CurrentDataIter; ++CurrentDataIter)
-			{
-				InOutData.FindOrAdd(CurrentDataIter.Key()) = CurrentDataIter.Value();
-			}
+			ClothSimulationProxy->PostSimulate_GameThread(&InOutData);
+			ClothSimulationProxy->PostProcess_GameThread(&InOutData);
 		}
 	}
 
@@ -329,6 +323,32 @@ namespace UE::Chaos::ClothAsset
 		return INDEX_NONE;
 	}
 
+	int32 FSKMClothingSimulation::GetNextLinkedSimulationGroupId(int32 PrevSimulationGroupId) const
+	{
+		if (AssetClothSimulationProperties.IsValidIndex(PrevSimulationGroupId))
+		{
+			if (const UChaosClothAssetSKMClothingAsset* const ClothAssetSKMClothingAsset = AssetClothSimulationProperties[PrevSimulationGroupId].Key)
+			{
+				if (const UChaosClothAssetBase* const Asset = ClothAssetSKMClothingAsset->GetAsset())
+				{
+					const int32 ModelIndex = ClothAssetSKMClothingAsset->GetClothSimulationModelIndex();
+					for (int32 Index = PrevSimulationGroupId + 1; Index < AssetClothSimulationProperties.Num(); ++Index)
+					{
+						if (const UChaosClothAssetSKMClothingAsset* const NextClothAssetSKMClothingAsset = AssetClothSimulationProperties[Index].Key)
+						{
+							if (NextClothAssetSKMClothingAsset->GetAsset() == Asset &&
+								NextClothAssetSKMClothingAsset->GetClothSimulationModelIndex() == ModelIndex)
+							{
+								return Index;
+							}
+						}
+					}
+				}
+			}
+		}
+		return INDEX_NONE;
+	}
+
 	const FReferenceSkeleton* FSKMClothingSimulation::GetReferenceSkeleton() const
 	{
 		checkf(OwnerComponent, TEXT("No actors, owner component not set. Call CreateActor first."));
@@ -343,16 +363,19 @@ namespace UE::Chaos::ClothAsset
 
 	bool FSKMClothingSimulation::IsSimulationEnabled() const
 	{
+		checkf(OwnerComponent, TEXT("No actors, owner component not set. Call CreateActor first."));
 		return OwnerComponent->CanSimulateClothing();
 	}
 
 	bool FSKMClothingSimulation::IsSimulationSuspended() const
 	{
+		checkf(OwnerComponent, TEXT("No actors, owner component not set. Call CreateActor first."));
 		return OwnerComponent->IsClothingSimulationSuspended();
 	}
 
 	float FSKMClothingSimulation::GetClothGeometryScale() const
 	{
+		checkf(OwnerComponent, TEXT("No actors, owner component not set. Call CreateActor first."));
 		return OwnerComponent->ClothGeometryScale;
 	}
 
diff --git a/Engine/Plugins/ChaosClothAsset/Source/ChaosClothAssetEngine/Private/ChaosClothAsset/ClothSimulationProxy.cpp b/Engine/Plugins/ChaosClothAsset/Source/ChaosClothAssetEngine/Private/ChaosClothAsset/ClothSimulationProxy.cpp
--- a/Engine/Plugins/ChaosClothAsset/Source/ChaosClothAssetEngine/Private/ChaosClothAsset/ClothSimulationProxy.cpp
+++ b/Engine/Plugins/ChaosClothAsset/Source/ChaosClothAssetEngine/Private/ChaosClothAsset/ClothSimulationProxy.cpp
@@ -1,6 +1,7 @@
 // Copyright Epic Games, Inc. All Rights Reserved.
 
 #include "ChaosClothAsset/ClothSimulationProxy.h"
+#include "Algo/NoneOf.h"
 #include "ChaosClothAsset/ClothAsset.h"
 #include "ChaosClothAsset/ClothAssetPrivate.h"
 #include "ChaosClothAsset/ClothComponent.h"
@@ -15,6 +16,7 @@
 #include "ChaosCloth/ChaosClothingSimulationSolver.h"
 #include "ChaosCloth/ChaosClothVisualization.h"
 #include "Engine/SkinnedAssetCommon.h"
+#include "Misc/ScopeExit.h"
 #include "PhysicsEngine/PhysicsSettings.h"
 #include "PhysicsField/PhysicsFieldComponent.h"
 #include "Rendering/SkeletalMeshRenderData.h"
@@ -200,6 +202,11 @@ namespace UE::Chaos::ClothAsset
 				const int32 SimulationGroupId = ClothComponentAdapter.GetSimulationGroupId(Asset, ModelIndex);
 				if (SimulationGroupId != INDEX_NONE)  // INDEX_NONE means this model shouldn't be added to the simulation
 				{
+					checkf(Algo::NoneOf(Cloths, [&SimulationGroupId](const TUniquePtr<::Chaos::FClothingSimulationCloth>& Cloth)
+						{
+							return Cloth && Cloth->GetGroupId() == SimulationGroupId;
+						}), TEXT("The ClothComponentAdapter cannot return the same SimulationGroupId for two different Asset/ModelIndex pairs (unless INDEX_NONE)."));
+
 					const TSharedPtr<const FChaosClothSimulationModel> ClothSimulationModel = Asset->GetClothSimulationModel(ModelIndex);
 
 					FString DebugName;
@@ -291,14 +298,6 @@ namespace UE::Chaos::ClothAsset
 		bIsInitialized = true;
 	}
 
-	bool FClothSimulationProxy::SetupSimulationData(float DeltaTime)
-	{
-		PreProcess_GameThread(DeltaTime);
-		PreSimulate_GameThread(DeltaTime);
-		PostProcess_GameThread();
-		return bIsSimulating;
-	}
-
 	void FClothSimulationProxy::HardResetSimulation_GameThread()
 	{
 		// Reset solver, cloths, and configs, but keep existing meshes
@@ -470,15 +469,15 @@ namespace UE::Chaos::ClothAsset
 		return bIsSimulating;
 	}
 
-	void FClothSimulationProxy::PostSimulate_GameThread()
+	void FClothSimulationProxy::PostSimulate_GameThread(TMap<int32, FClothSimulData>* ExternalClothSimulData)
 	{
 		if (bIsSimulating)
 		{
-			WriteSimulationData();
+			WriteSimulationData(ExternalClothSimulData);
 		}
 	}
 
-	void FClothSimulationProxy::PostProcess_GameThread()
+	void FClothSimulationProxy::PostProcess_GameThread(TMap<int32, FClothSimulData>* ExternalClothSimulData)
 	{
 		if (!bIsPreProcessed)
 		{
@@ -493,7 +492,7 @@ namespace UE::Chaos::ClothAsset
 			{
 				// Take care of the LOD switching, as the simulation won't do it
 				UpdateClothLODs();
-				WriteSimulationData();
+				WriteSimulationData(ExternalClothSimulData);
 			}
 			else
 			{
@@ -778,7 +777,7 @@ namespace UE::Chaos::ClothAsset
 		}
 	}
 
-	void FClothSimulationProxy::WriteSimulationData()
+	void FClothSimulationProxy::WriteSimulationData(TMap<int32, FClothSimulData>* OutClothSimulDataOverride)
 	{
 		using namespace ::Chaos;
 		check(!IsValidRef(ParallelInitializationTask));
@@ -789,45 +788,51 @@ namespace UE::Chaos::ClothAsset
 		SCOPE_CYCLE_COUNTER(STAT_ClothSimulationProxy_WriteSimulationData);
 		SCOPE_CYCLE_COUNTER(STAT_ClothWriteback);
 
-		// Check if our bone map is actually valid, if not there is no clothing data to build
-		const USkinnedMeshComponent& OwnerComponent = ClothComponentAdapter.GetOwnerComponent();
-		const USkinnedMeshComponent* const LeaderPoseComponent = OwnerComponent.LeaderPoseComponent.Get();
-
-		if (LeaderPoseComponent && !OwnerComponent.GetLeaderBoneMap().Num())
+		// Use either the proxy owned or overridden data
+		if (!OutClothSimulDataOverride)
 		{
-			CurrentSimulationData.Reset();
-			return;
+			OutClothSimulDataOverride = &CurrentSimulationData;
 		}
 
-		// Get the list of used indices and cache the result of the ClothSimulData map find
-		TArray<int32, TInlineAllocator<16>> AssetIndices;
-		AssetIndices.Reserve(Cloths.Num());
+		// Cleanup the unused simulation data map on exit so that the render vertex factory won't run unnecessarily
+		TBitArray<TInlineAllocator<16>> SimulationGroupIdsToRemove;
+		auto FlagForRemovalOnExit = [&SimulationGroupIdsToRemove](const int32 SimulationGroupId, bool bRemove)
+			{
+				constexpr bool bRemoveDefault = false;  // Only remove owned data
+				SimulationGroupIdsToRemove.SetNum(FMath::Max(SimulationGroupIdsToRemove.Num(), SimulationGroupId + 1), bRemoveDefault);
+				SimulationGroupIdsToRemove[SimulationGroupId] = bRemove;
+			};
 
 		for (const TUniquePtr<FClothingSimulationCloth>& Cloth : Cloths)
 		{
-			if (Cloth &&
-				Cloth->GetMesh() &&
-				Cloth->GetParticleRangeId(Solver.Get()) != INDEX_NONE &&
-				Cloth->GetLODIndex(Solver.Get()) != INDEX_NONE)  // Needs a valid LOD
+			for (int32 SimulationGroupId = Cloth ? Cloth->GetGroupId() : INDEX_NONE;
+				SimulationGroupId != INDEX_NONE;
+				SimulationGroupId = ClothComponentAdapter.GetNextLinkedSimulationGroupId(SimulationGroupId))
 			{
-				AssetIndices.AddUnique((int32)Cloth->GetGroupId());
-				// TODO: Asset indices should be taken from the ClothComponentAdapter so multiple ids can be recovered from this group
+				FlagForRemovalOnExit(SimulationGroupId, true);  // Flag for removal all owned data until it has been successfully copied
 			}
 		}
 
-		// Cleanup all past keys (ensures that the cloth vertex factory won't run unnecessarily)
-		if (!AssetIndices.Num())
+		ON_SCOPE_EXIT
+			{
+				for (TMap<int32, FClothSimulData>::TIterator It = OutClothSimulDataOverride->CreateIterator(); It; ++It)
+				{
+					const int32 SimulationGroupId = It.Key();
+					if (SimulationGroupIdsToRemove.IsValidIndex(SimulationGroupId) && SimulationGroupIdsToRemove[SimulationGroupId])
+					{
+						It.RemoveCurrent();
+					}
+				}
+			};
+
+		// Check if our bone map is actually valid, if not there is no clothing data to build
+		const USkinnedMeshComponent& OwnerComponent = ClothComponentAdapter.GetOwnerComponent();
+		const USkinnedMeshComponent* const LeaderPoseComponent = OwnerComponent.LeaderPoseComponent.Get();
+
+		if (LeaderPoseComponent && !OwnerComponent.GetLeaderBoneMap().Num())
 		{
-			CurrentSimulationData.Reset();
 			return;
 		}
-		for (TMap<int32, FClothSimulData>::TIterator It = CurrentSimulationData.CreateIterator(); It; ++It)
-		{
-			if (!AssetIndices.Contains(It->Key))
-			{
-				It.RemoveCurrent();
-			}
-		}
 
 		// Get the solver's local space
 		const FVec3& LocalSpaceLocation = Solver->GetLocalSpaceLocation(); // Note: Since the ReferenceSpaceTransform can be suspended with the simulation, it is important that the suspended local space location is used too in order to get the simulation data back into reference space
@@ -837,142 +842,109 @@ namespace UE::Chaos::ClothAsset
 		const TArray<FTransform>& ComponentSpaceTransforms = LeaderPoseComponent ? LeaderPoseComponent->GetComponentSpaceTransforms() : OwnerComponent.GetComponentSpaceTransforms();
 
 		// Set the simulation data for each of the cloths
-		for (const int32 AssetIndex : AssetIndices)
+		for (const TUniquePtr<FClothingSimulationCloth>& Cloth : Cloths)
 		{
-			bool bHasCloth = false;
-			int32 ReferenceBoneIndex = INDEX_NONE;
-			FRigidTransform3 ReferenceSpaceTransform;
-
-			FClothSimulData* Data = CurrentSimulationData.Find(AssetIndex);
-			if (Data)
+			if (!Cloth ||
+				!Cloth->GetMesh() ||
+				Cloth->GetParticleRangeId(Solver.Get()) == INDEX_NONE ||
+				Cloth->GetLODIndex(Solver.Get()) == INDEX_NONE)  // Needs a valid LOD
 			{
-				Data->Reset();  // Clear the data to prepare them for the append operations to come
+				continue;
 			}
 
-			for (const TUniquePtr<FClothingSimulationCloth>& Cloth : Cloths)
-			{
-				if (Cloth->GetGroupId() != AssetIndex)
-				{
-					continue;
-				}
-
-				// Get the reference bone index for this cloth
-				const int32 ClothReferenceBoneIndex = LeaderPoseComponent ? OwnerComponent.GetLeaderBoneMap()[Cloth->GetReferenceBoneIndex()] : Cloth->GetReferenceBoneIndex();
-				if (bHasCloth && ClothReferenceBoneIndex != ReferenceBoneIndex)
-				{
-					UE_CLOG(bHasNonMatchingReferenceBones, LogChaosClothAsset, Warning, TEXT("Multiple cloth simulation models are used with different reference bone indices within the same simulation group"),
-						OwnerComponent.GetSkinnedAsset() ? *OwnerComponent.GetSkinnedAsset()->GetName() : *OwnerComponent.GetName());
-					bHasNonMatchingReferenceBones = true;
-					bHasCloth = false;
-					break;
-				}
+			// Get the reference bone index for this cloth
+			const int32 ReferenceBoneIndex = LeaderPoseComponent ? OwnerComponent.GetLeaderBoneMap()[Cloth->GetReferenceBoneIndex()] : Cloth->GetReferenceBoneIndex();
 
-				if (!bHasCloth)
-				{
-					// Set the ReferenceBoneIndex for this group
-					ReferenceBoneIndex = ClothReferenceBoneIndex;
+			// Get the reference bone transform for this cloth
+			if (!ComponentSpaceTransforms.IsValidIndex(ReferenceBoneIndex))
+			{
+				UE_CLOG(!bHasInvalidReferenceBoneTransforms, LogChaosClothAsset, Warning, TEXT("Failed to write back clothing simulation data for component %s as bone transforms are invalid."),
+					OwnerComponent.GetSkinnedAsset() ? *OwnerComponent.GetSkinnedAsset()->GetName() : *OwnerComponent.GetName());
+				bHasInvalidReferenceBoneTransforms = true;
+				continue;
+			}
 
-					// Get the reference bone transform for this cloth
-					if (!ComponentSpaceTransforms.IsValidIndex(ReferenceBoneIndex))
-					{
-						UE_CLOG(!bHasInvalidReferenceBoneTransforms, LogChaosClothAsset, Warning, TEXT("Failed to write back clothing simulation data for component %s as bone transforms are invalid."),
-							OwnerComponent.GetSkinnedAsset() ? *OwnerComponent.GetSkinnedAsset()->GetName() : *OwnerComponent.GetName());
-						bHasInvalidReferenceBoneTransforms = true;
-						CurrentSimulationData.Reset();
-						bHasCloth = false;
-						break;
-					}
+			// Set the FClothSimulData common properties for this group
+			const int32 SimulationGroupId = Cloth->GetGroupId();
+			FClothSimulData& Data = OutClothSimulDataOverride->FindOrAdd(SimulationGroupId);
 
-					// Get the reference transform used in the current animation pose
-					FTransform ReferenceBoneTransform = ComponentSpaceTransforms[ReferenceBoneIndex];
-					ReferenceBoneTransform *= ClothSimulationContext->ComponentTransform;
-					ReferenceBoneTransform.SetScale3D(FVector(1.0f));  // Scale is already baked in the cloth mesh
+			// Get the reference transform used in the current animation pose
+			FTransform ReferenceBoneTransform = ComponentSpaceTransforms[ReferenceBoneIndex];
+			ReferenceBoneTransform *= ClothSimulationContext->ComponentTransform;
+			ReferenceBoneTransform.SetScale3D(FVector(1.0f));  // Scale is already baked in the cloth mesh
 
-					// Set the FClothSimulData common properties for this group
-					if (!Data)
-					{
-						Data = &CurrentSimulationData.Add(AssetIndex);
-					}
+			// Set the world space transform to be this cloth's reference bone
+			Data.Transform = ReferenceBoneTransform;
+			Data.ComponentRelativeTransform = ReferenceBoneTransform.GetRelativeTransform(ClothSimulationContext->ComponentTransform);
 
-					// Set the world space transform to be this cloth's reference bone
-					Data->Transform = ReferenceBoneTransform;
-					Data->ComponentRelativeTransform = ReferenceBoneTransform.GetRelativeTransform(ClothSimulationContext->ComponentTransform);
-
-					// Set the current LOD these data apply to, so that the correct deformer mappings can be applied
-					// Note: If the LOD has changed while the simulation is suspended, the cloth still needs to be updated with the correct LOD data
-					//       This should be handled by calling UpdateClothLODs when not ticking/simulating.
-					const int32 LODIndex = Cloth->GetMesh()->GetLODIndex();
-					ensure(LODIndex == Cloth->GetLODIndex(Solver.Get()));
-					Data->LODIndex = Cloth->GetMesh()->GetOwnerLODIndex(LODIndex);  // The owner component LOD index can be different to the cloth mesh LOD index
-
-					// Retrieve the last reference space transform used for this cloth
-					// Note: This won't necessary match the current bone reference transform when the simulation is paused,
-					//       and still allows for the correct positioning of the sim data while the component is animated.
-					ReferenceSpaceTransform = Cloth->GetReferenceSpaceTransform();
-					ReferenceSpaceTransform.AddToTranslation(-LocalSpaceLocation);
-				}
+			// Set the current LOD these data apply to, so that the correct deformer mappings can be applied
+			// Note: If the LOD has changed while the simulation is suspended, the cloth still needs to be updated with the correct LOD data
+			//       This should be handled by calling UpdateClothLODs when not ticking/simulating.
+			const int32 LODIndex = Cloth->GetMesh()->GetLODIndex();
+			ensure(LODIndex == Cloth->GetLODIndex(Solver.Get()));
+			Data.LODIndex = Cloth->GetMesh()->GetOwnerLODIndex(LODIndex);  // The owner component LOD index can be different to the cloth mesh LOD index
 
-				// Copy positions and normals
-				Data->Positions.Append(Cloth->GetParticlePositions(Solver.Get()));
-				Data->Normals.Append(Cloth->GetParticleNormals(Solver.Get()));
+			// Retrieve the last reference space transform used for this cloth
+			// Note: This won't necessary match the current bone reference transform when the simulation is paused,
+			//       and still allows for the correct positioning of the sim data while the component is animated.
+			FRigidTransform3 ReferenceSpaceTransform = Cloth->GetReferenceSpaceTransform();
+			ReferenceSpaceTransform.AddToTranslation(-LocalSpaceLocation);
 
-				check(Data->Positions.Num() == Data->Normals.Num());
-				bHasCloth = true;
-			}
+			// Copy positions and normals
+			Data.Positions = Cloth->GetParticlePositions(Solver.Get());
+			Data.Normals = Cloth->GetParticleNormals(Solver.Get());
+			check(Data.Positions.Num() == Data.Normals.Num());
 
 			// Transform into the cloth reference simulation space used at the time of simulation
 			bool bSimDataIsFinite = true;
-			if (bHasCloth && Data)
+#if INTEL_ISPC
+			if (bTransformClothSimulData_ISPC_Enabled)
 			{
-	#if INTEL_ISPC
-				if (bTransformClothSimulData_ISPC_Enabled)
-				{
-					// ISPC is assuming float input here
-					check(sizeof(ispc::FVector3f) == Data->Positions.GetTypeSize());
-					check(sizeof(ispc::FVector3f) == Data->Normals.GetTypeSize());
+				// ISPC is assuming float input here
+				check(sizeof(ispc::FVector3f) == Data.Positions.GetTypeSize());
+				check(sizeof(ispc::FVector3f) == Data.Normals.GetTypeSize());
 
-					if (bSimulationDataCheckFinite)
-					{
-						bSimDataIsFinite = ispc::TransformClothSimulDataAndCheckFinite(
-							(ispc::FVector3f*)Data->Positions.GetData(),
-							(ispc::FVector3f*)Data->Normals.GetData(),
-							(ispc::FTransform&)ReferenceSpaceTransform,
-							LocalSpaceScale,
-							Data->Positions.Num());
-					}
-					else
-					{
-						ispc::TransformClothSimulData(
-							(ispc::FVector3f*)Data->Positions.GetData(),
-							(ispc::FVector3f*)Data->Normals.GetData(),
-							(ispc::FTransform&)ReferenceSpaceTransform,
-							LocalSpaceScale,
-							Data->Positions.Num());
-					}
+				if (bSimulationDataCheckFinite)
+				{
+					bSimDataIsFinite = ispc::TransformClothSimulDataAndCheckFinite(
+						(ispc::FVector3f*)Data.Positions.GetData(),
+						(ispc::FVector3f*)Data.Normals.GetData(),
+						(ispc::FTransform&)ReferenceSpaceTransform,
+						LocalSpaceScale,
+						Data.Positions.Num());
 				}
 				else
-	#endif
 				{
-					if (bSimulationDataCheckFinite)
-					{
-						auto VectorIsFinite = [](const FSolverVec3& Vector)
-							{
-								return FMath::IsFinite(Vector.X) && FMath::IsFinite(Vector.Y) && FMath::IsFinite(Vector.Z);
-							};
-						for (int32 Index = 0; Index < Data->Positions.Num(); ++Index)
+					ispc::TransformClothSimulData(
+						(ispc::FVector3f*)Data.Positions.GetData(),
+						(ispc::FVector3f*)Data.Normals.GetData(),
+						(ispc::FTransform&)ReferenceSpaceTransform,
+						LocalSpaceScale,
+						Data.Positions.Num());
+				}
+			}
+			else
+#endif
+			{
+				if (bSimulationDataCheckFinite)
+				{
+					auto VectorIsFinite = [](const FSolverVec3& Vector)
 						{
-							Data->Positions[Index] = FVec3f(ReferenceSpaceTransform.InverseTransformPosition(LocalSpaceScale * FVec3(Data->Positions[Index])));
-							Data->Normals[Index] = FVec3f(ReferenceSpaceTransform.InverseTransformVector(FVec3(-Data->Normals[Index])));
-							bSimDataIsFinite = bSimDataIsFinite && VectorIsFinite(Data->Positions[Index]) && VectorIsFinite(Data->Normals[Index]);
-						}
+							return FMath::IsFinite(Vector.X) && FMath::IsFinite(Vector.Y) && FMath::IsFinite(Vector.Z);
+						};
+					for (int32 Index = 0; Index < Data.Positions.Num(); ++Index)
+					{
+						Data.Positions[Index] = FVec3f(ReferenceSpaceTransform.InverseTransformPosition(LocalSpaceScale * FVec3(Data.Positions[Index])));
+						Data.Normals[Index] = FVec3f(ReferenceSpaceTransform.InverseTransformVector(FVec3(-Data.Normals[Index])));
+						bSimDataIsFinite = bSimDataIsFinite && VectorIsFinite(Data.Positions[Index]) && VectorIsFinite(Data.Normals[Index]);
 					}
-					else
+				}
+				else
+				{
+					for (int32 Index = 0; Index < Data.Positions.Num(); ++Index)
 					{
-						for (int32 Index = 0; Index < Data->Positions.Num(); ++Index)
-						{
-							Data->Positions[Index] = FVec3f(ReferenceSpaceTransform.InverseTransformPosition(LocalSpaceScale * FVec3(Data->Positions[Index])));
-							Data->Normals[Index] = FVec3f(ReferenceSpaceTransform.InverseTransformVector(FVec3(-Data->Normals[Index])));
-						}
+						Data.Positions[Index] = FVec3f(ReferenceSpaceTransform.InverseTransformPosition(LocalSpaceScale * FVec3(Data.Positions[Index])));
+						Data.Normals[Index] = FVec3f(ReferenceSpaceTransform.InverseTransformVector(FVec3(-Data.Normals[Index])));
 					}
 				}
 			}
@@ -985,11 +957,21 @@ namespace UE::Chaos::ClothAsset
 					bNeedResetClothFromNanCheck = true;
 				}
 			}
-
-			// Ensures that the cloth vertex factory won't run unnecessarily
-			if (Data && (!bHasCloth || !bSimDataIsFinite))
+			if (bSimDataIsFinite)
 			{
-				CurrentSimulationData.Remove(AssetIndex);
+				// Mark the sim data as valid
+				FlagForRemovalOnExit(SimulationGroupId, false);
+
+				// Copy linked data (for when a single cloth asset is used through several Clothing Data in multiple render sections)
+				for (int32 NextLinkedSimulationGroupId = ClothComponentAdapter.GetNextLinkedSimulationGroupId(SimulationGroupId);
+					NextLinkedSimulationGroupId != INDEX_NONE;
+					NextLinkedSimulationGroupId = ClothComponentAdapter.GetNextLinkedSimulationGroupId(NextLinkedSimulationGroupId))
+				{
+					FClothSimulData& NextData = OutClothSimulDataOverride->FindOrAdd(NextLinkedSimulationGroupId);
+					NextData = Data;  // TODO: Deprecate the current map type to use shared pointers to avoid Data duplication if ever possible (the editor warns for shared cloth asset setup to avoid this).
+
+					FlagForRemovalOnExit(NextLinkedSimulationGroupId, false);
+				}
 			}
 		}
 	}
diff --git a/Engine/Plugins/ChaosClothAsset/Source/ChaosClothAssetEngine/Public/ChaosClothAsset/ClothAssetSKMClothingSimulation.h b/Engine/Plugins/ChaosClothAsset/Source/ChaosClothAssetEngine/Public/ChaosClothAsset/ClothAssetSKMClothingSimulation.h
--- a/Engine/Plugins/ChaosClothAsset/Source/ChaosClothAssetEngine/Public/ChaosClothAsset/ClothAssetSKMClothingSimulation.h
+++ b/Engine/Plugins/ChaosClothAsset/Source/ChaosClothAssetEngine/Public/ChaosClothAsset/ClothAssetSKMClothingSimulation.h
@@ -118,7 +118,7 @@ namespace UE::Chaos::ClothAsset
 		}
 		virtual TArray<const UChaosClothAssetBase*> GetAssets() const override;
 		virtual int32 GetSimulationGroupId(const UChaosClothAssetBase* Asset, int32 ModelIndex) const override;
-
+		virtual int32 GetNextLinkedSimulationGroupId(int32 PrevSimulationGroupId) const override;
 		virtual const FReferenceSkeleton* GetReferenceSkeleton() const override;
 		virtual EClothingTeleportMode GetClothTeleportMode() const override;
 		virtual bool IsSimulationEnabled() const override;
diff --git a/Engine/Plugins/ChaosClothAsset/Source/ChaosClothAssetEngine/Public/ChaosClothAsset/ClothComponentAdapter.h b/Engine/Plugins/ChaosClothAsset/Source/ChaosClothAssetEngine/Public/ChaosClothAsset/ClothComponentAdapter.h
--- a/Engine/Plugins/ChaosClothAsset/Source/ChaosClothAssetEngine/Public/ChaosClothAsset/ClothComponentAdapter.h
+++ b/Engine/Plugins/ChaosClothAsset/Source/ChaosClothAssetEngine/Public/ChaosClothAsset/ClothComponentAdapter.h
@@ -33,6 +33,15 @@ namespace UE::Chaos::ClothAsset
 		/** Return the index of the model assigned to the render section, or INDEX_NONE when the model shouldn't be simulated. */
 		virtual int32 GetSimulationGroupId(const UChaosClothAssetBase* Asset, int32 ModelId) const = 0;
 
+		/**
+		 * Return the next SimulationGroupId when multiple group IDs are sharing the same simulated cloth,
+		 * or INDEX_NONE when SimulationGroupId is the last linked element.
+		 * Use GetSimulationGroupId() to get the initial ID, and call GetNextLinkedSimulationGroupId() with the previous ID
+		 * to retrieve the subsequent linked IDs.
+		 * This function allows for a single simulated cloth to update multiple simulation data slots.
+		 */
+		virtual int32 GetNextLinkedSimulationGroupId(int32 PrevSimulationGroupId) const { return INDEX_NONE; }
+
 		/** Return the reference skeleton as passed to the component renderer. */
 		virtual const FReferenceSkeleton* GetReferenceSkeleton() const = 0;
 
diff --git a/Engine/Plugins/ChaosClothAsset/Source/ChaosClothAssetEngine/Public/ChaosClothAsset/ClothSimulationProxy.h b/Engine/Plugins/ChaosClothAsset/Source/ChaosClothAssetEngine/Public/ChaosClothAsset/ClothSimulationProxy.h
--- a/Engine/Plugins/ChaosClothAsset/Source/ChaosClothAssetEngine/Public/ChaosClothAsset/ClothSimulationProxy.h
+++ b/Engine/Plugins/ChaosClothAsset/Source/ChaosClothAssetEngine/Public/ChaosClothAsset/ClothSimulationProxy.h
@@ -76,10 +76,18 @@ namespace UE::Chaos::ClothAsset
 		UE_API void PreProcess_GameThread(float DeltaTime, bool bForceWaitForInitialization = false);
 		/** Prepare the simulation data. For specialized usage when not calling Tick_GameThread. */
         UE_API bool PreSimulate_GameThread(float DeltaTime);
-		/** Write data back onto the game thread once the simulation is done. For specialized usage when not calling Tick_GameThread. */
-		UE_API void PostSimulate_GameThread();
-		/** Post setup, required for when the simulation didn't run. For specialized usage when not calling Tick_GameThread. */
-		UE_API void PostProcess_GameThread();
+		/**
+		 * Write data back onto the game thread once the simulation is done.
+		 * For specialized usage when not calling Tick_GameThread.
+		 * @param OutClothSimulDataOverride An optional override to store the resulting simulation data, or nullptr to use the internally owned map insted.
+		 */
+		UE_API void PostSimulate_GameThread(TMap<int32, FClothSimulData>* OutClothSimulDataOverride = nullptr);
+		/**
+		 * Post setup, required for when the simulation didn't run.
+		 * For specialized usage when not calling Tick_GameThread.
+		 * @param OutClothSimulDataOverride An optional override to store the resulting simulation data, or nullptr to use the internally owned map insted.
+		 */
+		UE_API void PostProcess_GameThread(TMap<int32, FClothSimulData>* OutClothSimulDataOverride = nullptr);
 
 		/* Force a pending cloth reset. For specialized usage when not calling Tick_GameThread. */
 		UE_API void ForcePendingReset_GameThread();
@@ -113,17 +121,13 @@ namespace UE::Chaos::ClothAsset
 	protected:
 		UE_API void PostConstructorInternal(bool bAsyncInitialization);
 		UE_API void Tick();
-		UE_API void WriteSimulationData();
-		UE_DEPRECATED(5.6, "Use PreProcess_GameThread, PreSimulate_GameThread, PostSimulate_GameThread, and PostProcess_GameThread instead.")
-		UE_API bool SetupSimulationData(float DeltaTime);
+		UE_API void WriteSimulationData(TMap<int32, FClothSimulData>* OutClothSimulDataOverride = nullptr);
 		UE_API void InitializeConfigs();
 
-		UE_DEPRECATED(5.6, "This method will be made private. Use PreProcess_Internal instead.")
-		UE_API void FillSimulationContext(float DeltaTime, bool bIsInitialization = false);
-
 		UE_API void PreProcess_Internal(float DeltaTime);
 
 	private:
+		UE_API void FillSimulationContext(float DeltaTime, bool bIsInitialization = false);
 
 		// Begin proxy initialization. Can't be done in parallel.
 		UE_API void BeginInitialization_GameThread();
diff --git a/Engine/Source/Editor/Persona/Private/AnimationEditorViewportClient.cpp b/Engine/Source/Editor/Persona/Private/AnimationEditorViewportClient.cpp
--- a/Engine/Source/Editor/Persona/Private/AnimationEditorViewportClient.cpp
+++ b/Engine/Source/Editor/Persona/Private/AnimationEditorViewportClient.cpp
@@ -1396,21 +1396,25 @@ FText FAnimationViewportClient::GetDisplayInfo(bool bDisplayAllInfo) const
 		if (const IClothingSimulationInterface* const ClothingSimulation = ClothingSimulationInstance.GetClothingSimulation())
 		{
 			// Cloth stats
+			if (ClothingSimulationInstances.Num() > 1)
+			{
+				TextValue = ConcatenateLine(TextValue, FText::Format(LOCTEXT("SimulationInstance", "--- Simulation #{0} ---"), Index));
+			}
 			if (const int32 NumActiveCloths = ClothingSimulation->GetNumCloths())
 			{
-				TextValue = ConcatenateLine(TextValue, FText::Format(LOCTEXT("NumActiveCloths", "{0}>Active Cloths: {1}"), Index, NumActiveCloths));
+				TextValue = ConcatenateLine(TextValue, FText::Format(LOCTEXT("NumActiveCloths", "Active Cloths: {0}"), NumActiveCloths));
 			}
 			if (const int32 NumKinematicParticles = ClothingSimulation->GetNumKinematicParticles())
 			{
-				TextValue = ConcatenateLine(TextValue, FText::Format(LOCTEXT("NumKinematicParticles", "{0}>Kinematic Particles: {1}"), Index, NumKinematicParticles));
+				TextValue = ConcatenateLine(TextValue, FText::Format(LOCTEXT("NumKinematicParticles", "Kinematic Particles: {0}"), NumKinematicParticles));
 			}
 			if (const int32 NumDynamicParticles = ClothingSimulation->GetNumDynamicParticles())
 			{
-				TextValue = ConcatenateLine(TextValue, FText::Format(LOCTEXT("NumDynamicParticles", "#{0}>Dynamic Particles: {1}"), Index, NumDynamicParticles));
+				TextValue = ConcatenateLine(TextValue, FText::Format(LOCTEXT("NumDynamicParticles", "Dynamic Particles : {0}"), NumDynamicParticles));
 			}
 			if (const int32 NumIterations = ClothingSimulation->GetNumIterations())
 			{
-				TextValue = ConcatenateLine(TextValue, FText::Format(LOCTEXT("NumIterations", "{0}:>Iterations: {1}"), Index, NumIterations));
+				TextValue = ConcatenateLine(TextValue, FText::Format(LOCTEXT("NumIterations", "Iterations: {0}"), NumIterations));
 			}
 			if (const float SimulationTime = ClothingSimulation->GetSimulationTime())
 			{
@@ -1422,11 +1426,11 @@ FText FAnimationViewportClient::GetDisplayInfo(bool bDisplayAllInfo) const
 				NumberFormatOptions.MaximumIntegralDigits = 6;
 				NumberFormatOptions.MinimumFractionalDigits = 2;
 				NumberFormatOptions.MaximumFractionalDigits = 2;
-				TextValue = ConcatenateLine(TextValue, FText::Format(LOCTEXT("SimulationTime", "{0}>Simulation Time: {1}ms"), Index, FText::AsNumber(SimulationTime, &NumberFormatOptions)));
+				TextValue = ConcatenateLine(TextValue, FText::Format(LOCTEXT("SimulationTime", "Simulation Time: {0}ms"), FText::AsNumber(SimulationTime, &NumberFormatOptions)));
 			}
 			if (ClothingSimulation->IsTeleported())
 			{
-				TextValue = ConcatenateLine(TextValue, FText::Format(LOCTEXT("IsTeleported", "{0}>Simulation Teleport Activated"), Index));
+				TextValue = ConcatenateLine(TextValue, LOCTEXT("IsTeleported", "Simulation Teleport Activated"));
 			}
 		}
 	}
