Cloth Asset - Moved the ChaosClothAssetTools module from the ChaosClothAssetEditor plugin to the ChaosClothAsset plugin.

#rnx
#Synced-CL 48809585

diff --git a/Engine/Plugins/ChaosClothAsset/ChaosClothAsset.uplugin b/Engine/Plugins/ChaosClothAsset/ChaosClothAsset.uplugin
--- a/Engine/Plugins/ChaosClothAsset/ChaosClothAsset.uplugin
+++ b/Engine/Plugins/ChaosClothAsset/ChaosClothAsset.uplugin
@@ -24,6 +24,12 @@
 			"Name": "ChaosClothAssetEngine",
 			"Type": "Runtime",
 			"LoadingPhase": "Default"
+		},
+		{
+			"Name": "ChaosClothAssetTools",
+			"Type": "UncookedOnly",
+			"LoadingPhase": "Default",
+			"PlatformAllowList": [ "Win64", "Mac", "Linux" ]
 		}
 	],
 	"Plugins": [
diff --git a/Engine/Plugins/ChaosClothAsset/Source/ChaosClothAssetTools/ChaosClothAssetTools.Build.cs b/Engine/Plugins/ChaosClothAsset/Source/ChaosClothAssetTools/ChaosClothAssetTools.Build.cs
new file mode 100644
--- /dev/null
+++ b/Engine/Plugins/ChaosClothAsset/Source/ChaosClothAssetTools/ChaosClothAssetTools.Build.cs
@@ -0,0 +1,46 @@
+// Copyright Epic Games, Inc. All Rights Reserved.
+
+using UnrealBuildTool;
+
+public class ChaosClothAssetTools : ModuleRules
+{
+	public ChaosClothAssetTools(ReadOnlyTargetRules Target) : base(Target)
+	{
+		PCHUsage = ModuleRules.PCHUsageMode.UseExplicitOrSharedPCHs;
+
+		PublicDependencyModuleNames.AddRange(
+			new string[] {
+			"MeshConversion",
+			}
+		);
+
+		PrivateDependencyModuleNames.AddRange(
+			new string[] {
+				"Core",
+				"Engine",
+				"CoreUObject",
+				"RenderCore",
+				"RHI",
+				"Chaos",
+				"ChaosCloth",
+				"ChaosClothEditor",  // For Chaos::FSimulationEditorExtender
+				"ChaosClothAsset",
+				"ChaosClothAssetEngine",
+				"DataflowEditor",
+				"DataflowEngine",
+				"SlateCore",
+				"UnrealEd",
+				"Projects",
+				"ClothPainter",                 // For clothing asset exporter
+				"ClothingSystemRuntimeInterface",
+				"ClothingSystemRuntimeCommon",
+				"ClothingSystemEditorInterface",  // For UChaosClothAssetSKMClothingAssetFactory
+				"GeometryCore",					// For DynamicMesh conversion
+				"MeshConversion",
+				"MeshDescription",
+				"SkeletalMeshDescription",		// For FSkeletalMeshAttributes::DefaultSkinWeightProfileName
+				"SkeletalMeshUtilitiesCommon"
+			}
+		);
+	}
+}
diff --git a/Engine/Plugins/ChaosClothAsset/Source/ChaosClothAssetTools/Private/ChaosClothAsset/ChaosClothAssetToolsModule.cpp b/Engine/Plugins/ChaosClothAsset/Source/ChaosClothAssetTools/Private/ChaosClothAsset/ChaosClothAssetToolsModule.cpp
new file mode 100644
--- /dev/null
+++ b/Engine/Plugins/ChaosClothAsset/Source/ChaosClothAssetTools/Private/ChaosClothAsset/ChaosClothAssetToolsModule.cpp
@@ -0,0 +1,76 @@
+// Copyright Epic Games, Inc. All Rights Reserved.
+
+#include "Features/IModularFeatures.h"
+#include "Modules/ModuleManager.h"
+#include "ChaosClothAsset/ClothAssetEditorStyle.h"
+#include "ChaosClothAsset/ClothAssetSKMClothingAssetFactory.h"
+#include "ChaosClothAsset/ClothComponentEditorStyle.h"
+#include "ChaosClothAsset/ClothingAssetToClothAssetExporter.h"
+#include "ChaosClothAsset/SkeletalMeshConverter.h"
+#include "ChaosClothAsset/ClothAssetSKMSimulationEditorExtender.h"
+#include "ClothingSystemEditorInterfaceModule.h"
+
+namespace UE::Chaos::ClothAsset
+{
+	class FChaosClothAssetToolsModule
+		: public IModuleInterface
+		, public IClothingAssetExporterClassProvider
+		, public IClothAssetSkeletalMeshConverterClassProvider
+		, public IClothingAssetFactoryProvider
+
+	{
+	public:
+		// IModuleInterface implementation
+		virtual void StartupModule() override
+		{
+			// Register asset icons
+			FClothAssetEditorStyle::Get();
+			FClothComponentEditorStyle::Get();
+
+			// Register modular features
+			IModularFeatures::Get().RegisterModularFeature(IClothingAssetExporterClassProvider::FeatureName, static_cast<IClothingAssetExporterClassProvider*>(this));
+			IModularFeatures::Get().RegisterModularFeature(IClothAssetSkeletalMeshConverterClassProvider::FeatureName, static_cast<IClothAssetSkeletalMeshConverterClassProvider*>(this));
+			IModularFeatures::Get().RegisterModularFeature(IClothingAssetFactoryProvider::FeatureName, static_cast<IClothingAssetFactoryProvider*>(this));
+			IModularFeatures::Get().RegisterModularFeature(FClothingSystemEditorInterfaceModule::ExtenderFeatureName, &SKMSimulationEditorExtender);
+
+		}
+
+		virtual void ShutdownModule() override
+		{
+			if (UObjectInitialized())
+			{
+				// Unregister modular features
+				IModularFeatures::Get().UnregisterModularFeature(IClothingAssetExporterClassProvider::FeatureName, static_cast<IClothingAssetExporterClassProvider*>(this));
+				IModularFeatures::Get().UnregisterModularFeature(IClothAssetSkeletalMeshConverterClassProvider::FeatureName, static_cast<IClothAssetSkeletalMeshConverterClassProvider*>(this));
+				IModularFeatures::Get().UnregisterModularFeature(IClothingAssetFactoryProvider::FeatureName, static_cast<IClothingAssetFactoryProvider*>(this));
+				IModularFeatures::Get().UnregisterModularFeature(FClothingSystemEditorInterfaceModule::ExtenderFeatureName, &SKMSimulationEditorExtender);
+			}
+		}
+
+		//~ Begin IClothingSimulationFactoryClassProvider implementation
+		virtual TSubclassOf<UClothingAssetExporter> GetClothingAssetExporterClass() const override
+		{
+			return UClothingAssetToChaosClothAssetExporter::StaticClass();
+		}
+		//~ End IClothingSimulationFactoryClassProvider implementation
+
+		//~ Begin IClothAssetSkeletalMeshConverterClassProvider implementation
+		virtual TSubclassOf<UClothAssetSkeletalMeshConverter> GetClothAssetSkeletalMeshConverter() const override
+		{
+			return UClothAssetEditorSkeletalMeshConverter::StaticClass();
+		}
+		//~ End IClothAssetSkeletalMeshConverterClassProvider implementation
+
+		//~ Begin IClothingAssetFactoryProvider
+		virtual UClothingAssetFactoryBase* GetFactory() override
+		{
+			return UChaosClothAssetSKMClothingAssetFactory::StaticClass()->GetDefaultObject<UClothingAssetFactoryBase>();
+		}
+		//~ End IClothingAssetFactoryProvider
+
+	private:
+		UE::Chaos::ClothAsset::FSKMSimulationEditorExtender SKMSimulationEditorExtender;
+	};
+}
+
+IMPLEMENT_MODULE(UE::Chaos::ClothAsset::FChaosClothAssetToolsModule, ChaosClothAssetTools);
diff --git a/Engine/Plugins/ChaosClothAsset/Source/ChaosClothAssetTools/Private/ChaosClothAsset/ClothAssetEditorStyle.cpp b/Engine/Plugins/ChaosClothAsset/Source/ChaosClothAssetTools/Private/ChaosClothAsset/ClothAssetEditorStyle.cpp
new file mode 100644
--- /dev/null
+++ b/Engine/Plugins/ChaosClothAsset/Source/ChaosClothAssetTools/Private/ChaosClothAsset/ClothAssetEditorStyle.cpp
@@ -0,0 +1,33 @@
+// Copyright Epic Games, Inc. All Rights Reserved.
+
+#include "ChaosClothAsset/ClothAssetEditorStyle.h"
+#include "Styling/SlateStyleRegistry.h"
+#include "Interfaces/IPluginManager.h"
+
+namespace UE::Chaos::ClothAsset
+{
+	TOptional<FClothAssetEditorStyle> FClothAssetEditorStyle::Singleton;
+
+	FClothAssetEditorStyle::FClothAssetEditorStyle()
+		: FSlateStyleSet("ClothAssetEditorStyle")
+	{
+		TSharedPtr<IPlugin> ChaosClothAssetPlugin = IPluginManager::Get().FindPlugin("ChaosClothAssetEditor");
+		if (ChaosClothAssetPlugin.IsValid())
+		{
+			SetContentRoot(ChaosClothAssetPlugin->GetBaseDir() / TEXT("Resources"));
+
+			Set("ClassIcon.ChaosClothAsset", new FSlateVectorImageBrush(RootToContentDir(TEXT("ClothAsset_16.svg")), FVector2D(16)));
+			Set("ClassThumbnail.ChaosClothAsset", new FSlateVectorImageBrush(RootToContentDir(TEXT("ClothAsset_64.svg")), FVector2D(64)));
+
+			Set("ClassIcon.ChaosClothPreset", new FSlateVectorImageBrush(RootToContentDir(TEXT("ClothPreset_16.svg")), FVector2D(16)));
+			Set("ClassThumbnail.ChaosClothPreset", new FSlateVectorImageBrush(RootToContentDir(TEXT("ClothPreset_64.svg")), FVector2D(64)));
+		}
+
+		FSlateStyleRegistry::RegisterSlateStyle(*this);
+	}
+
+	FClothAssetEditorStyle::~FClothAssetEditorStyle()
+	{
+		FSlateStyleRegistry::UnRegisterSlateStyle(*this);
+	}
+}  // End namespace UE::Chaos::ClothAsset
diff --git a/Engine/Plugins/ChaosClothAsset/Source/ChaosClothAssetTools/Private/ChaosClothAsset/ClothAssetFactory.cpp b/Engine/Plugins/ChaosClothAsset/Source/ChaosClothAssetTools/Private/ChaosClothAsset/ClothAssetFactory.cpp
new file mode 100644
--- /dev/null
+++ b/Engine/Plugins/ChaosClothAsset/Source/ChaosClothAssetTools/Private/ChaosClothAsset/ClothAssetFactory.cpp
@@ -0,0 +1,45 @@
+// Copyright Epic Games, Inc. All Rights Reserved.
+
+#include "ChaosClothAsset/ClothAssetFactory.h"
+#include "ChaosClothAsset/ClothAsset.h"
+#include "Dataflow/AssetDefinition_DataflowAsset.h"
+#include "HAL/IConsoleManager.h"
+
+#include UE_INLINE_GENERATED_CPP_BY_NAME(ClothAssetFactory)
+
+#define LOCTEXT_NAMESPACE "ChaosClothAssetFactory"
+
+DEFINE_LOG_CATEGORY_STATIC(LogChaosClothAssetFactory, Log, All);
+
+namespace UE::Chaos::ClothAsset::Private
+{
+	static TAutoConsoleVariable<bool> CVarCreateWithEmbeddedDataflow(
+		TEXT("p.ChaosClothAsset.CreateWithEmbeddedDataflow"),
+		true,
+		TEXT("Set how to create new Cloth Assets, either with embedded (default) or external (legacy) Dataflow assets.\n"));
+}
+
+UChaosClothAssetFactory::UChaosClothAssetFactory(const FObjectInitializer& ObjectInitializer)
+	: Super(ObjectInitializer)
+{
+	bEditorImport = true;
+	bEditAfterNew = true;
+	SupportedClass = UChaosClothAsset::StaticClass();
+}
+
+UObject* UChaosClothAssetFactory::FactoryCreateNew(UClass* Class, UObject* Parent, FName Name, EObjectFlags Flags, UObject* /*Context*/, FFeedbackContext* /*Warn*/)
+{
+	using namespace UE::Chaos::ClothAsset;
+
+	static const FString TemplatePath = TEXT("/ChaosClothAssetEditor/ClothAssetTemplate.ClothAssetTemplate");
+	static const TCHAR* AssetPrefix = TEXT("CA_");
+	const bool bEmbedDataflow = Private::CVarCreateWithEmbeddedDataflow.GetValueOnAnyThread();
+	return UE::DataflowAssetDefinitionHelpers::FactoryCreateNew(Class, Parent, Name, Flags, &TemplatePath, bEmbedDataflow, AssetPrefix);
+}
+
+FString UChaosClothAssetFactory::GetDefaultNewAssetName() const
+{
+	return FString(TEXT("CA_NewChaosClothAsset"));
+}
+
+#undef LOCTEXT_NAMESPACE
diff --git a/Engine/Plugins/ChaosClothAsset/Source/ChaosClothAssetTools/Private/ChaosClothAsset/ClothAssetSKMClothingAssetFactory.cpp b/Engine/Plugins/ChaosClothAsset/Source/ChaosClothAssetTools/Private/ChaosClothAsset/ClothAssetSKMClothingAssetFactory.cpp
new file mode 100644
--- /dev/null
+++ b/Engine/Plugins/ChaosClothAsset/Source/ChaosClothAssetTools/Private/ChaosClothAsset/ClothAssetSKMClothingAssetFactory.cpp
@@ -0,0 +1,33 @@
+// Copyright Epic Games, Inc. All Rights Reserved.
+
+#include "ChaosClothAsset/ClothAssetSKMClothingAssetFactory.h"
+#include "ChaosClothAsset/ClothAssetBase.h"
+#include "Engine/SkeletalMesh.h"
+#include "ChaosClothAsset/ClothAssetSKMClothingAsset.h"
+
+#include UE_INLINE_GENERATED_CPP_BY_NAME(ClothAssetSKMClothingAssetFactory)
+
+UChaosClothAssetSKMClothingAssetFactory::UChaosClothAssetSKMClothingAssetFactory(const FObjectInitializer& ObjectInitializer)
+	: Super(ObjectInitializer)
+{
+}
+
+const UClass* UChaosClothAssetSKMClothingAssetFactory::GetSupportedSourceAssetType() const
+{
+	return UChaosClothAssetBase::StaticClass();
+}
+
+TArray<UClothingAssetBase*> UChaosClothAssetSKMClothingAssetFactory::CreateFromSourceAsset(USkeletalMesh* TargetMesh, const UObject* SourceAsset) const
+{
+	TArray<UClothingAssetBase*> ClothingAssets;
+	if (const UChaosClothAssetBase* const ClothAsset = Cast<UChaosClothAssetBase>(SourceAsset))
+	{
+		const FName AssetName = StaticFindObject(UClothingAssetBase::StaticClass(), TargetMesh, *ClothAsset->GetName()) ?
+			MakeUniqueObjectName(TargetMesh, UClothingAssetBase::StaticClass(), ClothAsset->GetFName()) :
+			ClothAsset->GetFName();
+		UChaosClothAssetSKMClothingAsset* const ClothAssetSKMClothingAsset = NewObject<UChaosClothAssetSKMClothingAsset>(TargetMesh, AssetName);
+		ClothAssetSKMClothingAsset->SetAsset(Cast<UChaosClothAssetBase>(SourceAsset));
+		ClothingAssets.Emplace(ClothAssetSKMClothingAsset);
+	}
+	return ClothingAssets;
+}
diff --git a/Engine/Plugins/ChaosClothAsset/Source/ChaosClothAssetTools/Private/ChaosClothAsset/ClothAssetSKMClothingAssetFactory.h b/Engine/Plugins/ChaosClothAsset/Source/ChaosClothAssetTools/Private/ChaosClothAsset/ClothAssetSKMClothingAssetFactory.h
new file mode 100644
--- /dev/null
+++ b/Engine/Plugins/ChaosClothAsset/Source/ChaosClothAssetTools/Private/ChaosClothAsset/ClothAssetSKMClothingAssetFactory.h
@@ -0,0 +1,35 @@
+// Copyright Epic Games, Inc. All Rights Reserved.
+
+#pragma once
+
+#include "ClothingAssetFactoryInterface.h"
+#include "ClothAssetSKMClothingAssetFactory.generated.h"
+
+UCLASS(MinimalAPI)
+class UChaosClothAssetSKMClothingAssetFactory : public UClothingAssetFactoryBase
+{
+	GENERATED_BODY()
+public:
+	UChaosClothAssetSKMClothingAssetFactory(const FObjectInitializer& ObjectInitializer);
+
+	//~ Begin UClothingAssetFactoryBase Interface
+	virtual const UClass* GetSupportedSourceAssetType() const override;
+
+	virtual TArray<UClothingAssetBase*> CreateFromSourceAsset(USkeletalMesh* TargetMesh, const UObject* SourceAsset) const override;
+
+	virtual UClothingAssetBase* CreateFromSkeletalMesh(USkeletalMesh* TargetMesh, FSkeletalMeshClothBuildParams& Params) override
+	{
+		return nullptr;  // TODO: Implement create from skeletal mesh
+	}
+
+	virtual UClothingAssetBase* CreateFromExistingCloth(USkeletalMesh* /*TargetMesh*/, USkeletalMesh* /*SourceMesh*/, UClothingAssetBase* /*SourceAsset*/) override
+	{
+		return nullptr;  // TODO: Implement create from another clothing data
+	}
+
+	virtual UClothingAssetBase* ImportLodToClothing(USkeletalMesh* TargetMesh, FSkeletalMeshClothBuildParams& Params) override
+	{
+		return nullptr;
+	}
+	//~ End UClothingAssetFactoryBase Interface
+};
diff --git a/Engine/Plugins/ChaosClothAsset/Source/ChaosClothAssetTools/Private/ChaosClothAsset/ClothAssetSKMSimulationEditorExtender.cpp b/Engine/Plugins/ChaosClothAsset/Source/ChaosClothAssetTools/Private/ChaosClothAsset/ClothAssetSKMSimulationEditorExtender.cpp
new file mode 100644
--- /dev/null
+++ b/Engine/Plugins/ChaosClothAsset/Source/ChaosClothAssetTools/Private/ChaosClothAsset/ClothAssetSKMSimulationEditorExtender.cpp
@@ -0,0 +1,38 @@
+// Copyright Epic Games, Inc. All Rights Reserved.
+
+#include "ChaosClothAsset/ClothAssetSKMSimulationEditorExtender.h"
+#include "ChaosClothAsset/ClothAssetSKMClothingSimulation.h"
+
+#define LOCTEXT_NAMESPACE "ChaosSimulationEditorExtender"
+
+namespace UE::Chaos::ClothAsset
+{
+	FSKMSimulationEditorExtender::FSKMSimulationEditorExtender()
+		: ::Chaos::FSimulationEditorExtender()
+	{
+	}
+
+	FSKMSimulationEditorExtender::~FSKMSimulationEditorExtender() = default;
+
+	UClass* FSKMSimulationEditorExtender::GetSupportedSimulationFactoryClass()
+	{
+		return UChaosClothAssetSKMClothingSimulationFactory::StaticClass();
+	}
+
+	FName FSKMSimulationEditorExtender::GetSectionName() const
+	{
+		return FName("ClothAsset_Visualizations");
+	}
+
+	FText FSKMSimulationEditorExtender::GetSectionHeadingText() const
+	{
+		return FText(LOCTEXT("ClothAssetVisualizationSection", "Cloth Asset Visualization"));
+	}
+
+	const ::Chaos::FClothVisualizationNoGC* FSKMSimulationEditorExtender::GetClothVisualization(const IClothingSimulationInterface* Simulation) const
+	{
+		return Simulation ? static_cast<const FSKMClothingSimulation*>(Simulation)->GetClothVisualization() : nullptr;
+	}
+}  // namespace UE::Chaos::ClothAsset
+
+#undef LOCTEXT_NAMESPACE
diff --git a/Engine/Plugins/ChaosClothAsset/Source/ChaosClothAssetTools/Private/ChaosClothAsset/ClothAssetSKMSimulationEditorExtender.h b/Engine/Plugins/ChaosClothAsset/Source/ChaosClothAssetTools/Private/ChaosClothAsset/ClothAssetSKMSimulationEditorExtender.h
new file mode 100644
--- /dev/null
+++ b/Engine/Plugins/ChaosClothAsset/Source/ChaosClothAssetTools/Private/ChaosClothAsset/ClothAssetSKMSimulationEditorExtender.h
@@ -0,0 +1,32 @@
+// Copyright Epic Games, Inc. All Rights Reserved.
+
+#pragma once
+
+#include "ChaosClothEditor/ChaosSimulationEditorExtender.h"
+
+namespace Chaos
+{
+	class FClothVisualizationNoGC;
+}
+
+namespace UE::Chaos::ClothAsset
+{
+	/** Chaos extension to the asset editor. */
+	class FSKMSimulationEditorExtender final : public ::Chaos::FSimulationEditorExtender
+	{
+	public:
+		FSKMSimulationEditorExtender();
+		virtual ~FSKMSimulationEditorExtender() override;
+
+	private:
+		//~ Begin ISimulationEditorExtender Interface
+		virtual UClass* GetSupportedSimulationFactoryClass() override;
+		//~ End ISimulationEditorExtender Interface
+
+		//~ Begin ::Chaos::FSimulationEditorExtender Interface
+		virtual FName GetSectionName() const override;
+		virtual FText GetSectionHeadingText() const override;
+		virtual const ::Chaos::FClothVisualizationNoGC* GetClothVisualization(const IClothingSimulationInterface* Simulation) const override;
+		//~ End ::Chaos::FSimulationEditorExtender Interface
+	};
+}
diff --git a/Engine/Plugins/ChaosClothAsset/Source/ChaosClothAssetTools/Private/ChaosClothAsset/ClothComponentEditorStyle.cpp b/Engine/Plugins/ChaosClothAsset/Source/ChaosClothAssetTools/Private/ChaosClothAsset/ClothComponentEditorStyle.cpp
new file mode 100644
--- /dev/null
+++ b/Engine/Plugins/ChaosClothAsset/Source/ChaosClothAssetTools/Private/ChaosClothAsset/ClothComponentEditorStyle.cpp
@@ -0,0 +1,30 @@
+// Copyright Epic Games, Inc. All Rights Reserved.
+
+#include "ClothComponentEditorStyle.h"
+#include "Styling/SlateStyleRegistry.h"
+#include "Interfaces/IPluginManager.h"
+
+namespace UE::Chaos::ClothAsset
+{
+	TOptional<FClothComponentEditorStyle> FClothComponentEditorStyle::Singleton;
+
+	FClothComponentEditorStyle::FClothComponentEditorStyle()
+		: FSlateStyleSet("ClothComponentEditorStyle")
+	{
+		TSharedPtr<IPlugin> ChaosClothAssetPlugin = IPluginManager::Get().FindPlugin("ChaosClothAssetEditor");
+		if (ChaosClothAssetPlugin.IsValid())
+		{
+			SetContentRoot(ChaosClothAssetPlugin->GetBaseDir() / TEXT("Resources"));
+
+			Set("ClassIcon.ChaosClothComponent", new FSlateVectorImageBrush(RootToContentDir(TEXT("ClothAsset_16.svg")), FVector2D(16)));
+			Set("ClassThumbnail.ChaosClothComponent", new FSlateVectorImageBrush(RootToContentDir(TEXT("ClothAsset_64.svg")), FVector2D(64)));
+
+			FSlateStyleRegistry::RegisterSlateStyle(*this);
+		}
+	}
+
+	FClothComponentEditorStyle::~FClothComponentEditorStyle()
+	{
+		FSlateStyleRegistry::UnRegisterSlateStyle(*this);
+	}
+}  // End namespace UE::Chaos::ClothAsset
diff --git a/Engine/Plugins/ChaosClothAsset/Source/ChaosClothAssetTools/Private/ChaosClothAsset/ClothComponentEditorStyle.h b/Engine/Plugins/ChaosClothAsset/Source/ChaosClothAssetTools/Private/ChaosClothAsset/ClothComponentEditorStyle.h
new file mode 100644
--- /dev/null
+++ b/Engine/Plugins/ChaosClothAsset/Source/ChaosClothAssetTools/Private/ChaosClothAsset/ClothComponentEditorStyle.h
@@ -0,0 +1,37 @@
+// Copyright Epic Games, Inc. All Rights Reserved.
+
+#pragma once
+
+#include "Styling/SlateStyle.h"
+
+namespace UE::Chaos::ClothAsset
+{
+	/**
+	 * Editor style setting up the cloth asset icons in editor.
+	 * TODO: Merge all styles to a single singleton.
+	 */
+	class FClothComponentEditorStyle final : public FSlateStyleSet
+	{
+	public:
+		FClothComponentEditorStyle();
+		virtual ~FClothComponentEditorStyle() override;
+
+	public:
+		static FClothComponentEditorStyle& Get()
+		{
+			if (!Singleton.IsSet())
+			{
+				Singleton.Emplace();
+			}
+			return Singleton.GetValue();
+		}
+
+		static void Destroy()
+		{
+			Singleton.Reset();
+		}
+
+	private:
+		static TOptional<FClothComponentEditorStyle> Singleton;
+	};
+}  // End namespace UE::Chaos::ClothAsset
diff --git a/Engine/Plugins/ChaosClothAsset/Source/ChaosClothAssetTools/Private/ChaosClothAsset/ClothPatternToDynamicMesh.cpp b/Engine/Plugins/ChaosClothAsset/Source/ChaosClothAssetTools/Private/ChaosClothAsset/ClothPatternToDynamicMesh.cpp
new file mode 100644
--- /dev/null
+++ b/Engine/Plugins/ChaosClothAsset/Source/ChaosClothAssetTools/Private/ChaosClothAsset/ClothPatternToDynamicMesh.cpp
@@ -0,0 +1,780 @@
+// Copyright Epic Games, Inc. All Rights Reserved. 
+
+#if WITH_EDITOR
+
+#include "ChaosClothAsset/ClothPatternToDynamicMesh.h"
+#include "ChaosClothAsset/ClothPatternToDynamicMeshMappingSupport.h"
+#include "ChaosClothAsset/ClothAsset.h"
+#include "ChaosClothAsset/ClothComponent.h"
+#include "ChaosClothAsset/CollectionClothFacade.h"
+#include "ChaosClothAsset/ClothCollectionGroup.h"
+#include "DynamicMesh/MeshNormals.h"
+#include "Engine/SkeletalMesh.h"
+#include "SkeletalMeshAttributes.h"
+#include "ToDynamicMesh.h"
+
+namespace UE::Chaos::ClothAsset
+{
+
+//
+// Wrapper for accessing a Cloth Pattern. Implements the interface expected by TToDynamicMesh<>.
+//
+class FClothPatternWrapper
+{
+public:
+
+	typedef int32 TriIDType;
+	typedef int32 VertIDType;
+	typedef int32 WedgeIDType;
+
+	typedef int32 UVIDType;
+	typedef int32 NormalIDType;
+	typedef int32 ColorIDType;
+
+	FClothPatternWrapper(const FCollectionClothConstFacade& ClothFacade, int32 PatternIndex, EClothPatternVertexType VertexDataType) :
+		VertexDataType(VertexDataType),
+		Cloth(ClothFacade)
+	{
+		NumTexCoords = 0;
+		if (PatternIndex == INDEX_NONE)
+		{
+			// All patterns in one dynamic mesh
+			switch (VertexDataType)
+			{
+			case EClothPatternVertexType::Render:
+			{
+				const int32 NumVertices = Cloth.GetNumRenderVertices();
+				VertIDs.SetNum(NumVertices); 
+				TConstArrayView<TArray<FVector2f>> UVs = Cloth.GetRenderUVs();
+				for (int32 VtxIndex = 0; VtxIndex < NumVertices; ++VtxIndex)
+				{
+					VertIDs[VtxIndex] = VtxIndex; 
+					NumTexCoords = FMath::Max(NumTexCoords, UVs[VtxIndex].Num());
+				}
+				NormalIDs = VertIDs;
+
+				const TConstArrayView<FIntVector3> Indices = Cloth.GetRenderIndices();
+				const int32 NumFaces = Cloth.GetNumRenderFaces();
+				TriIDs.Reserve(NumFaces);
+				for (int32 TriIndex = 0; TriIndex < NumFaces; ++TriIndex)
+				{
+					if (Indices[TriIndex][0] != INDEX_NONE &&
+						Indices[TriIndex][1] != INDEX_NONE &&
+						Indices[TriIndex][2] != INDEX_NONE)
+					{
+						TriIDs.Add(TriIndex);
+					}
+				}
+			} break;
+			case EClothPatternVertexType::Sim2D:
+			{
+				const int32 NumVertices = Cloth.GetNumSimVertices2D();
+				const TConstArrayView<int32> SimVertex3DLookup = Cloth.GetSimVertex3DLookup();
+				VertIDs.Reserve(NumVertices);
+				NormalIDs.Reserve(NumVertices);
+				for (int32 VtxIndex = 0; VtxIndex < NumVertices; ++VtxIndex)
+				{
+					if (SimVertex3DLookup[VtxIndex] != INDEX_NONE)
+					{
+						VertIDs.Add(VtxIndex);
+						NormalIDs.Add(SimVertex3DLookup[VtxIndex]);
+					}
+				}
+
+				const TConstArrayView<FIntVector3> Indices = Cloth.GetSimIndices2D();
+				const int32 NumFaces = Cloth.GetNumSimFaces();
+				TriIDs.Reserve(NumFaces);
+				for (int32 TriIndex = 0; TriIndex < NumFaces; ++TriIndex)
+				{
+					if (Indices[TriIndex][0] != INDEX_NONE &&
+						Indices[TriIndex][1] != INDEX_NONE &&
+						Indices[TriIndex][2] != INDEX_NONE &&
+						SimVertex3DLookup[Indices[TriIndex][0]] != INDEX_NONE &&
+						SimVertex3DLookup[Indices[TriIndex][1]] != INDEX_NONE &&
+						SimVertex3DLookup[Indices[TriIndex][2]] != INDEX_NONE)
+					{
+						TriIDs.Add(TriIndex);
+					}
+				}
+			}
+			break;
+			case EClothPatternVertexType::Sim3D:
+			{
+				const int32 NumVertices = Cloth.GetNumSimVertices3D();
+				VertIDs.SetNum(NumVertices);
+				for (int32 VtxIndex = 0; VtxIndex < NumVertices; ++VtxIndex)
+				{
+					VertIDs[VtxIndex] = VtxIndex;
+				}
+				NormalIDs = VertIDs;
+
+				const TConstArrayView<FIntVector3> Indices = Cloth.GetSimIndices3D();
+				const int32 NumFaces = Cloth.GetNumSimFaces();
+				TriIDs.Reserve(NumFaces);
+				for (int32 TriIndex = 0; TriIndex < NumFaces; ++TriIndex)
+				{
+					if (Indices[TriIndex][0] != INDEX_NONE &&
+						Indices[TriIndex][1] != INDEX_NONE &&
+						Indices[TriIndex][2] != INDEX_NONE)
+					{
+						TriIDs.Add(TriIndex);
+					}
+				}
+				NumTexCoords = 1; // SimPositions2D will be stored as UVs
+			}break;
+			default:
+				checkNoEntry();
+			}
+		}
+		else
+		{
+			switch (VertexDataType)
+			{
+			case EClothPatternVertexType::Render:
+			{
+				FCollectionClothRenderPatternConstFacade Pattern = Cloth.GetRenderPattern(PatternIndex);
+				const int32 NumVertices = Pattern.GetNumRenderVertices();
+				const int32 VertexOffset = Pattern.GetRenderVerticesOffset();
+				VertIDs.SetNum(NumVertices);
+				TConstArrayView<TArray<FVector2f>> UVs = Pattern.GetRenderUVs();
+				for (int32 VtxIndex = 0; VtxIndex < NumVertices; ++VtxIndex)
+				{
+					VertIDs[VtxIndex] = VtxIndex + VertexOffset;
+					NumTexCoords = FMath::Max(NumTexCoords, UVs[VtxIndex].Num());
+				}
+				NormalIDs = VertIDs;
+
+				const TConstArrayView<FIntVector3> Indices = Pattern.GetRenderIndices();
+				const int32 NumFaces = Pattern.GetNumRenderFaces();
+				const int32 FaceOffset = Pattern.GetRenderFacesOffset();
+				TriIDs.Reserve(NumFaces);
+				for (int32 TriIndex = 0; TriIndex < NumFaces; ++TriIndex)
+				{
+					if (Indices[TriIndex][0] != INDEX_NONE &&
+						Indices[TriIndex][1] != INDEX_NONE &&
+						Indices[TriIndex][2] != INDEX_NONE)
+					{
+						TriIDs.Add(TriIndex + FaceOffset);
+					}
+				}
+			} break;
+			case EClothPatternVertexType::Sim2D:
+			{
+				FCollectionClothSimPatternConstFacade Pattern = Cloth.GetSimPattern(PatternIndex);
+				const int32 NumVertices = Pattern.GetNumSimVertices2D();
+				const TConstArrayView<int32> SimVertex3DLookup = Pattern.GetSimVertex3DLookup();
+				const int32 VertexOffset = Pattern.GetSimVertices2DOffset();
+				VertIDs.Reserve(NumVertices);
+				NormalIDs.Reserve(NumVertices);
+				for (int32 VtxIndex = 0; VtxIndex < NumVertices; ++VtxIndex)
+				{
+					if (SimVertex3DLookup[VtxIndex] != INDEX_NONE)
+					{
+						VertIDs.Add(VtxIndex + VertexOffset);
+						NormalIDs.Add(SimVertex3DLookup[VtxIndex]);
+					}
+				}
+
+				const TConstArrayView<FIntVector3> Indices = Pattern.GetSimIndices2D();
+				const int32 NumFaces = Pattern.GetNumSimFaces();
+				const int32 FaceOffset = Pattern.GetSimFacesOffset();
+				TriIDs.Reserve(NumFaces);
+				for (int32 TriIndex = 0; TriIndex < NumFaces; ++TriIndex)
+				{
+					if (Indices[TriIndex][0] != INDEX_NONE &&
+						Indices[TriIndex][1] != INDEX_NONE &&
+						Indices[TriIndex][2] != INDEX_NONE &&
+						SimVertex3DLookup[Indices[TriIndex][0] - VertexOffset] != INDEX_NONE &&
+						SimVertex3DLookup[Indices[TriIndex][1] - VertexOffset] != INDEX_NONE &&
+						SimVertex3DLookup[Indices[TriIndex][2] - VertexOffset] != INDEX_NONE)
+					{
+						TriIDs.Add(TriIndex + FaceOffset);
+					}
+				}
+			} break;
+			case EClothPatternVertexType::Sim3D:
+			{
+				FCollectionClothSimPatternConstFacade Pattern = Cloth.GetSimPattern(PatternIndex);
+				VertIDs = Pattern.GetSimVertex3DLookup();
+				NormalIDs = Pattern.GetSimVertex3DLookup();
+
+				const TConstArrayView<FIntVector3> Indices = Pattern.GetSimIndices3D();
+				const int32 NumFaces = Pattern.GetNumSimFaces();
+				const int32 FaceOffset = Pattern.GetSimFacesOffset();
+				TriIDs.Reserve(NumFaces);
+				for (int32 TriIndex = 0; TriIndex < NumFaces; ++TriIndex)
+				{
+					if (Indices[TriIndex][0] != INDEX_NONE &&
+						Indices[TriIndex][1] != INDEX_NONE &&
+						Indices[TriIndex][2] != INDEX_NONE)
+					{
+						TriIDs.Add(TriIndex + FaceOffset);
+					}
+				}
+				NumTexCoords = 1; // SimPositions2D will be stored as UVs
+			} break;
+			default:
+				checkNoEntry();
+			}
+		}
+
+		//
+		// Weight map layers precomputation
+		//
+		if(VertexDataType != EClothPatternVertexType::Render)
+		{
+			WeightMapNames = Cloth.GetWeightMapNames();
+		}
+		else 
+		{
+			WeightMapNames = Cloth.GetUserDefinedAttributeNames<float>(ClothCollectionGroup::RenderVertices);
+		}
+
+		// Set the reference skeleton if available
+		const FSoftObjectPath& SkeletalMeshPathName = ClothFacade.GetSkeletalMeshSoftObjectPathName();
+		const USkeletalMesh* const SkeletalMesh = Cast<USkeletalMesh>(SkeletalMeshPathName.TryLoad());
+		RefSkeleton = SkeletalMesh ? &SkeletalMesh->GetRefSkeleton() : nullptr;
+	}
+
+	int32 NumTris() const
+	{
+		return TriIDs.Num();
+	}
+
+	int32 NumVerts() const
+	{
+		return VertIDs.Num();
+	}
+
+	int32 NumUVLayers() const
+	{
+		return NumTexCoords;
+	}
+
+	int32 NumWeightMapLayers() const
+	{
+		return WeightMapNames.Num();
+	}
+
+	FName GetWeightMapName(int32 LayerIndex) const
+	{
+		return WeightMapNames[LayerIndex];
+	}
+
+	float GetVertexWeight(int32 LayerIndex, VertIDType VertexIndex) const
+	{
+		checkSlow(LayerIndex < WeightMapNames.Num());
+
+		// All weights live on 3D indices.
+		const int32 VertexWeightIndex = (VertexDataType == EClothPatternVertexType::Sim2D) ? Cloth.GetSimVertex3DLookup()[VertexIndex] : VertexIndex;
+
+		TConstArrayView<float> WeightMap;
+		
+		if (VertexDataType != EClothPatternVertexType::Render)
+		{
+			WeightMap = Cloth.GetWeightMap(WeightMapNames[LayerIndex]);
+		}
+		else
+		{
+			WeightMap = Cloth.GetUserDefinedAttribute<float>(WeightMapNames[LayerIndex], ClothCollectionGroup::RenderVertices);
+		}
+
+		checkSlow(VertexWeightIndex < WeightMap.Num());
+		if (VertexWeightIndex < WeightMap.Num())
+		{
+			return WeightMap[VertexWeightIndex];
+		}
+		else
+		{
+			return 0;
+		}
+	}
+
+	// --"Vertex Buffer" info
+	const TArray<VertIDType>& GetVertIDs() const
+	{
+		return VertIDs;
+	}
+
+	FVector3d GetPosition(VertIDType VtxID) const
+	{
+		switch (VertexDataType)
+		{
+			case EClothPatternVertexType::Render:
+			{
+				const TConstArrayView<FVector3f> RenderPositions = Cloth.GetRenderPosition();
+				return FVector3d(RenderPositions[VtxID]);
+			}
+			case EClothPatternVertexType::Sim2D:
+			{
+				const TConstArrayView<FVector2f> SimPositions = Cloth.GetSimPosition2D();
+				const FVector2f& Pos = SimPositions[VtxID];
+				return FVector3d(Pos[0], Pos[1], 0.0);
+			}
+			case EClothPatternVertexType::Sim3D:
+			{
+				const TConstArrayView<FVector3f> SimPositions = Cloth.GetSimPosition3D();
+				const FVector3f& Pos = SimPositions[VtxID];
+				return FVector3d(Pos[0], Pos[1], Pos[2]);
+			}
+			default:
+				checkNoEntry();
+				return FVector3d();
+		};
+	}
+
+	// --"Index Buffer" info
+	const TArray<TriIDType>& GetTriIDs() const
+	{
+		return TriIDs;
+	}
+
+	bool GetTri(TriIDType TriID, VertIDType& VID0, VertIDType& VID1, VertIDType& VID2) const
+	{
+		FIntVector Face;
+
+		switch (VertexDataType)
+		{
+		case EClothPatternVertexType::Render:
+			Face = Cloth.GetRenderIndices()[TriID];
+			break;
+		case EClothPatternVertexType::Sim2D:
+			Face = Cloth.GetSimIndices2D()[TriID];
+			break;
+		case EClothPatternVertexType::Sim3D:
+			Face = Cloth.GetSimIndices3D()[TriID];
+			break;
+		default:
+			checkNoEntry();
+		}
+
+		VID0 = Face[0];
+		VID1 = Face[1];
+		VID2 = Face[2];
+
+		return true;
+	}
+	
+	bool HasNormals() const
+	{
+		return true;
+	}
+
+	bool HasTangents() const
+	{
+		return VertexDataType == EClothPatternVertexType::Render;
+	}
+
+	bool HasBiTangents() const
+	{
+		return VertexDataType == EClothPatternVertexType::Render;
+	}
+
+	bool HasColors() const
+	{
+		return VertexDataType == EClothPatternVertexType::Render;
+	}
+
+	// -- Access to per-wedge attributes -- //
+	void GetWedgeIDs(const TriIDType& TriID, WedgeIDType& WID0, WedgeIDType& WID1, WedgeIDType& WID2) const
+	{
+		check(VertexDataType == EClothPatternVertexType::Sim3D);
+		const FIntVector& Face = Cloth.GetSimIndices2D()[TriID];
+
+		WID0 = Face[0];
+		WID1 = Face[1];
+		WID2 = Face[2];
+	}
+
+	FVector2f GetWedgeUV(int32 UVLayerIndex, WedgeIDType WID) const
+	{
+		check(VertexDataType == EClothPatternVertexType::Sim3D);
+		return Cloth.GetSimPosition2D()[WID];
+	}
+
+	FVector3f GetWedgeNormal(WedgeIDType WID) const
+	{
+		checkf(false, TEXT("FClothPatternWrapper: ClothPatterns are not expected to use Wedges"));
+		return FVector3f();
+	}
+
+	FVector3f GetWedgeTangent(WedgeIDType WID) const
+	{
+		checkf(false, TEXT("FClothPatternWrapper: ClothPatterns are not expected to use Wedges"));
+		return FVector3f();
+	}
+
+	FVector3f GetWedgeBiTangent(WedgeIDType WID) const
+	{
+		checkf(false, TEXT("FClothPatternWrapper: ClothPatterns are not expected to use Wedges"));
+		return FVector3f();
+	}
+	
+	FVector4f GetWedgeColor(WedgeIDType WID) const
+	{
+		checkf(false, TEXT("FClothPatternWrapper: ClothPatterns are not expected to use Wedges"));
+		return FVector4f();
+	}
+	// -- End of per-wedge attribute access -- //
+
+	int32 GetMaterialIndex(TriIDType TriID) const
+	{
+		checkf(false, TEXT("FClothPatternWrapper: Material indexing should be accomplished by passing a function into Convert"));
+		return 0;
+	}
+
+	int32 NumSkinWeightAttributes() const 
+	{ 
+		return 1;
+	}
+
+	UE::AnimationCore::FBoneWeights GetVertexSkinWeight(int32 SkinWeightAttributeIndex, VertIDType InVertexID) const 
+	{
+		using namespace UE::AnimationCore;
+		
+		checkfSlow(SkinWeightAttributeIndex == 0, TEXT("Cloth assets should only have one skin weight profile")); 
+
+		const bool bGetRenderMeshData = (VertexDataType == EClothPatternVertexType::Render);
+		const TConstArrayView<TArray<int32>> BoneIndices = bGetRenderMeshData ? Cloth.GetRenderBoneIndices() : Cloth.GetSimBoneIndices();
+		const TConstArrayView<TArray<float>> BoneWeights = bGetRenderMeshData ? Cloth.GetRenderBoneWeights() : Cloth.GetSimBoneWeights();
+
+		// All weights live on 3D indices. Need to convert 2D index to 3D index.
+		const VertIDType VertexID = (VertexDataType == EClothPatternVertexType::Sim2D) ? Cloth.GetSimVertex3DLookup()[InVertexID] : InVertexID;
+
+		if (ensure(VertexID >= 0 && VertexID < BoneIndices.Num()))
+		{
+			const TArray<int32> Indices = BoneIndices[VertexID];
+			const TArray<float> Weights = BoneWeights[VertexID];
+			const int32 NumInfluences = Indices.Num();
+			check(Weights.Num() == NumInfluences);
+			
+			TArray<FBoneWeight> BoneWeightArray;
+			BoneWeightArray.SetNumUninitialized(NumInfluences);
+
+			for (int32 Idx = 0; Idx < NumInfluences; ++Idx)
+			{
+				BoneWeightArray[Idx] = FBoneWeight(static_cast<FBoneIndexType>(Indices[Idx]), Weights[Idx]);
+			}
+
+			return FBoneWeights::Create(BoneWeightArray, FBoneWeightsSettings());
+		}
+		else 
+		{
+			return FBoneWeights();
+		}
+	}
+
+	FName GetSkinWeightAttributeName(int32 SkinWeightAttributeIndex) const 
+	{ 
+		checkfSlow(SkinWeightAttributeIndex == 0, TEXT("Cloth assets should only have one skin weight profile")); 
+		
+		return FSkeletalMeshAttributes::DefaultSkinWeightProfileName;
+	}
+
+	int32 GetNumBones() const 
+	{ 
+		return RefSkeleton ? RefSkeleton->GetRawBoneNum() : 0;
+	}
+
+    FName GetBoneName(int32 BoneIdx) const
+	{
+		if (ensure(BoneIdx >= 0 && BoneIdx < GetNumBones()) && RefSkeleton)
+		{
+			return RefSkeleton->GetRawRefBoneInfo()[BoneIdx].Name;
+		}
+		
+		return NAME_None;
+	}
+
+	int32 GetBoneParentIndex(int32 BoneIdx) const
+	{
+		if (ensure(BoneIdx >= 0 && BoneIdx < GetNumBones()) && RefSkeleton)
+		{
+			return RefSkeleton->GetRawRefBoneInfo()[BoneIdx].ParentIndex;
+		}
+
+		return INDEX_NONE;
+	}
+
+	FTransform GetBonePose(int32 BoneIdx) const
+	{
+		if (ensure(BoneIdx >= 0 && BoneIdx < GetNumBones()) && RefSkeleton)
+		{
+			return RefSkeleton->GetRawRefBonePose()[BoneIdx];
+		}
+
+		return FTransform::Identity;
+	}
+
+	FVector4f GetBoneColor(int32 BoneIdx) const
+	{
+		return FVector4f::One();
+	}
+
+	const TArray<int32>& GetNormalIDs() const 
+	{ 
+		return NormalIDs; 
+	}
+
+	FVector3f GetNormal(NormalIDType ID) const 
+	{ 
+		return (VertexDataType == EClothPatternVertexType::Render) ? Cloth.GetRenderNormal()[ID] : Cloth.GetSimNormal()[ID];
+	}
+
+	bool GetNormalTri(const TriIDType& TriID, NormalIDType& ID0, NormalIDType& ID1, NormalIDType& ID2) const 
+	{ 
+		if (VertexDataType == EClothPatternVertexType::Sim2D)
+		{
+			// All normal data lives on 3D indices.
+			const FIntVector Face = Cloth.GetSimIndices3D()[TriID];
+
+			ID0 = Face[0];
+			ID1 = Face[1];
+			ID2 = Face[2];
+
+			return true;
+		}
+		else
+		{
+		return GetTri(TriID, ID0, ID1, ID2); 
+	}
+	}
+	
+	const TArray<int32>& GetUVIDs(int32 LayerID) const 
+	{ 
+		return (VertexDataType == EClothPatternVertexType::Render) ? NormalIDs : EmptyArray;
+	}
+
+	FVector2f GetUV(int32 LayerID, UVIDType UVID) const 
+	{ 
+		checkf(VertexDataType == EClothPatternVertexType::Render, TEXT("Requested UVs from a Sim mesh"));
+		const TConstArrayView<FVector2f> VertexUVs = Cloth.GetRenderUVs()[UVID];
+		return VertexUVs.IsValidIndex(LayerID) ? VertexUVs[LayerID] : FVector2f(0.f);
+	}
+
+	bool GetUVTri(int32 LayerID, const TriIDType& TriID, UVIDType& ID0, UVIDType& ID1, UVIDType& ID2) const 
+	{ 
+		return GetTri(TriID, ID0, ID1, ID2); 
+	}
+	
+	const TArray<int32>& GetTangentIDs() const 
+	{ 
+		return (VertexDataType == EClothPatternVertexType::Render) ? NormalIDs : EmptyArray;
+	}
+
+	FVector3f GetTangent(NormalIDType ID) const 
+	{ 
+		checkf(VertexDataType == EClothPatternVertexType::Render, TEXT("Requested Tangent from a Sim mesh"));
+		return Cloth.GetRenderTangentU()[ID];
+	}
+
+	bool GetTangentTri(const TriIDType& TriID, NormalIDType& ID0, NormalIDType& ID1, NormalIDType& ID2) const 
+	{ 
+		return GetNormalTri(TriID, ID0, ID1, ID2);
+	}
+
+	const TArray<int32>& GetBiTangentIDs() const 
+	{ 
+		return (VertexDataType == EClothPatternVertexType::Render) ? NormalIDs : EmptyArray;
+	}
+
+	FVector3f GetBiTangent(NormalIDType ID) const 
+	{
+		checkf(VertexDataType == EClothPatternVertexType::Render, TEXT("Requested Bitangent from a Sim mesh"));
+		return Cloth.GetRenderTangentV()[ID];
+	}
+
+	bool GetBiTangentTri(const TriIDType& TriID, NormalIDType& ID0, NormalIDType& ID1, NormalIDType& ID2) const
+	{
+		return GetNormalTri(TriID, ID0, ID1, ID2);
+	}
+
+	const TArray<int32>& GetColorIDs() const 
+	{ 
+		return (VertexDataType == EClothPatternVertexType::Render) ? NormalIDs : EmptyArray;
+	}
+
+	FVector4f GetColor(ColorIDType VID) const 
+	{
+		checkf(VertexDataType == EClothPatternVertexType::Render, TEXT("Requested color from a Sim mesh"));
+		return Cloth.GetRenderColor()[VID];
+	}
+
+	bool GetColorTri(const TriIDType& TriID, ColorIDType& ID0, ColorIDType& ID1, ColorIDType& ID2) const 
+	{
+		return GetNormalTri(TriID, ID0, ID1, ID2);
+	}
+
+private:
+
+	const EClothPatternVertexType VertexDataType;
+	const FCollectionClothConstFacade& Cloth;
+
+	TArray<TriIDType> TriIDs;		// indices into Indices
+	TArray<VertIDType> VertIDs;		// indices into Positions2D or Positions3D
+	TArray<NormalIDType> NormalIDs; // indices into Normals (and all per vertex data other than positions)
+	int32 NumTexCoords;
+	
+	TArray<FName> WeightMapNames;
+
+	TArray<int32> EmptyArray;
+
+	const FReferenceSkeleton* RefSkeleton = nullptr;
+};
+
+
+void FClothPatternToDynamicMesh::Convert(
+	const TSharedRef<const FManagedArrayCollection> ClothCollection,
+	int32 PatternIndex,
+	EClothPatternVertexType VertexDataType,
+	UE::Geometry::FDynamicMesh3& MeshOut,
+	bool bDisableAttributes,
+	int32 MaterialOffset)
+{
+	const FCollectionClothConstFacade ClothFacade(ClothCollection);
+
+	// Actual conversion
+	UE::Geometry::TToDynamicMesh<FClothPatternWrapper> PatternToDynamicMesh;
+	FClothPatternWrapper PatternWrapper(ClothFacade, PatternIndex, VertexDataType);
+
+	auto TriangleToGroupFunction = [](FClothPatternWrapper::TriIDType) { return 0; };
+
+	if (bDisableAttributes)
+	{
+		PatternToDynamicMesh.ConvertWOAttributes(MeshOut, PatternWrapper, TriangleToGroupFunction);
+	}
+	else
+	{
+		constexpr bool bCopyTangents = false;
+		const bool bIsRenderType = VertexDataType == EClothPatternVertexType::Render;
+		auto TriangleToMaterialFunction = [PatternIndex, bIsRenderType, &ClothFacade](FClothPatternWrapper::TriIDType TriID)->int32
+		{
+			if (bIsRenderType)
+			{
+				if (PatternIndex != INDEX_NONE)
+				{
+					return PatternIndex;
+				}
+				const int32 FoundPattern = ClothFacade.FindRenderPatternByFaceIndex(TriID);
+				check(FoundPattern != INDEX_NONE);
+				return FoundPattern;
+			}
+
+			// Sim meshes will have a default material with MaterialID zero applied
+			return 0;
+		};
+
+		PatternToDynamicMesh.Convert(MeshOut, PatternWrapper, TriangleToGroupFunction, TriangleToMaterialFunction, bCopyTangents);
+	}
+
+	// Add non-manifold mapping data if DynamicMesh indices don't match orig indices
+	const int32 NumVertices = VertexDataType == EClothPatternVertexType::Render ? ClothFacade.GetNumRenderVertices() :
+		VertexDataType == EClothPatternVertexType::Sim2D ? ClothFacade.GetNumSimVertices2D() : ClothFacade.GetNumSimVertices3D();
+	bool bVertexIndicesMatch = false;
+	if (PatternToDynamicMesh.ToSrcVertIDMap.Num() == NumVertices)
+	{
+		bVertexIndicesMatch = true;
+		for (int32 VertId = 0; VertId < NumVertices; ++VertId)
+		{
+			if (PatternToDynamicMesh.ToSrcVertIDMap[VertId] != VertId)
+			{
+				bVertexIndicesMatch = false;
+				break;
+			}
+		}
+	}
+	if (!bVertexIndicesMatch)
+	{
+		MeshOut.EnableAttributes();
+		FClothPatternToDynamicMeshMappingSupport::AttachVertexMappingData(PatternToDynamicMesh.ToSrcVertIDMap, MeshOut);
+	}
+	const int32 NumTriangles = VertexDataType == EClothPatternVertexType::Render ? ClothFacade.GetNumRenderFaces() : ClothFacade.GetNumSimFaces();
+	bool bTriangleIndicesMatch = false;
+	if (PatternToDynamicMesh.ToSrcTriIDMap.Num() == NumTriangles)
+	{
+		bTriangleIndicesMatch = true;
+		for (int32 TriId = 0; TriId < NumTriangles; ++TriId)
+		{
+			if (PatternToDynamicMesh.ToSrcTriIDMap[TriId] != TriId)
+			{
+				bTriangleIndicesMatch = false;
+				break;
+			}
+		}
+	}
+	if (!bTriangleIndicesMatch)
+	{
+		MeshOut.EnableAttributes();
+		FClothPatternToDynamicMeshMappingSupport::AttachTriangleMappingData(PatternToDynamicMesh.ToSrcTriIDMap, MeshOut);
+	}
+
+	// Fix up any triangles without valid material IDs and apply material offset
+	if (VertexDataType == EClothPatternVertexType::Render && MaterialOffset != INDEX_NONE)
+	{
+		const TConstArrayView<FSoftObjectPath> MaterialPaths = ClothFacade.GetRenderMaterialSoftObjectPathName();
+		int32 DefaultMaterialID = 0;
+		for (const int32 TriID : MeshOut.TriangleIndicesItr())
+		{
+			int32 MaterialID;
+			MeshOut.Attributes()->GetMaterialID()->GetValue(TriID, &MaterialID);
+			//const int32 MaterialID = MeshOut.Attributes()->GetMaterialID()->GetValue(TriID);
+			if (!MaterialPaths.IsValidIndex(MaterialID))
+			{
+				MeshOut.Attributes()->GetMaterialID()->SetValue(TriID, DefaultMaterialID);
+			}
+			else
+			{
+				MeshOut.Attributes()->GetMaterialID()->SetValue(TriID, MaterialID + MaterialOffset);
+			}
+		}
+	}
+}
+
+void FClothPatternToDynamicMesh::Convert(
+	const UChaosClothAsset* ClothAssetMeshIn,
+	int32 LODIndex,
+	int32 PatternIndex,
+	EClothPatternVertexType VertexDataType,
+	UE::Geometry::FDynamicMesh3& MeshOut,
+	int32 MaterialOffset)
+{
+	const TArray<TSharedRef<const FManagedArrayCollection>>& ClothCollections = ClothAssetMeshIn->GetClothCollections();
+	check(ClothCollections.IsValidIndex(LODIndex));
+
+	constexpr bool bDisableAttributes = false;
+	Convert(ClothCollections[LODIndex], PatternIndex, VertexDataType, MeshOut, bDisableAttributes, MaterialOffset);
+}
+
+}	// namespace UE::Chaos::ClothAsset
+
+#else
+
+namespace UE::Chaos::ClothAsset
+{
+void FClothPatternToDynamicMesh::Convert(
+	const TSharedRef<const FManagedArrayCollection> ClothCollection,
+	int32 PatternIndex,
+	EClothPatternVertexType VertexDataType,
+	UE::Geometry::FDynamicMesh3& MeshOut,
+	bool bDisableAttributes,
+	int32 MaterialOffset)
+{
+	// Conversion only supported with editor.
+	check(0);
+}
+
+void FClothPatternToDynamicMesh::Convert(
+	const UChaosClothAsset* ClothAssetMeshIn,
+	FDynamicMesh3& MeshOut,
+	int32 LODIndex,
+	int32 PatternIndex,
+	int32 MaterialOffset)
+{
+	// Conversion only supported with editor.
+	check(0);
+}
+
+}	// namespace UE::Chaos::ClothAsset
+
+#endif  // end with editor
+
diff --git a/Engine/Plugins/ChaosClothAsset/Source/ChaosClothAssetTools/Private/ChaosClothAsset/ClothPatternToDynamicMeshMappingSupport.cpp b/Engine/Plugins/ChaosClothAsset/Source/ChaosClothAssetTools/Private/ChaosClothAsset/ClothPatternToDynamicMeshMappingSupport.cpp
new file mode 100644
--- /dev/null
+++ b/Engine/Plugins/ChaosClothAsset/Source/ChaosClothAssetTools/Private/ChaosClothAsset/ClothPatternToDynamicMeshMappingSupport.cpp
@@ -0,0 +1,376 @@
+// Copyright Epic Games, Inc. All Rights Reserved. 
+
+#include "ChaosClothAsset/ClothPatternToDynamicMeshMappingSupport.h"
+#include "DynamicMesh/DynamicMeshAttributeSet.h"
+
+namespace UE::Chaos::ClothAsset
+{
+
+class FDynamicMeshSourceTriangleIdAttributeChange final :
+	public UE::Geometry::FDynamicMeshAttributeChangeBase
+{
+	friend class FDynamicMeshSourceTriangleIdAttribute;
+	TArray<TPair<int32, int32>> OldSourceTriangleIds, NewSourceTriangleIds;
+
+	bool bOldInvalidState = false;
+	bool bNewInvalidState = false;
+
+public:
+	FDynamicMeshSourceTriangleIdAttributeChange() = default;
+
+	virtual ~FDynamicMeshSourceTriangleIdAttributeChange() override = default;
+
+	void SaveInitialTriangle(const UE::Geometry::FDynamicMeshAttributeBase* Attribute, int TriangleID) override;
+
+	void StoreAllFinalTriangles(const UE::Geometry::FDynamicMeshAttributeBase* Attribute, const TArray<int>& TriangleIDs) override;
+
+	bool Apply(UE::Geometry::FDynamicMeshAttributeBase* Attribute, bool bRevert) const override;
+};
+
+// TODO: Can this class inherit from UE::Geometry::TDynamicMeshScalarTriangleAttribute<int32>
+// and then only need to implement a smaller subset of the methods? (On* and IsValid?)
+class FDynamicMeshSourceTriangleIdAttribute final :
+	public UE::Geometry::FDynamicMeshAttributeBase
+{
+	friend class FDynamicMeshSourceTriangleIdAttributeChange;
+
+	UE::Geometry::FDynamicMesh3* Parent = nullptr;
+	UE::Geometry::TDynamicVector<int32> SourceTriangleIds;
+	bool bValid = false;
+
+public:
+	FDynamicMeshSourceTriangleIdAttribute() = default;
+
+	FDynamicMeshSourceTriangleIdAttribute(UE::Geometry::FDynamicMesh3* InParent) : Parent(InParent)
+	{
+	}
+
+	virtual ~FDynamicMeshSourceTriangleIdAttribute() override = default;
+
+	const UE::Geometry::FDynamicMesh3* GetParent() const
+	{
+		return Parent;
+	}
+	UE::Geometry::FDynamicMesh3* GetParent()
+	{
+		return Parent;
+	}
+
+	bool IsValid() const
+	{
+		return bValid;
+	}
+
+	UE::Geometry::FDynamicMeshAttributeBase* MakeCopy(UE::Geometry::FDynamicMesh3* ParentIn) const override
+	{
+		FDynamicMeshSourceTriangleIdAttribute* Attribute = new FDynamicMeshSourceTriangleIdAttribute(ParentIn);
+		Attribute->Copy(*this);
+		return Attribute;
+	}
+
+	UE::Geometry::FDynamicMeshAttributeBase* MakeNew(UE::Geometry::FDynamicMesh3* ParentIn) const override
+	{
+		FDynamicMeshSourceTriangleIdAttribute* Attribute = new FDynamicMeshSourceTriangleIdAttribute(ParentIn);
+		Attribute->Initialize();
+		return Attribute;
+	}
+
+	void CompactInPlace(const UE::Geometry::FCompactMaps& CompactMaps) override
+	{
+		for (int32 FromTID = 0, NumTID = CompactMaps.NumTriangleMappings(); FromTID < NumTID; FromTID++)
+		{
+			const int32 ToTID = CompactMaps.GetTriangleMapping(FromTID);
+			if (ToTID == UE::Geometry::FCompactMaps::InvalidID)
+			{
+				continue;
+			}
+			if (ensure(ToTID <= FromTID))
+			{
+				SourceTriangleIds[ToTID] = SourceTriangleIds[FromTID];
+			}
+		}
+		SourceTriangleIds.Resize(Parent->MaxTriangleID());
+	}
+
+	void Reparent(UE::Geometry::FDynamicMesh3* NewParent) override
+	{
+		Parent = NewParent;
+	}
+
+	bool CopyThroughMapping(const UE::Geometry::TDynamicAttributeBase<UE::Geometry::FDynamicMesh3>* Source, const UE::Geometry::FMeshIndexMappings& Mapping) override
+	{
+		for (const TPair<int32, int32>& MapTID : Mapping.GetTriangleMap().GetForwardMap())
+		{
+			int32 SourceTId;
+			if (!ensure(Source->CopyOut(MapTID.Key, &SourceTId, sizeof(SourceTId))))
+			{
+				return false;
+			}
+			SetValue(MapTID.Value, SourceTId);
+		}
+		return true;
+	}
+
+	virtual bool Append(const TDynamicAttributeBase& Source, const UE::Geometry::FDynamicMesh3::FAppendInfo& Info) override
+	{
+		int32 NewMaxID = Info.NumTriangle + Info.TriangleOffset;
+		if (NewMaxID > SourceTriangleIds.Num())
+		{
+			SourceTriangleIds.SetNum(NewMaxID);
+		}
+		for (int32 Idx = 0; Idx < Info.NumTriangle; ++Idx)
+		{
+			int32 TargetID = Idx + Info.TriangleOffset;
+			if (!Parent->IsTriangle(TargetID))
+			{
+				continue;
+			}
+			int32 SourceTID;
+			if (!ensure(Source.CopyOut(Idx, &SourceTID, sizeof(SourceTID))))
+			{
+				return false;
+			}
+			SetValue(Idx + Info.TriangleOffset, SourceTID);
+		}
+		return true;
+	}
+
+	virtual void AppendDefaulted(const UE::Geometry::FDynamicMesh3::FAppendInfo& Info) override
+	{
+		int32 NewMaxID = Info.NumTriangle + Info.TriangleOffset;
+		if (NewMaxID > SourceTriangleIds.Num())
+		{
+			SourceTriangleIds.Resize(NewMaxID, INDEX_NONE);
+		}
+	}
+
+	bool CopyOut(int RawID, void* Buffer, int BufferSize) const override
+	{
+		if (BufferSize != sizeof(int32) || !Parent->IsTriangle(RawID))
+		{
+			return false;
+		}
+
+		*static_cast<int32*>(Buffer) = SourceTriangleIds[RawID];
+		return true;
+	}
+
+	bool CopyIn(int RawID, void* Buffer, int BufferSize) override
+	{
+		if (BufferSize != sizeof(int32) || !Parent->IsTriangle(RawID))
+		{
+			return false;
+		}
+
+		SourceTriangleIds[RawID] = *static_cast<int32*>(Buffer);
+		return true;
+	}
+
+	TUniquePtr<UE::Geometry::TDynamicAttributeChangeBase<UE::Geometry::FDynamicMesh3>> NewBlankChange() const override
+	{
+		return MakeUnique<FDynamicMeshSourceTriangleIdAttributeChange>();
+	}
+
+	void Initialize()
+	{
+		SourceTriangleIds.Resize(Parent->MaxTriangleID());
+		SourceTriangleIds.Fill(INDEX_NONE);
+	}
+
+	void InitializeFromArray(TConstArrayView<int32> TriangleToNonManifoldVertexIDMap)
+	{
+		check(TriangleToNonManifoldVertexIDMap.Num() == Parent->MaxTriangleID());
+		SourceTriangleIds.Resize(Parent->MaxTriangleID());
+
+		for (int32 TriangleID : Parent->TriangleIndicesItr())
+		{
+			SetValue(TriangleID, TriangleToNonManifoldVertexIDMap[TriangleID]);
+		}
+		bValid = true;
+	}
+
+	void Copy(const FDynamicMeshSourceTriangleIdAttribute& Copy)
+	{
+		UE::Geometry::FDynamicMeshAttributeBase::CopyParentClassData(Copy);
+		SourceTriangleIds = Copy.SourceTriangleIds;
+		bValid = Copy.bValid;
+	}
+
+	int32 GetValue(int32 InTriangleId) const
+	{
+		return SourceTriangleIds[InTriangleId];
+	}
+
+	void SetValue(int32 InTriangleId, int32 InSourceTriangled)
+	{
+		SourceTriangleIds[InTriangleId] = InSourceTriangled;
+	}
+
+	// Any topo operation on the mesh will invalidate the non-manifold information.
+	virtual void OnSplitEdge(const DynamicMeshInfo::FEdgeSplitInfo&) override
+	{
+		bValid = false;
+	}
+
+	virtual void OnFlipEdge(const DynamicMeshInfo::FEdgeFlipInfo& FlipInfo) override
+	{
+		bValid = false;
+	}
+
+	virtual void OnCollapseEdge(const DynamicMeshInfo::FEdgeCollapseInfo&) override
+	{
+		bValid = false;
+	}
+
+	virtual void OnPokeTriangle(const DynamicMeshInfo::FPokeTriangleInfo&) override
+	{
+		bValid = false;
+	}
+
+	virtual void OnMergeEdges(const DynamicMeshInfo::FMergeEdgesInfo&) override
+	{
+		bValid = false;
+	}
+
+	virtual void OnMergeVertices(const DynamicMeshInfo::FMergeVerticesInfo&) override
+	{
+		bValid = false;
+	}
+
+	virtual void OnSplitVertex(const DynamicMeshInfo::FVertexSplitInfo&, const TArrayView<const int>&) override
+	{
+		bValid = false;
+	}
+
+	virtual void OnNewVertex(int VertexID, bool bInserted) override
+	{
+		bValid = false;
+	}
+
+	virtual void OnRemoveVertex(int VertexID) override
+	{
+		bValid = false;
+	}
+};
+
+
+void FDynamicMeshSourceTriangleIdAttributeChange::SaveInitialTriangle(const UE::Geometry::FDynamicMeshAttributeBase* Attribute, int TriangleID)
+{
+	const FDynamicMeshSourceTriangleIdAttribute* SourceAttribute = static_cast<const FDynamicMeshSourceTriangleIdAttribute*>(Attribute);
+	if (OldSourceTriangleIds.IsEmpty())
+	{
+		bOldInvalidState = SourceAttribute->IsValid();
+	}
+
+	OldSourceTriangleIds.Emplace(TriangleID, SourceAttribute->GetValue(TriangleID));
+}
+
+void FDynamicMeshSourceTriangleIdAttributeChange::StoreAllFinalTriangles(const UE::Geometry::FDynamicMeshAttributeBase* Attribute, const TArray<int>& TriangleIDs)
+{
+	NewSourceTriangleIds.Reserve(NewSourceTriangleIds.Num() + TriangleIDs.Num());
+	const FDynamicMeshSourceTriangleIdAttribute* SourceAttribute = static_cast<const FDynamicMeshSourceTriangleIdAttribute*>(Attribute);
+	for (int32 TriangleID : TriangleIDs)
+	{
+		NewSourceTriangleIds.Emplace(TriangleID, SourceAttribute->GetValue(TriangleID));
+	}
+
+	// Store the last known valid state.  
+	bNewInvalidState = SourceAttribute->IsValid();
+}
+
+bool FDynamicMeshSourceTriangleIdAttributeChange::Apply(UE::Geometry::FDynamicMeshAttributeBase* Attribute, bool bRevert) const
+{
+	const TArray<TPair<int32, int32>>& Changes = bRevert ? OldSourceTriangleIds : NewSourceTriangleIds;
+	FDynamicMeshSourceTriangleIdAttribute* SourceAttribute = static_cast<FDynamicMeshSourceTriangleIdAttribute*>(Attribute);
+	for (const TPair<int32, int32>& Item : Changes)
+	{
+		if (ensure(SourceAttribute->GetParent()->IsTriangle(Item.Key)))
+		{
+			SourceAttribute->SetValue(Item.Key, Item.Value);
+		}
+	}
+
+	// Restore the valid state as well.
+	SourceAttribute->bValid = bRevert ? bOldInvalidState : bNewInvalidState;
+
+	return true;
+}
+
+FClothPatternToDynamicMeshMappingSupport::FClothPatternToDynamicMeshMappingSupport(const UE::Geometry::FDynamicMesh3& Mesh)
+	: UE::Geometry::FNonManifoldMappingSupport(Mesh)
+{
+	Reset(Mesh);
+}
+
+const FName FClothPatternToDynamicMeshMappingSupport::ClothMeshTIDsAttrName = FName(TEXT("ClothMeshTIDsAttr"));
+
+void FClothPatternToDynamicMeshMappingSupport::Reset(const UE::Geometry::FDynamicMesh3& Mesh)
+{
+	UE::Geometry::FNonManifoldMappingSupport::Reset(Mesh);
+
+	if (const UE::Geometry::FDynamicMeshAttributeSet* Attributes = DynamicMesh->Attributes())
+	{
+		if (Attributes->HasAttachedAttribute(ClothMeshTIDsAttrName))
+		{
+			SrcTIDsAttribute = static_cast<const FDynamicMeshSourceTriangleIdAttribute*>(Attributes->GetAttachedAttribute(ClothMeshTIDsAttrName));
+		}
+	}
+}
+
+bool FClothPatternToDynamicMeshMappingSupport::IsMappedTriangleInSource() const
+{
+	return SrcTIDsAttribute != nullptr && SrcTIDsAttribute->IsValid();
+}
+
+int32 FClothPatternToDynamicMeshMappingSupport::GetOriginalTriangleID(const int32 Tid) const
+{
+	checkSlow(DynamicMesh->IsTriangle(Tid));
+
+	int32 Result = Tid;
+	if (SrcTIDsAttribute)
+	{
+		Result = SrcTIDsAttribute->GetValue(Tid);
+	}
+	return Result;
+}
+
+bool FClothPatternToDynamicMeshMappingSupport::AttachTriangleMappingData(const TArray<int32>& TriangleToOriginalVertexIDMap, UE::Geometry::FDynamicMesh3& InOutMesh)
+{
+	const int32 MaxTID = InOutMesh.MaxTriangleID();
+	UE::Geometry::FDynamicMeshAttributeSet* const Attributes = InOutMesh.Attributes();
+
+	if (TriangleToOriginalVertexIDMap.Num() < MaxTID || !Attributes)
+	{
+		return false;
+	}
+
+	// create and attach a triangle ID buffer, removing any pre-exising one.
+	FDynamicMeshSourceTriangleIdAttribute* TIDAttr = [&Attributes, &InOutMesh] {
+		// replace existing buffer.  we *should* be able to re-use it
+		if (Attributes->HasAttachedAttribute(ClothMeshTIDsAttrName))
+		{
+			Attributes->RemoveAttribute(ClothMeshTIDsAttrName);
+		}
+		FDynamicMeshSourceTriangleIdAttribute* SrcMeshTIDAttr = new FDynamicMeshSourceTriangleIdAttribute(&InOutMesh);
+		SrcMeshTIDAttr->SetName(ClothMeshTIDsAttrName);
+		InOutMesh.Attributes()->AttachAttribute(ClothMeshTIDsAttrName, SrcMeshTIDAttr);
+		return SrcMeshTIDAttr;
+		}();
+
+	TIDAttr->InitializeFromArray(TriangleToOriginalVertexIDMap);
+	return true;
+}
+
+void FClothPatternToDynamicMeshMappingSupport::RemoveTriangleMappingData(UE::Geometry::FDynamicMesh3& InOutMesh)
+{
+	UE::Geometry::FDynamicMeshAttributeSet* Attributes = InOutMesh.Attributes();
+
+	if (Attributes && Attributes->HasAttachedAttribute(ClothMeshTIDsAttrName))
+	{
+		Attributes->RemoveAttribute(ClothMeshTIDsAttrName);
+	}
+}
+
+
+}	// namespace UE::Chaos::ClothAsset
+
diff --git a/Engine/Plugins/ChaosClothAsset/Source/ChaosClothAssetTools/Private/ChaosClothAsset/ClothingAssetToClothAssetExporter.cpp b/Engine/Plugins/ChaosClothAsset/Source/ChaosClothAssetTools/Private/ChaosClothAsset/ClothingAssetToClothAssetExporter.cpp
new file mode 100644
--- /dev/null
+++ b/Engine/Plugins/ChaosClothAsset/Source/ChaosClothAssetTools/Private/ChaosClothAsset/ClothingAssetToClothAssetExporter.cpp
@@ -0,0 +1,165 @@
+// Copyright Epic Games, Inc. All Rights Reserved.
+
+#include "ClothingAssetToClothAssetExporter.h"
+#include "ChaosClothAsset/ClothAsset.h"
+#include "ChaosClothAsset/CollectionClothFacade.h"
+#include "ChaosClothAsset/ClothGeometryTools.h"
+#include "Animation/Skeleton.h"
+#include "DynamicMesh/DynamicMesh3.h"
+#include "DynamicMesh/DynamicMeshAttributeSet.h"
+#include "DynamicMesh/NonManifoldMappingSupport.h"
+#include "ToDynamicMesh.h"
+#include "Engine/SkinnedAssetCommon.h"
+#include "Engine/SkeletalMesh.h"
+#include "GeometryCollection/ManagedArrayCollection.h"
+#include "Materials/Material.h"
+#include "Misc/MessageDialog.h"
+#include "PhysicsEngine/PhysicsAsset.h"
+#include "ClothingAsset.h"
+
+#include UE_INLINE_GENERATED_CPP_BY_NAME(ClothingAssetToClothAssetExporter)
+
+#define LOCTEXT_NAMESPACE "ClothingAssetToClothAssetExporter"
+
+namespace UE::Chaos::ClothAsset::Private
+{
+struct FSimpleSrcMeshInterface
+{
+	typedef int32     VertIDType;
+	typedef int32     TriIDType;
+
+	FSimpleSrcMeshInterface(const TArray<FVector3f>& InPositions, const TArray<uint32>& InIndices)
+		:Positions(InPositions), Indices(InIndices)
+	{
+		VertIDs.SetNumUninitialized(Positions.Num());
+		for (int32 VtxIndex = 0; VtxIndex < Positions.Num(); ++VtxIndex)
+		{
+			VertIDs[VtxIndex] = VtxIndex;
+		}
+
+		check(Indices.Num() % 3 == 0);
+		const int32 NumFaces = Indices.Num() / 3;
+		TriIDs.SetNumUninitialized(NumFaces);
+		for (int32 TriIndex = 0; TriIndex < NumFaces; ++TriIndex)
+		{
+			TriIDs[TriIndex] = 3 * TriIndex;
+		}
+	}
+
+	// accounting.
+	int32 NumTris() const { return TriIDs.Num(); }
+	int32 NumVerts() const { return VertIDs.Num(); }
+
+	// --"Vertex Buffer" info
+	const TArray<VertIDType>& GetVertIDs() const { return VertIDs; }
+	const FVector GetPosition(const VertIDType VtxID) const { return FVector(Positions[VtxID]); }
+
+	// --"Index Buffer" info
+	const TArray<TriIDType>& GetTriIDs() const { return TriIDs; }
+	// return false if this TriID is not contained in mesh.
+	bool GetTri(const TriIDType TriID, VertIDType& VID0, VertIDType& VID1, VertIDType& VID2) const
+	{
+		VID0 = Indices[TriID + 0];
+		VID1 = Indices[TriID + 1];
+		VID2 = Indices[TriID + 2];
+
+		return true;
+	}
+
+private:
+	const TArray<FVector3f>& Positions;
+	const TArray<uint32>& Indices;
+
+	TArray<TriIDType> TriIDs; // TriID = first index in flat Indices array
+	TArray<VertIDType> VertIDs;
+};
+} // namespace UE::Chaos::ClothAsset::Private
+
+
+UClass* UClothingAssetToChaosClothAssetExporter::GetExportedType() const
+{
+	return UChaosClothAsset::StaticClass();
+}
+
+void UClothingAssetToChaosClothAssetExporter::Export(const UClothingAssetBase* ClothingAsset, UObject* ExportedAsset)
+{
+	using namespace UE::Chaos::ClothAsset;
+
+	const UClothingAssetCommon* const ClothingAssetCommon = ExactCast<UClothingAssetCommon>(ClothingAsset);
+	if (!ClothingAssetCommon)
+	{
+		const FText TitleMessage = LOCTEXT("ClothingAssetExporterTitle", "Error Exporting Clothing Asset");
+		const FText ErrorMessage = LOCTEXT("ClothingAssetExporterError", "Can only export from known ClothingAssetCommon types.");
+		FMessageDialog::Open(EAppMsgType::Ok, EAppReturnType::Ok, ErrorMessage, TitleMessage);
+		return;
+	}
+
+	UChaosClothAsset* const ClothAsset = CastChecked<UChaosClothAsset>(ExportedAsset);
+	check(ClothAsset);
+
+	static const FString DefaultMaterialPathName = FString(TEXT("/Engine/EditorMaterials/Cloth/CameraLitDoubleSided.CameraLitDoubleSided"));
+	const FString PhysicsAssetPathName = ClothingAssetCommon->PhysicsAsset ? ClothingAssetCommon->PhysicsAsset->GetPathName() : FString();
+	const FString SkeletalMeshPathName = CastChecked<USkeletalMesh>(ClothingAssetCommon->GetOuter())->GetPathName();
+
+	const int32 NumLods = ClothingAssetCommon->LodData.Num();
+	TArray<TSharedRef<const FManagedArrayCollection>> ClothCollections;
+	ClothCollections.Reset(NumLods);
+
+	// Create the LODs
+	for (int32 LodIndex = 0; LodIndex < NumLods; ++LodIndex)
+	{
+		const FClothLODDataCommon& ClothLODData = ClothingAssetCommon->LodData[LodIndex];
+		const FClothPhysicalMeshData& PhysicalMeshData = ClothLODData.PhysicalMeshData;
+
+		TSharedRef<FManagedArrayCollection> ClothCollection = MakeShared<FManagedArrayCollection>();
+		FCollectionClothFacade ClothFacade(ClothCollection);
+		ClothFacade.DefineSchema();
+
+		using namespace UE::Chaos::ClothAsset::Private;
+
+		// Build a DynamicMesh from Positions and Indices
+		UE::Geometry::TToDynamicMeshBase<FSimpleSrcMeshInterface> ToDynamicMesh;
+		FSimpleSrcMeshInterface SimpleSrc(PhysicalMeshData.Vertices, PhysicalMeshData.Indices);
+
+		UE::Geometry::FDynamicMesh3 DynamicMesh;
+		ToDynamicMesh.Convert(DynamicMesh, SimpleSrc, [](FSimpleSrcMeshInterface::TriIDType) {return 0; });
+		UE::Geometry::FNonManifoldMappingSupport::AttachNonManifoldVertexMappingData(ToDynamicMesh.ToSrcVertIDMap, DynamicMesh);
+
+		constexpr int32 UVChannelIndexNone = INDEX_NONE;
+		constexpr bool bAppend = false;
+		FClothGeometryTools::BuildSimMeshFromDynamicMesh(ClothCollection, DynamicMesh, UVChannelIndexNone, FVector2f(1.f), bAppend);
+
+		// Set the physics asset if any
+		ClothFacade.SetPhysicsAssetSoftObjectPathName(PhysicsAssetPathName);
+		// Set the skeleton
+		ClothFacade.SetSkeletalMeshSoftObjectPathName(SkeletalMeshPathName);
+
+		// Bind the sim mesh to root bone
+		constexpr bool bBindSimMesh = true;
+		constexpr bool bBindRenderMesh = false;  // Will get bound below in CopySimMeshToRenderMesh
+		FClothGeometryTools::BindMeshToRootBone(ClothCollection, bBindSimMesh, bBindRenderMesh);
+
+		// Set the render mesh to duplicate the sim mesh
+		constexpr bool bSingleRenderPattern = true;
+		FClothGeometryTools::CopySimMeshToRenderMesh(ClothCollection, DefaultMaterialPathName, bSingleRenderPattern);
+
+		ClothCollections.Emplace(MoveTemp(ClothCollection));
+	}
+
+	if (!NumLods)
+	{
+		// Make sure that at least one empty LOD is always created
+		TSharedRef<FManagedArrayCollection> ClothCollection = MakeShared<FManagedArrayCollection>();
+		FCollectionClothFacade ClothFacade(ClothCollection);
+		ClothFacade.DefineSchema();
+		ClothFacade.SetPhysicsAssetSoftObjectPathName(PhysicsAssetPathName);
+		ClothFacade.SetSkeletalMeshSoftObjectPathName(SkeletalMeshPathName);
+		ClothCollections.Emplace(MoveTemp(ClothCollection));
+	}
+
+	// Rebuild the asset
+	ClothAsset->Build(ClothCollections);
+}
+
+#undef LOCTEXT_NAMESPACE
+
diff --git a/Engine/Plugins/ChaosClothAsset/Source/ChaosClothAssetTools/Private/ChaosClothAsset/ClothingAssetToClothAssetExporter.h b/Engine/Plugins/ChaosClothAsset/Source/ChaosClothAssetTools/Private/ChaosClothAsset/ClothingAssetToClothAssetExporter.h
new file mode 100644
--- /dev/null
+++ b/Engine/Plugins/ChaosClothAsset/Source/ChaosClothAssetTools/Private/ChaosClothAsset/ClothingAssetToClothAssetExporter.h
@@ -0,0 +1,26 @@
+// Copyright Epic Games, Inc. All Rights Reserved.
+
+#pragma once
+
+#include "ClothingAssetExporter.h"
+
+#include "ClothingAssetToClothAssetExporter.generated.h"
+
+/**
+ * Implement a clothing asset exporter as a way to migrate clothing assets to the new Chaos Cloth asset format.
+ */
+UCLASS()
+class UClothingAssetToChaosClothAssetExporter : public UClothingAssetExporter
+{
+    GENERATED_BODY()
+public:
+	/** Return the class of the exported asset type. */
+	virtual UClass* GetExportedType() const override;
+
+	/**
+	 * Export the specified asset.
+	 * \param ClothingAsset the source clothing asset to be exported.
+	 * \param ExportedAsset the destination asset object in the exported type provided by the caller, ready to be filled by the Export function.
+	 */
+	virtual void Export(const UClothingAssetBase* ClothingAsset, UObject* ExportedAsset) override;
+};
diff --git a/Engine/Plugins/ChaosClothAsset/Source/ChaosClothAssetTools/Private/ChaosClothAsset/SkeletalMeshConverter.cpp b/Engine/Plugins/ChaosClothAsset/Source/ChaosClothAssetTools/Private/ChaosClothAsset/SkeletalMeshConverter.cpp
new file mode 100644
--- /dev/null
+++ b/Engine/Plugins/ChaosClothAsset/Source/ChaosClothAssetTools/Private/ChaosClothAsset/SkeletalMeshConverter.cpp
@@ -0,0 +1,138 @@
+// Copyright Epic Games, Inc. All Rights Reserved.
+
+#include "ChaosClothAsset/SkeletalMeshConverter.h"
+#include "ChaosClothAsset/ClothAssetBase.h"
+#include "ChaosClothAsset/ClothPatternToDynamicMesh.h"
+#include "ChaosClothAsset/CollectionClothFacade.h"
+#include "DynamicMesh/DynamicMesh3.h"
+#include "DynamicMesh/DynamicMeshAttributeSet.h"
+#include "Engine/SkeletalMesh.h"
+#include "Materials/Material.h"
+#include "Rendering/SkeletalMeshModel.h"
+#include "DynamicMeshToMeshDescription.h"
+#include "MaterialDomain.h"
+#include "SkeletalMeshAttributes.h"
+#include "StaticToSkeletalMeshConverter.h"
+
+#include UE_INLINE_GENERATED_CPP_BY_NAME(SkeletalMeshConverter)
+
+bool UClothAssetEditorSkeletalMeshConverter::ExportToSkeletalMesh(const UChaosClothAssetBase& ClothAssetBase, USkeletalMesh& SkeletalMesh) const
+{
+	using namespace UE::Chaos::ClothAsset;
+
+	// Create one dynamic mesh from the multiple collections
+	TArray<UMaterialInterface*> RenderMaterials;
+	TArray<FDynamicMesh3> DynamicMeshes;
+
+	const int32 NumLods = ClothAssetBase.GetLODNum();
+	DynamicMeshes.SetNum(NumLods);
+
+	int32 MaterialOffset = 0;
+
+	const int32 NumModels = ClothAssetBase.GetNumClothSimulationModels();
+
+	for (int32 ModelIndex = 0; ModelIndex < NumModels; ++ModelIndex)
+	{
+		const TArray<TSharedRef<const FManagedArrayCollection>>& ClothCollections = ClothAssetBase.GetCollections(ModelIndex);
+
+		for (int32 LodIndex = 0; LodIndex < NumLods; ++LodIndex)
+		{
+			if (!ClothCollections.IsValidIndex(LodIndex))
+			{
+				continue;
+			}
+			const TSharedRef<const FManagedArrayCollection>& ClothCollection = ClothCollections[LodIndex];
+			FCollectionClothConstFacade ClothFacade(ClothCollection);
+
+			if (ClothFacade.IsValid())
+			{
+				FDynamicMesh3 DynamicMeshPiece;
+
+				constexpr bool bDisableAttributes = false;
+				const int32 PatternIndex = INDEX_NONE;
+				DynamicMeshPiece.EnableAttributes();
+				DynamicMeshPiece.Attributes()->EnableMaterialID();
+				FClothPatternToDynamicMesh ClothPatternToDynamicMesh;
+				ClothPatternToDynamicMesh.Convert(
+					ClothCollection,
+					PatternIndex,
+					EClothPatternVertexType::Render,
+					DynamicMeshPiece,
+					bDisableAttributes,
+					MaterialOffset);
+
+				if (!DynamicMeshes[LodIndex].VertexCount() || !DynamicMeshes[LodIndex].TriangleCount())
+				{
+					DynamicMeshes[LodIndex] = MoveTemp(DynamicMeshPiece);
+				}
+				else
+				{
+					DynamicMeshes[LodIndex].AppendWithOffsets(DynamicMeshPiece);
+				}
+
+				const TArrayView<const FSoftObjectPath> MaterialPaths = ClothFacade.GetRenderMaterialSoftObjectPathName();
+				RenderMaterials.Reserve(RenderMaterials.Num() + MaterialPaths.Num());
+				for (int32 MaterialIndex = 0; MaterialIndex < MaterialPaths.Num(); ++MaterialIndex)
+				{
+					RenderMaterials.Emplace(Cast<UMaterialInterface>(MaterialPaths[MaterialIndex].TryLoad()));
+				}
+				MaterialOffset += MaterialPaths.Num();
+			}
+		}
+	}
+
+	// Create the skeletal mesh from the dynamic mesh
+	TArray<const FMeshDescription*> MeshDescriptions;
+
+	SkeletalMesh.PreEditChange(nullptr);
+	SkeletalMesh.GetImportedModel()->LODModels.Empty();
+	SkeletalMesh.SetNumSourceModels(0);
+
+	TArray<FMeshDescription> LocalDescriptions;
+	for (const FDynamicMesh3& DynamicMesh : DynamicMeshes)
+	{
+		// Create a mesh description
+		FMeshDescription& MeshDescription = LocalDescriptions.AddDefaulted_GetRef();
+
+		// Add skeletal mesh attributes to the mesh description
+		FSkeletalMeshAttributes Attributes(MeshDescription);
+		Attributes.Register();
+
+		// Convert dynamic mesh to the mesh description
+		FConversionToMeshDescriptionOptions ConverterOptions;
+		FDynamicMeshToMeshDescription Converter(ConverterOptions);
+		constexpr bool bCopyTangents = true;
+		Converter.Convert(&DynamicMesh, MeshDescription, bCopyTangents);
+
+		// Add the created description to the list
+		MeshDescriptions.Add(&MeshDescription);
+	}
+
+	TArray<FSkeletalMaterial> SkeletalMaterials;
+	SkeletalMaterials.Reserve(RenderMaterials.Num());
+	for (UMaterialInterface* RenderMaterial : RenderMaterials)
+	{
+		if (RenderMaterial)
+		{
+			SkeletalMaterials.Emplace(RenderMaterial, RenderMaterial->GetFName());
+		}
+		else
+		{
+			SkeletalMaterials.Emplace(UMaterial::GetDefaultMaterial(MD_Surface));
+		}
+	}
+	if (!SkeletalMaterials.Num())
+	{
+		SkeletalMaterials.Emplace(UMaterial::GetDefaultMaterial(MD_Surface));
+	}
+
+	constexpr bool bRecomputeNormals = false;
+	constexpr bool bRecomputeTangents = false;
+	return FStaticToSkeletalMeshConverter::InitializeSkeletalMeshFromMeshDescriptions(
+		&SkeletalMesh,
+		MeshDescriptions,
+		SkeletalMaterials,
+		ClothAssetBase.GetRefSkeleton(),
+		bRecomputeNormals,
+		bRecomputeTangents);
+}
diff --git a/Engine/Plugins/ChaosClothAsset/Source/ChaosClothAssetTools/Private/ChaosClothAsset/SkeletalMeshConverter.h b/Engine/Plugins/ChaosClothAsset/Source/ChaosClothAssetTools/Private/ChaosClothAsset/SkeletalMeshConverter.h
new file mode 100644
--- /dev/null
+++ b/Engine/Plugins/ChaosClothAsset/Source/ChaosClothAssetTools/Private/ChaosClothAsset/SkeletalMeshConverter.h
@@ -0,0 +1,16 @@
+// Copyright Epic Games, Inc. All Rights Reserved.
+
+#pragma once
+
+#include "ChaosClothAsset/SkeletalMeshConverterClassProvider.h"
+#include "SkeletalMeshConverter.generated.h"
+
+UCLASS()
+class UClothAssetEditorSkeletalMeshConverter : public UClothAssetSkeletalMeshConverter
+{
+	GENERATED_BODY()
+
+public:
+	/** Build a skeletal mesh from the specified Cloth Asset. */
+	virtual bool ExportToSkeletalMesh(const class UChaosClothAssetBase& ClothAssetBase, class USkeletalMesh& SkeletalMesh) const override;
+};
diff --git a/Engine/Plugins/ChaosClothAsset/Source/ChaosClothAssetTools/Public/ChaosClothAsset/ClothAssetEditorStyle.h b/Engine/Plugins/ChaosClothAsset/Source/ChaosClothAssetTools/Public/ChaosClothAsset/ClothAssetEditorStyle.h
new file mode 100644
--- /dev/null
+++ b/Engine/Plugins/ChaosClothAsset/Source/ChaosClothAssetTools/Public/ChaosClothAsset/ClothAssetEditorStyle.h
@@ -0,0 +1,40 @@
+// Copyright Epic Games, Inc. All Rights Reserved.
+
+#pragma once
+
+#include "Styling/SlateStyle.h"
+
+#define UE_API CHAOSCLOTHASSETTOOLS_API
+
+namespace UE::Chaos::ClothAsset
+{
+	/**
+	 * Editor style setting up the cloth asset icons in editor.
+	 */
+	class FClothAssetEditorStyle final : public FSlateStyleSet
+	{
+	public:
+		UE_API FClothAssetEditorStyle();
+		UE_API virtual ~FClothAssetEditorStyle() override;
+
+	public:
+		static FClothAssetEditorStyle& Get()
+		{
+			if (!Singleton.IsSet())
+			{
+				Singleton.Emplace();
+			}
+			return Singleton.GetValue();
+		}
+
+		static void Destroy()
+		{
+			Singleton.Reset();
+		}
+
+	private:
+		static UE_API TOptional<FClothAssetEditorStyle> Singleton;
+	};
+}  // End namespace UE::Chaos::ClothAsset
+
+#undef UE_API
diff --git a/Engine/Plugins/ChaosClothAsset/Source/ChaosClothAssetTools/Public/ChaosClothAsset/ClothAssetFactory.h b/Engine/Plugins/ChaosClothAsset/Source/ChaosClothAssetTools/Public/ChaosClothAsset/ClothAssetFactory.h
new file mode 100644
--- /dev/null
+++ b/Engine/Plugins/ChaosClothAsset/Source/ChaosClothAssetTools/Public/ChaosClothAsset/ClothAssetFactory.h
@@ -0,0 +1,32 @@
+// Copyright Epic Games, Inc. All Rights Reserved.
+
+#pragma once
+
+#include "CoreMinimal.h"
+#include "UObject/ObjectMacros.h"
+#include "Factories/Factory.h"
+
+#include "ClothAssetFactory.generated.h"
+
+#define UE_API CHAOSCLOTHASSETTOOLS_API
+
+/**
+ * Having a cloth factory allows the cloth asset to be created from the Editor's menus.
+ */
+UCLASS(MinimalAPI, Experimental)
+class UChaosClothAssetFactory : public UFactory
+{
+	GENERATED_BODY()
+public:
+	UE_API UChaosClothAssetFactory(const FObjectInitializer& ObjectInitializer);
+
+	/** UFactory Interface */
+	virtual bool CanCreateNew() const override { return true; }
+	virtual bool FactoryCanImport(const FString& Filename) override { return false; }
+	virtual bool ShouldShowInNewMenu() const override { return true; }
+	UE_API virtual UObject* FactoryCreateNew(UClass* InClass, UObject* InParent, FName InName, EObjectFlags Flags, UObject* Context, FFeedbackContext* Warn) override;
+	UE_API virtual FString GetDefaultNewAssetName() const override;
+	/** End UFactory Interface */
+};
+
+#undef UE_API
diff --git a/Engine/Plugins/ChaosClothAsset/Source/ChaosClothAssetTools/Public/ChaosClothAsset/ClothPatternToDynamicMesh.h b/Engine/Plugins/ChaosClothAsset/Source/ChaosClothAssetTools/Public/ChaosClothAsset/ClothPatternToDynamicMesh.h
new file mode 100644
--- /dev/null
+++ b/Engine/Plugins/ChaosClothAsset/Source/ChaosClothAssetTools/Public/ChaosClothAsset/ClothPatternToDynamicMesh.h
@@ -0,0 +1,49 @@
+// Copyright Epic Games, Inc. All Rights Reserved.
+
+#pragma once
+
+#include "ClothPatternVertexType.h"
+#include "Templates/SharedPointer.h"
+
+#define UE_API CHAOSCLOTHASSETTOOLS_API
+
+
+namespace UE::Geometry { class FDynamicMesh3; }
+class UChaosClothAsset;
+struct FManagedArrayCollection;
+
+namespace UE::Chaos::ClothAsset
+{
+
+/**
+* Convert a single pattern from a ClothAsset to a FDynamicMesh3. 
+* When PatternIndex = INDEX_NONE, convert the entire cloth asset to an FDynamicMesh3.
+* When VertexDataType = EClothPatternVertexType::Sim3D, this generates a welded mesh (with native welded indexing) when PatternIndex = INDEX_NONE,
+*	and an unwelded (with native 2D unwelded indexing) when PatternIndex is a valid SimPattern Index.
+* Use MaterialOffset to specify an offset into the material group ID, and run a Fix up on any triangles without valid material IDs.
+* When MaterialOffset is INDEX_NONE, no offset is applied and nbo fix up is run. 
+*/
+class FClothPatternToDynamicMesh
+{
+public:
+
+	UE_API void Convert(
+		const TSharedRef<const FManagedArrayCollection> ClothCollection,
+		int32 PatternIndex,
+		EClothPatternVertexType VertexDataType,
+		UE::Geometry::FDynamicMesh3& MeshOut,
+		bool bDisableAttributes = false,
+		int32 MaterialOffset = INDEX_NONE);
+
+	UE_API void Convert(
+		const UChaosClothAsset* ClothAssetMeshIn,
+		int32 LODIndex,
+		int32 PatternIndex,
+		EClothPatternVertexType VertexDataType,
+		UE::Geometry::FDynamicMesh3& MeshOut,
+		int32 MaterialOffset = INDEX_NONE);
+};
+
+}	// namespace UE::Chaos::ClothAsset
+
+#undef UE_API
diff --git a/Engine/Plugins/ChaosClothAsset/Source/ChaosClothAssetTools/Public/ChaosClothAsset/ClothPatternToDynamicMeshMappingSupport.h b/Engine/Plugins/ChaosClothAsset/Source/ChaosClothAssetTools/Public/ChaosClothAsset/ClothPatternToDynamicMeshMappingSupport.h
new file mode 100644
--- /dev/null
+++ b/Engine/Plugins/ChaosClothAsset/Source/ChaosClothAssetTools/Public/ChaosClothAsset/ClothPatternToDynamicMeshMappingSupport.h
@@ -0,0 +1,143 @@
+// Copyright Epic Games, Inc. All Rights Reserved.
+
+#pragma once
+
+#include "Templates/SharedPointer.h"
+#include "DynamicMesh/NonManifoldMappingSupport.h"
+
+
+namespace UE::Geometry { class FDynamicMesh3; }
+
+namespace UE::Chaos::ClothAsset
+{
+
+class FDynamicMeshSourceTriangleIdAttribute;
+
+// This uses the NonManifoldMapping attribute for vertices, and a new attribute for triangles.
+class FClothPatternToDynamicMeshMappingSupport : public UE::Geometry::FNonManifoldMappingSupport
+{
+public:
+	CHAOSCLOTHASSETTOOLS_API FClothPatternToDynamicMeshMappingSupport(const UE::Geometry::FDynamicMesh3& Mesh);
+
+	/**
+	* Update the support for a new DynamicMesh.
+	*/
+	CHAOSCLOTHASSETTOOLS_API void Reset(const UE::Geometry::FDynamicMesh3& Mesh);
+
+	/*
+	* Return true if attribute data indicates that the source data that was converted to this DynamicMesh contained mapped vertices.
+	*/
+	bool IsMappedVertexInSource() const
+	{
+		return IsNonManifoldVertexInSource();
+	}
+
+	/*
+	* Return true if attribute data indicates that the source data that was converted to this DynamicMesh contained mapped triangles.
+	*/
+	CHAOSCLOTHASSETTOOLS_API bool IsMappedTriangleInSource() const;
+
+	/*
+	* Return true if the provided DynamicMesh vertex id resulted from a remapped vertex in the source data.
+	* @param vid - the id of a vertex in the DynamicMesh.
+	*
+	* Note: the code assumes but does not check that vid is a valid vertex id
+	*/
+	bool IsMappedVertexID(const int32 vid) const
+	{
+		return !(vid == GetOriginalNonManifoldVertexID(vid));
+	}
+
+	/*
+	* Return true if the provided DynamicMesh triangle id resulted from a remapped triangle in the source data.
+	* @param Tid - the id of a triangle in the DynamicMesh.
+	*
+	* Note: the code assumes but does not check that Tid is a valid triangle id
+	*/
+	bool IsMappedTriangleID(const int32 Tid) const
+	{
+		return !(Tid == GetOriginalTriangleID(Tid));
+	}
+
+	/*
+	* Return the vertex ID in the original cloth data used to generate this DynamicMesh associated with the provided vertex id.
+	* In the case that the source data was actually manifold the returned vertex id will be identical to the DynamicMesh vertex id.
+	* @param vid - the id of a vertex in the DynamicMesh.
+	*
+	* Note: the code assumes but does not check that vid is a valid vertex vid.
+	*/
+	int32 GetOriginalVertexID(const int32 vid) const
+	{
+		return GetOriginalNonManifoldVertexID(vid);
+	}
+
+	/*
+	* Return the triangle ID in the original cloth data used to generate this DynamicMesh associated with the provided triangle id.
+	* In the case that the source data was actually manifold the returned triangle id will be identical to the DynamicMesh triangle id.
+	* @param Tid - the id of a triangle in the DynamicMesh.
+	*
+	* Note: the code assumes but does not check that Tid is a valid triangle id
+	*/
+	CHAOSCLOTHASSETTOOLS_API int32 GetOriginalTriangleID(const int32 Tid) const;
+
+
+	// --- helper functions.
+
+	/*
+	* Attaches or replaces vertex mapping data to the provided mesh.
+	* @param VertexToOriginalVertexIDMap - an array that maps each DynamicMesh vertex id to the associated original vertex id.
+	* @return false on failure (no attribute will be attached to the DynamicMesh in this case)
+	*
+	* Note: Failure occurs if the DynamicMesh does not have attributes enabled or if the provided array is not long enough to provide a mapping value for each DynamicMesh vertex id.
+	*/
+	static bool AttachVertexMappingData(const TArray<int32>& VertexToOriginalVertexIDMap, UE::Geometry::FDynamicMesh3& InOutMesh)
+	{
+		return AttachNonManifoldVertexMappingData(VertexToOriginalVertexIDMap, InOutMesh);
+	}
+
+	/*
+	* Attaches or replaces  triangle mapping data to the provided mesh.
+	* @param TriangleToOriginalVertexIDMap - an array that maps each DynamicMesh triangle id to the associated original triangle id.
+	* @return false on failure (no attribute will be attached to the DynamicMesh in this case)
+	*
+	* Note: Failure occurs if the DynamicMesh does not have attributes enabled or if the provided array is not long enough to provide a mapping value for each DynamicMesh triangle id.
+	*/
+	static CHAOSCLOTHASSETTOOLS_API bool AttachTriangleMappingData(const TArray<int32>& TriangleToOriginalVertexIDMap, UE::Geometry::FDynamicMesh3& InOutMesh);
+
+	/*
+	*  Removes triangle mapping data.
+	*
+	*  Note: this will invalidate any NonManifoldMappingSupport/ClothPatternToDynamicMeshMappingSupport object associated with this DynamicMesh,
+	*  and subsequent use of such object will produce unexpected results.
+	*/
+	static void RemoveVertexMappingData(UE::Geometry::FDynamicMesh3& InOutMesh)
+	{
+		RemoveNonManifoldVertexMappingData(InOutMesh);
+	}
+
+	static CHAOSCLOTHASSETTOOLS_API void RemoveTriangleMappingData(UE::Geometry::FDynamicMesh3& InOutMesh);
+
+	/*
+	*  Removes all mapping data.
+	*
+	*  Note, this will invalidate any  NonManifoldMappingSupport/ClothPatternToDynamicMeshMappingSupport object associated with this DynamicMesh,
+	*  and subsequent use of such object will produce unexpected results.
+	*/
+	static void RemoveAllMappingData(UE::Geometry::FDynamicMesh3& InOutMesh)
+	{
+		RemoveAllNonManifoldMappingData(InOutMesh);
+		RemoveTriangleMappingData(InOutMesh);
+	}
+
+	/*
+	* Name used to identify triangle attribute data generated during conversion to a DynamicMesh in the case that the source was non-manifold.
+	*/
+	static const CHAOSCLOTHASSETTOOLS_API FName ClothMeshTIDsAttrName;
+
+protected:
+
+	const FDynamicMeshSourceTriangleIdAttribute* SrcTIDsAttribute = nullptr;
+};
+
+}	// namespace UE::Chaos::ClothAsset
+
diff --git a/Engine/Plugins/ChaosClothAsset/Source/ChaosClothAssetTools/Public/ChaosClothAsset/ClothPatternVertexType.h b/Engine/Plugins/ChaosClothAsset/Source/ChaosClothAssetTools/Public/ChaosClothAsset/ClothPatternVertexType.h
new file mode 100644
--- /dev/null
+++ b/Engine/Plugins/ChaosClothAsset/Source/ChaosClothAssetTools/Public/ChaosClothAsset/ClothPatternVertexType.h
@@ -0,0 +1,16 @@
+// Copyright Epic Games, Inc. All Rights Reserved.
+
+#pragma once
+
+#include "CoreTypes.h"
+
+namespace UE::Chaos::ClothAsset
+{
+	enum class EClothPatternVertexType : uint8
+	{
+		Sim2D = 0,
+		Sim3D = 1,
+		Render = 2
+	};
+}
+
diff --git a/Engine/Plugins/ChaosClothAsset/Source/ChaosClothAssetTools/Public/ChaosClothAsset/ColorScheme.h b/Engine/Plugins/ChaosClothAsset/Source/ChaosClothAssetTools/Public/ChaosClothAsset/ColorScheme.h
new file mode 100644
--- /dev/null
+++ b/Engine/Plugins/ChaosClothAsset/Source/ChaosClothAssetTools/Public/ChaosClothAsset/ColorScheme.h
@@ -0,0 +1,16 @@
+// Copyright Epic Games, Inc. All Rights Reserved.
+
+#pragma once
+
+#include "Math/Color.h"
+
+namespace UE::Chaos::ClothAsset
+{
+	struct FColorScheme
+	{
+		static inline const FLinearColor Asset = FColor(242, 77, 26);  // Orange
+		static inline const FLinearColor NodeHeader = FColor(180, 120, 110);
+		static inline const FLinearColor NodeBody = FColor(18, 12, 11, 127);
+	};
+}
+
diff --git a/Engine/Plugins/ChaosClothAssetEditor/ChaosClothAssetEditor.uplugin b/Engine/Plugins/ChaosClothAssetEditor/ChaosClothAssetEditor.uplugin
--- a/Engine/Plugins/ChaosClothAssetEditor/ChaosClothAssetEditor.uplugin
+++ b/Engine/Plugins/ChaosClothAssetEditor/ChaosClothAssetEditor.uplugin
@@ -31,12 +31,6 @@
 			"Type": "Editor",
 			"LoadingPhase": "Default",
 			"PlatformAllowList": [ "Win64", "Mac", "Linux" ]
-		},
-		{
-			"Name": "ChaosClothAssetTools",
-			"Type": "UncookedOnly",
-			"LoadingPhase": "Default",
-			"PlatformAllowList": [ "Win64", "Mac", "Linux" ]
 		}
 	],
 	"Plugins": [
diff --git a/Engine/Plugins/ChaosClothAssetEditor/Source/ChaosClothAssetTools/ChaosClothAssetTools.Build.cs b/Engine/Plugins/ChaosClothAssetEditor/Source/ChaosClothAssetTools/ChaosClothAssetTools.Build.cs
deleted file mode 100644
diff --git a/Engine/Plugins/ChaosClothAssetEditor/Source/ChaosClothAssetTools/Private/ChaosClothAsset/ChaosClothAssetToolsModule.cpp b/Engine/Plugins/ChaosClothAssetEditor/Source/ChaosClothAssetTools/Private/ChaosClothAsset/ChaosClothAssetToolsModule.cpp
deleted file mode 100644
diff --git a/Engine/Plugins/ChaosClothAssetEditor/Source/ChaosClothAssetTools/Private/ChaosClothAsset/ClothAssetEditorStyle.cpp b/Engine/Plugins/ChaosClothAssetEditor/Source/ChaosClothAssetTools/Private/ChaosClothAsset/ClothAssetEditorStyle.cpp
deleted file mode 100644
diff --git a/Engine/Plugins/ChaosClothAssetEditor/Source/ChaosClothAssetTools/Private/ChaosClothAsset/ClothAssetFactory.cpp b/Engine/Plugins/ChaosClothAssetEditor/Source/ChaosClothAssetTools/Private/ChaosClothAsset/ClothAssetFactory.cpp
deleted file mode 100644
diff --git a/Engine/Plugins/ChaosClothAssetEditor/Source/ChaosClothAssetTools/Private/ChaosClothAsset/ClothAssetSKMClothingAssetFactory.cpp b/Engine/Plugins/ChaosClothAssetEditor/Source/ChaosClothAssetTools/Private/ChaosClothAsset/ClothAssetSKMClothingAssetFactory.cpp
deleted file mode 100644
diff --git a/Engine/Plugins/ChaosClothAssetEditor/Source/ChaosClothAssetTools/Private/ChaosClothAsset/ClothAssetSKMClothingAssetFactory.h b/Engine/Plugins/ChaosClothAssetEditor/Source/ChaosClothAssetTools/Private/ChaosClothAsset/ClothAssetSKMClothingAssetFactory.h
deleted file mode 100644
diff --git a/Engine/Plugins/ChaosClothAssetEditor/Source/ChaosClothAssetTools/Private/ChaosClothAsset/ClothAssetSKMSimulationEditorExtender.cpp b/Engine/Plugins/ChaosClothAssetEditor/Source/ChaosClothAssetTools/Private/ChaosClothAsset/ClothAssetSKMSimulationEditorExtender.cpp
deleted file mode 100644
diff --git a/Engine/Plugins/ChaosClothAssetEditor/Source/ChaosClothAssetTools/Private/ChaosClothAsset/ClothAssetSKMSimulationEditorExtender.h b/Engine/Plugins/ChaosClothAssetEditor/Source/ChaosClothAssetTools/Private/ChaosClothAsset/ClothAssetSKMSimulationEditorExtender.h
deleted file mode 100644
diff --git a/Engine/Plugins/ChaosClothAssetEditor/Source/ChaosClothAssetTools/Private/ChaosClothAsset/ClothComponentEditorStyle.cpp b/Engine/Plugins/ChaosClothAssetEditor/Source/ChaosClothAssetTools/Private/ChaosClothAsset/ClothComponentEditorStyle.cpp
deleted file mode 100644
diff --git a/Engine/Plugins/ChaosClothAssetEditor/Source/ChaosClothAssetTools/Private/ChaosClothAsset/ClothComponentEditorStyle.h b/Engine/Plugins/ChaosClothAssetEditor/Source/ChaosClothAssetTools/Private/ChaosClothAsset/ClothComponentEditorStyle.h
deleted file mode 100644
diff --git a/Engine/Plugins/ChaosClothAssetEditor/Source/ChaosClothAssetTools/Private/ChaosClothAsset/ClothPatternToDynamicMesh.cpp b/Engine/Plugins/ChaosClothAssetEditor/Source/ChaosClothAssetTools/Private/ChaosClothAsset/ClothPatternToDynamicMesh.cpp
deleted file mode 100644
diff --git a/Engine/Plugins/ChaosClothAssetEditor/Source/ChaosClothAssetTools/Private/ChaosClothAsset/ClothPatternToDynamicMeshMappingSupport.cpp b/Engine/Plugins/ChaosClothAssetEditor/Source/ChaosClothAssetTools/Private/ChaosClothAsset/ClothPatternToDynamicMeshMappingSupport.cpp
deleted file mode 100644
diff --git a/Engine/Plugins/ChaosClothAssetEditor/Source/ChaosClothAssetTools/Private/ChaosClothAsset/ClothingAssetToClothAssetExporter.cpp b/Engine/Plugins/ChaosClothAssetEditor/Source/ChaosClothAssetTools/Private/ChaosClothAsset/ClothingAssetToClothAssetExporter.cpp
deleted file mode 100644
diff --git a/Engine/Plugins/ChaosClothAssetEditor/Source/ChaosClothAssetTools/Private/ChaosClothAsset/ClothingAssetToClothAssetExporter.h b/Engine/Plugins/ChaosClothAssetEditor/Source/ChaosClothAssetTools/Private/ChaosClothAsset/ClothingAssetToClothAssetExporter.h
deleted file mode 100644
diff --git a/Engine/Plugins/ChaosClothAssetEditor/Source/ChaosClothAssetTools/Private/ChaosClothAsset/SkeletalMeshConverter.cpp b/Engine/Plugins/ChaosClothAssetEditor/Source/ChaosClothAssetTools/Private/ChaosClothAsset/SkeletalMeshConverter.cpp
deleted file mode 100644
diff --git a/Engine/Plugins/ChaosClothAssetEditor/Source/ChaosClothAssetTools/Private/ChaosClothAsset/SkeletalMeshConverter.h b/Engine/Plugins/ChaosClothAssetEditor/Source/ChaosClothAssetTools/Private/ChaosClothAsset/SkeletalMeshConverter.h
deleted file mode 100644
diff --git a/Engine/Plugins/ChaosClothAssetEditor/Source/ChaosClothAssetTools/Public/ChaosClothAsset/ClothAssetEditorStyle.h b/Engine/Plugins/ChaosClothAssetEditor/Source/ChaosClothAssetTools/Public/ChaosClothAsset/ClothAssetEditorStyle.h
deleted file mode 100644
diff --git a/Engine/Plugins/ChaosClothAssetEditor/Source/ChaosClothAssetTools/Public/ChaosClothAsset/ClothAssetFactory.h b/Engine/Plugins/ChaosClothAssetEditor/Source/ChaosClothAssetTools/Public/ChaosClothAsset/ClothAssetFactory.h
deleted file mode 100644
diff --git a/Engine/Plugins/ChaosClothAssetEditor/Source/ChaosClothAssetTools/Public/ChaosClothAsset/ClothPatternToDynamicMesh.h b/Engine/Plugins/ChaosClothAssetEditor/Source/ChaosClothAssetTools/Public/ChaosClothAsset/ClothPatternToDynamicMesh.h
deleted file mode 100644
diff --git a/Engine/Plugins/ChaosClothAssetEditor/Source/ChaosClothAssetTools/Public/ChaosClothAsset/ClothPatternToDynamicMeshMappingSupport.h b/Engine/Plugins/ChaosClothAssetEditor/Source/ChaosClothAssetTools/Public/ChaosClothAsset/ClothPatternToDynamicMeshMappingSupport.h
deleted file mode 100644
diff --git a/Engine/Plugins/ChaosClothAssetEditor/Source/ChaosClothAssetTools/Public/ChaosClothAsset/ClothPatternVertexType.h b/Engine/Plugins/ChaosClothAssetEditor/Source/ChaosClothAssetTools/Public/ChaosClothAsset/ClothPatternVertexType.h
deleted file mode 100644
diff --git a/Engine/Plugins/ChaosClothAssetEditor/Source/ChaosClothAssetTools/Public/ChaosClothAsset/ColorScheme.h b/Engine/Plugins/ChaosClothAssetEditor/Source/ChaosClothAssetTools/Public/ChaosClothAsset/ColorScheme.h
deleted file mode 100644
