Chaos SKM Cloth Asset - Fixed simulation on first loading the assets.
UChaosClothAssetSKMClothingAsset::HasAnySimulationMeshData() was dependent on correct asset loading order.

#rnx

#Synced-CL 48843225

diff --git a/Engine/Plugins/ChaosClothAsset/Source/ChaosClothAssetEngine/Private/ChaosClothAsset/ClothAssetSKMClothingAsset.cpp b/Engine/Plugins/ChaosClothAsset/Source/ChaosClothAssetEngine/Private/ChaosClothAsset/ClothAssetSKMClothingAsset.cpp
--- a/Engine/Plugins/ChaosClothAsset/Source/ChaosClothAssetEngine/Private/ChaosClothAsset/ClothAssetSKMClothingAsset.cpp
+++ b/Engine/Plugins/ChaosClothAsset/Source/ChaosClothAssetEngine/Private/ChaosClothAsset/ClothAssetSKMClothingAsset.cpp
@@ -129,13 +129,6 @@ UChaosClothAssetSKMClothingAsset::UChaosClothAssetSKMClothingAsset(const FObject
 	}
 }
 
-void UChaosClothAssetSKMClothingAsset::PostLoad()
-{
-	Super::PostLoad();
-
-	CalculateLODHasAnySimulationMeshData();
-}
-
 #if WITH_EDITOR
 void UChaosClothAssetSKMClothingAsset::SetAsset(const UChaosClothAssetBase* InAsset)
 {
@@ -200,10 +193,6 @@ void UChaosClothAssetSKMClothingAsset::OnAssetChanged(const bool bReregisterComp
 
 		OnModelChanged();
 	}
-	else
-	{
-		CalculateLODHasAnySimulationMeshData();
-	}
 
 	if (const USkeletalMesh* const OwnerMesh = Cast<USkeletalMesh>(GetOuter()))
 	{
@@ -244,6 +233,21 @@ void UChaosClothAssetSKMClothingAsset::OnAssetChanged(const bool bReregisterComp
 	}
 }
 
+bool UChaosClothAssetSKMClothingAsset::HasAnySimulationMeshData(const int32 LODIndex) const
+{
+	if (Asset && ClothSimulationModelIndex >= 0 && ClothSimulationModelIndex < Asset->GetNumClothSimulationModels())
+	{
+		if (const TSharedPtr<const FChaosClothSimulationModel> ClothSimulationModel = Asset->GetClothSimulationModel(ClothSimulationModelIndex))
+		{
+			return
+				ClothSimulationModel->IsValidLodIndex(LODIndex) &&
+				ClothSimulationModel->GetNumVertices(LODIndex) > 0 &&
+				ClothSimulationModel->GetNumTriangles(LODIndex) > 0;
+		}
+	}
+	return false;
+}
+
 void UChaosClothAssetSKMClothingAsset::PostEditChangeProperty(FPropertyChangedEvent& PropertyChangedEvent)
 {
 	Super::PostEditChangeProperty(PropertyChangedEvent);
@@ -547,7 +551,6 @@ void UChaosClothAssetSKMClothingAsset::OnModelChanged()
 		RestoreClothingFromBackup(*SkeletalMesh, *this, ClothingBindings);
 		SkeletalMesh->InvalidateDeriveDataCacheGUID();  // The DDC key only changes with the ClothingAsset GUID, but a model change still requires a mesh rebuild
 	}
-	CalculateLODHasAnySimulationMeshData();
 }
 
 FGuid UChaosClothAssetSKMClothingAsset::GetClothSimulationModelGuid() const
@@ -587,20 +590,4 @@ TArray<FString> UChaosClothAssetSKMClothingAsset::GetClothSimulationModelIds() c
 	return ClothSimulationModelIds;
 }
 
-void UChaosClothAssetSKMClothingAsset::CalculateLODHasAnySimulationMeshData()
-{
-	LODHasAnySimulationMeshData.Reset();
-	if (Asset && ClothSimulationModelIndex != INDEX_NONE)
-	{
-		if (const TSharedPtr<const FChaosClothSimulationModel> ClothSimulationModel = Asset->GetClothSimulationModel(ClothSimulationModelIndex))
-		{
-			LODHasAnySimulationMeshData.SetNumUninitialized(ClothSimulationModel->GetNumLods());
-			for (int32 LODIndex = 0; LODIndex < ClothSimulationModel->GetNumLods(); ++LODIndex)
-			{
-				LODHasAnySimulationMeshData[LODIndex] = ClothSimulationModel->GetNumVertices(LODIndex) > 0 && ClothSimulationModel->GetNumTriangles(LODIndex) > 0;
-			}
-		}
-	}
-}
-
 #undef LOCTEXT_NAMESPACE
diff --git a/Engine/Plugins/ChaosClothAsset/Source/ChaosClothAssetEngine/Public/ChaosClothAsset/ClothAssetSKMClothingAsset.h b/Engine/Plugins/ChaosClothAsset/Source/ChaosClothAssetEngine/Public/ChaosClothAsset/ClothAssetSKMClothingAsset.h
--- a/Engine/Plugins/ChaosClothAsset/Source/ChaosClothAssetEngine/Public/ChaosClothAsset/ClothAssetSKMClothingAsset.h
+++ b/Engine/Plugins/ChaosClothAsset/Source/ChaosClothAssetEngine/Public/ChaosClothAsset/ClothAssetSKMClothingAsset.h
@@ -47,14 +47,10 @@ public:
 		return ClothSimulationModelIndex;
 	}
 
-	bool HasAnySimulationMeshData(const int32 LODIndex) const
-	{
-		return LODHasAnySimulationMeshData.IsValidIndex(LODIndex) && LODHasAnySimulationMeshData[LODIndex];
-	}
+	bool HasAnySimulationMeshData(const int32 LODIndex) const;
 
 private:
 	//~ Begin UObject Interface
-	virtual void PostLoad() override;
 #if WITH_EDITOR
 	virtual void PostEditChangeProperty(FPropertyChangedEvent& PropertyChangedEvent) override;
 	//~ End UObject Interface
@@ -73,8 +69,6 @@ private:
 	virtual void PostUpdateAllAssets() override {}
 	//~ End UClothingAssetBase Interface
 
-	void CalculateLODHasAnySimulationMeshData();
-
 #if WITH_EDITOR
 	void OnModelChanged();
 
@@ -100,9 +94,6 @@ private:
 
 	UPROPERTY()
 	int32 ClothSimulationModelIndex = INDEX_NONE;
-
-	// Calculated from Asset during PostLoad or OnAssetChanged
-	TBitArray<> LODHasAnySimulationMeshData;
 };
 
 #undef UE_API
