CIS - Fixed issue 1154737: Undefined symbols: UChaosClothAssetSKMClothingAsset::HasAnySimulationMeshData and UE::Chaos::ClothAsset::FSKMClothingSimulation::HasAnySimulationMeshData.

#rnx
#Synced-CL 48843225

diff --git a/Engine/Plugins/ChaosClothAsset/Source/ChaosClothAssetEngine/Public/ChaosClothAsset/ClothAssetSKMClothingAsset.h b/Engine/Plugins/ChaosClothAsset/Source/ChaosClothAssetEngine/Public/ChaosClothAsset/ClothAssetSKMClothingAsset.h
--- a/Engine/Plugins/ChaosClothAsset/Source/ChaosClothAssetEngine/Public/ChaosClothAsset/ClothAssetSKMClothingAsset.h
+++ b/Engine/Plugins/ChaosClothAsset/Source/ChaosClothAssetEngine/Public/ChaosClothAsset/ClothAssetSKMClothingAsset.h
@@ -47,7 +47,7 @@ public:
 		return ClothSimulationModelIndex;
 	}
 
-	bool HasAnySimulationMeshData(const int32 LODIndex) const;
+	UE_API bool HasAnySimulationMeshData(const int32 LODIndex) const;
 
 private:
 	//~ Begin UObject Interface
