[CHAOS] Move memory report into separate command, fix the pure virtual call crash for Chaos Cloth CDO

#Synced-CL 49721529

diff --git a/Engine/Plugins/ChaosClothAsset/Source/ChaosClothAssetEngine/Private/ChaosClothAsset/ClothAssetBase.cpp b/Engine/Plugins/ChaosClothAsset/Source/ChaosClothAssetEngine/Private/ChaosClothAsset/ClothAssetBase.cpp
--- a/Engine/Plugins/ChaosClothAsset/Source/ChaosClothAssetEngine/Private/ChaosClothAsset/ClothAssetBase.cpp
+++ b/Engine/Plugins/ChaosClothAsset/Source/ChaosClothAssetEngine/Private/ChaosClothAsset/ClothAssetBase.cpp
@@ -21,12 +21,16 @@
 #include "Features/IModularFeatures.h"
 #include "Interfaces/ITargetPlatform.h"
 #endif
+#include "HAL/IConsoleManager.h"
+#include "ProfilingDebugging/ResourceSize.h"
 #include "Rendering/SkeletalMeshRenderData.h"
 #include "UObject/AssetRegistryTagsContext.h"
 #include "UObject/FortniteMainBranchObjectVersion.h"
+#include "UObject/ObjectMacros.h"
 #if WITH_EDITOR
 #include "UObject/PackageReload.h"
 #endif
+#include "UObject/UObjectGlobals.h"
 #include "ClothingSimulationFactory.h"
 #include "ClothingSimulationInstance.h"
 #include "ClothingSimulationInteractor.h"
@@ -59,6 +63,82 @@ namespace UE::Chaos::ClothAsset::Private
 #endif
 			}),
 		ECVF_Scalability);
+		
+#if !UE_BUILD_SHIPPING
+	FAutoConsoleCommandWithArgsAndOutputDevice CmdClothMemoryReport(
+		TEXT("p.ClothAsset.MemoryReport"),
+		TEXT("Outputs memory usage for the provided cloth asset name"),
+		FConsoleCommandWithArgsAndOutputDeviceDelegate::CreateLambda([](TArray<FString> Args, FOutputDevice& Ar)
+		{
+			if (Args.Num() != 1)
+			{
+				Ar.Logf(TEXT("Invalid usage, should be: p.ClothAsset.MemoryReport <ObjectName>"));
+				return;
+			}
+
+			const FString& ObjectName = Args[0];
+			UChaosClothAssetBase* Obj = FindFirstObject<UChaosClothAssetBase>(*ObjectName, EFindFirstObjectOptions::None, ELogVerbosity::Warning, TEXT("p.ClothAsset.MemoryReport"));
+
+			if (!Obj)
+			{
+				return;
+			}
+
+			FString MemoryReport;
+			MemoryReport.Appendf(TEXT("---- Memory report for [%s] [%s] ----"), *Obj->StaticClass()->GetName(), *Obj->GetName());
+
+			FResourceSizeEx RenderDataResourceSize;
+			if (Obj->GetResourceForRendering())
+			{
+				for (int32 LodIndex = 0; LodIndex < Obj->GetResourceForRendering()->LODRenderData.Num(); ++LodIndex)
+				{
+					FResourceSizeEx LODRenderDataResourceSize;
+					Obj->GetResourceForRendering()->LODRenderData[LodIndex].GetResourceSizeEx(LODRenderDataResourceSize);
+					MemoryReport.Appendf(TEXT("\n LODRenderData LOD%d size: %ld bytes"), LodIndex, LODRenderDataResourceSize.GetTotalMemoryBytes());
+				}
+
+				Obj->GetResourceForRendering()->GetResourceSizeEx(RenderDataResourceSize);
+			}
+			MemoryReport.Appendf(TEXT("\n Total RenderData size: %ld bytes"), RenderDataResourceSize.GetTotalMemoryBytes());
+
+			FResourceSizeEx ClothSimulationModelsResourceSize;
+			for (int32 ModelIndex = 0; ModelIndex < Obj->GetNumClothSimulationModels(); ++ModelIndex)
+			{
+				if (const TSharedPtr<FChaosClothSimulationModel> ClothSimulationModel = ConstCastSharedPtr<FChaosClothSimulationModel>(Obj->GetClothSimulationModel(ModelIndex)))
+				{
+					for (int32 LodIndex = 0; LodIndex < ClothSimulationModel->GetNumLods(); ++LodIndex)
+					{
+						FResourceSizeEx ClothSimulationLodModelResourceSize;
+						ClothSimulationModel->ClothSimulationLodModels[LodIndex].GetResourceSizeEx(ClothSimulationLodModelResourceSize);
+						MemoryReport.Appendf(TEXT("\n ClothSimulationModel%d LOD%d size: %ld bytes"), ModelIndex, LodIndex, ClothSimulationLodModelResourceSize.GetTotalMemoryBytes());
+					}
+
+					ClothSimulationModel->GetResourceSizeEx(ClothSimulationModelsResourceSize);
+				}
+			}
+			MemoryReport.Appendf(TEXT("\n Total ClothSimulationModel(s) size: %ld bytes"), ClothSimulationModelsResourceSize.GetTotalMemoryBytes());
+
+			const int64 TotalResourceSize = RenderDataResourceSize.GetTotalMemoryBytes() + ClothSimulationModelsResourceSize.GetTotalMemoryBytes();
+			MemoryReport.Appendf(
+				TEXT("\n Total resource size for Cloth Asset [%s]: %ld bytes (%.3f MB)"),
+				*Obj->GetName(),
+				TotalResourceSize,
+				(float)TotalResourceSize / (1024.f * 1024.f));
+
+			FResourceSizeEx CumulativeResourceSize(EResourceSizeMode::EstimatedTotal);
+			Obj->GetResourceSizeEx(CumulativeResourceSize);
+			const int64 TotalSize = CumulativeResourceSize.GetTotalMemoryBytes();
+			MemoryReport.Appendf(
+				TEXT("\n Total size for Cloth Asset [%s]: %ld bytes (%.3f MB)"),
+				*Obj->GetName(),
+				TotalSize,
+				(float)TotalSize / (1024.f * 1024.f));
+
+			UE_LOG(LogChaosClothAsset, Display, TEXT("\n%s"), *MemoryReport);
+		}),
+		ECVF_Cheat
+	);
+#endif //!UE_BUILD_SHIPPING
 }
 
 /** Used for locking resources during async building. */
@@ -346,65 +426,16 @@ void UChaosClothAssetBase::GetResourceSizeEx(FResourceSizeEx& CumulativeResource
 		GetResourceForRendering()->GetResourceSizeEx(CumulativeResourceSize);
 	}
 
-	for (int32 ModelIndex = 0; ModelIndex < GetNumClothSimulationModels(); ++ModelIndex)
+	if (!GetClass()->HasAnyClassFlags(CLASS_Abstract))
 	{
-		if (const TSharedPtr<FChaosClothSimulationModel> ClothSimulationModel = ConstCastSharedPtr<FChaosClothSimulationModel>(GetClothSimulationModel(ModelIndex)))
+		for (int32 ModelIndex = 0; ModelIndex < GetNumClothSimulationModels(); ++ModelIndex)
 		{
-			ClothSimulationModel->GetResourceSizeEx(CumulativeResourceSize);
-		}
-	}
-
-#if !UE_BUILD_SHIPPING
-	FString MemoryReport;
-	MemoryReport.Appendf(TEXT("---- Memory report for [%s] [%s] ----"), *StaticClass()->GetName(), *this->GetName());
-
-	FResourceSizeEx RenderDataResourceSize;
-	if (GetResourceForRendering())
-	{
-		for (int32 LodIndex = 0; LodIndex < GetResourceForRendering()->LODRenderData.Num(); ++LodIndex)
-		{
-			FResourceSizeEx LODRenderDataResourceSize;
-			GetResourceForRendering()->LODRenderData[LodIndex].GetResourceSizeEx(LODRenderDataResourceSize);
-			MemoryReport.Appendf(TEXT("\n LODRenderData LOD%d size: %ld bytes"), LodIndex, LODRenderDataResourceSize.GetTotalMemoryBytes());
-		}
-
-		GetResourceForRendering()->GetResourceSizeEx(RenderDataResourceSize);
-	}
-	MemoryReport.Appendf(TEXT("\n Total RenderData size: %ld bytes"), RenderDataResourceSize.GetTotalMemoryBytes());
-
-	FResourceSizeEx ClothSimulationModelsResourceSize;
-	for (int32 ModelIndex = 0; ModelIndex < GetNumClothSimulationModels(); ++ModelIndex)
-	{
-		if (const TSharedPtr<FChaosClothSimulationModel> ClothSimulationModel = ConstCastSharedPtr<FChaosClothSimulationModel>(GetClothSimulationModel(ModelIndex)))
-		{
-			for (int32 LodIndex = 0; LodIndex < ClothSimulationModel->GetNumLods(); ++LodIndex)
+			if (const TSharedPtr<FChaosClothSimulationModel> ClothSimulationModel = ConstCastSharedPtr<FChaosClothSimulationModel>(GetClothSimulationModel(ModelIndex)))
 			{
-				FResourceSizeEx ClothSimulationLodModelResourceSize;
-				ClothSimulationModel->ClothSimulationLodModels[LodIndex].GetResourceSizeEx(ClothSimulationLodModelResourceSize);
-				MemoryReport.Appendf(TEXT("\n ClothSimulationModel%d LOD%d size: %ld bytes"), ModelIndex, LodIndex, ClothSimulationLodModelResourceSize.GetTotalMemoryBytes());
+				ClothSimulationModel->GetResourceSizeEx(CumulativeResourceSize);
 			}
-
-			ClothSimulationModel->GetResourceSizeEx(ClothSimulationModelsResourceSize);
 		}
 	}
-	MemoryReport.Appendf(TEXT("\n Total ClothSimulationModel(s) size: %ld bytes"), ClothSimulationModelsResourceSize.GetTotalMemoryBytes());
-
-	const int64 TotalResourceSize = RenderDataResourceSize.GetTotalMemoryBytes() + ClothSimulationModelsResourceSize.GetTotalMemoryBytes();
-	MemoryReport.Appendf(
-		TEXT("\n Total resource size for Cloth Asset [%s]: %ld bytes (%.3f MB)"),
-		*this->GetName(),
-		TotalResourceSize,
-		(float)TotalResourceSize / (1024.f * 1024.f));
-
-	const int64 TotalSize = CumulativeResourceSize.GetTotalMemoryBytes();
-	MemoryReport.Appendf(
-		TEXT("\n Total size for Cloth Asset [%s]: %ld bytes (%.3f MB)"),
-		*this->GetName(),
-		TotalSize,
-		(float)TotalSize / (1024.f * 1024.f));
-
-	UE_LOG(LogChaosClothAsset, Display, TEXT("\n%s"), *MemoryReport);
-#endif
 }
 
 FSkeletalMeshLODInfo* UChaosClothAssetBase::GetLODInfo(int32 Index)
