Chaos Cloth Asset - Added NonTransactional to the interactor arrays property to avoid any issues if the cloth component was ever simulating in editor.

#rnx

#Synced-CL 49759196

diff --git a/Engine/Plugins/ChaosClothAsset/Source/ChaosClothAssetEngine/Public/ChaosClothAsset/ClothComponent.h b/Engine/Plugins/ChaosClothAsset/Source/ChaosClothAssetEngine/Public/ChaosClothAsset/ClothComponent.h
--- a/Engine/Plugins/ChaosClothAsset/Source/ChaosClothAssetEngine/Public/ChaosClothAsset/ClothComponent.h
+++ b/Engine/Plugins/ChaosClothAsset/Source/ChaosClothAssetEngine/Public/ChaosClothAsset/ClothComponent.h
@@ -501,7 +501,7 @@ private:
 	 * Simulation properties per cloth model. 
 	 * Transient because it doesn't need to be serialized but contains pointers to UObject that needs not be garbage collected.
 	 */
-	UPROPERTY(Transient)
+	UPROPERTY(Transient, NonTransactional)
 	TArray<FChaosClothSimulationProperties> ClothSimulationProperties;
 
 	TSharedPtr<UE::Chaos::ClothAsset::FClothSimulationProxy> ClothSimulationProxy;
