Chaos Cloth Asset - Moved the Thumbnail renderer out of the ChaosClothAssetEditor into the ChaosClothAsset plugin.
This is a first pass deprecation to tidy up the original implementation that made the thumbnail code unnecessarily public.

#rnx

#Synced-CL 50202394

diff --git a/Engine/Plugins/ChaosClothAsset/Source/ChaosClothAssetTools/ChaosClothAssetTools.Build.cs b/Engine/Plugins/ChaosClothAsset/Source/ChaosClothAssetTools/ChaosClothAssetTools.Build.cs
--- a/Engine/Plugins/ChaosClothAsset/Source/ChaosClothAssetTools/ChaosClothAssetTools.Build.cs
+++ b/Engine/Plugins/ChaosClothAsset/Source/ChaosClothAssetTools/ChaosClothAssetTools.Build.cs
@@ -42,5 +42,7 @@ public class ChaosClothAssetTools : ModuleRules
 				"SkeletalMeshUtilitiesCommon"
 			}
 		);
+
+		PrivateIncludePaths.Add(System.IO.Path.Combine(GetModuleDirectory("ChaosClothAssetTools"), "Internal"));
 	}
 }
diff --git a/Engine/Plugins/ChaosClothAsset/Source/ChaosClothAssetTools/Internal/ChaosClothAsset/ClothAssetThumbnailRenderer.h b/Engine/Plugins/ChaosClothAsset/Source/ChaosClothAssetTools/Internal/ChaosClothAsset/ClothAssetThumbnailRenderer.h
new file mode 100644
--- /dev/null
+++ b/Engine/Plugins/ChaosClothAsset/Source/ChaosClothAssetTools/Internal/ChaosClothAsset/ClothAssetThumbnailRenderer.h
@@ -0,0 +1,69 @@
+// Copyright Epic Games, Inc. All Rights Reserved.
+
+#pragma once
+
+#include "HAL/Platform.h"
+#include "ThumbnailRendering/SkeletalMeshThumbnailRenderer.h"
+
+#include "ClothAssetThumbnailRenderer.generated.h"
+
+class UChaosClothComponent;
+class UChaosClothAssetBase;
+
+/**
+ * Minimal actor with a cloth component used to render cloth assets previews
+ */
+UCLASS(MinimalAPI)
+class AChaosClothPreviewActor_Internal : public AActor
+{
+	GENERATED_BODY()
+public:
+	AChaosClothPreviewActor_Internal();
+
+	UChaosClothComponent* GetClothComponent() const
+	{
+		return ClothComponent;
+	}
+
+protected:
+	UPROPERTY()
+	TObjectPtr<UChaosClothComponent> ClothComponent;
+};
+
+namespace UE::Chaos::ClothAsset
+{
+	/**
+	 * Preview scene for a Cloth asset thumbnail
+	 */
+	class FThumbnailScene_Internal : public FThumbnailPreviewScene
+	{
+	public:
+		FThumbnailScene_Internal();
+
+		void SetClothAsset(UChaosClothAssetBase* InClothAsset);
+
+		virtual void GetViewMatrixParameters(const float InFOVDegrees, FVector& OutOrigin, float& OutOrbitPitch, float& OutOrbitYaw, float& OutOrbitZoom) const override;
+
+	private:
+		AChaosClothPreviewActor_Internal* PreviewActor = nullptr;
+	};
+}
+
+
+/**
+ * Renders a preview scene to be used as a thumbnail for a supported cloth asset
+ */
+UCLASS(MinimalAPI)
+class UChaosClothAssetThumbnailRenderer_Internal : public UDefaultSizedThumbnailRenderer
+{
+	GENERATED_BODY()
+
+	virtual void Draw(UObject* Object, int32 X, int32 Y, uint32 Width, uint32 Height, FRenderTarget* Viewport, FCanvas* Canvas, bool bAdditionalViewFamily) override;
+	virtual bool CanVisualizeAsset(UObject* Object) override;
+	virtual EThumbnailRenderFrequency GetThumbnailRenderFrequency(UObject* Object) const override;
+
+	virtual void BeginDestroy() override;
+
+protected:
+	TObjectInstanceThumbnailScene<UE::Chaos::ClothAsset::FThumbnailScene_Internal, 128> ClothThumbnailSceneCache;
+};
diff --git a/Engine/Plugins/ChaosClothAsset/Source/ChaosClothAssetTools/Private/ChaosClothAsset/ChaosClothAssetToolsModule.cpp b/Engine/Plugins/ChaosClothAsset/Source/ChaosClothAssetTools/Private/ChaosClothAsset/ChaosClothAssetToolsModule.cpp
--- a/Engine/Plugins/ChaosClothAsset/Source/ChaosClothAssetTools/Private/ChaosClothAsset/ChaosClothAssetToolsModule.cpp
+++ b/Engine/Plugins/ChaosClothAsset/Source/ChaosClothAssetTools/Private/ChaosClothAsset/ChaosClothAssetToolsModule.cpp
@@ -2,13 +2,16 @@
 
 #include "Features/IModularFeatures.h"
 #include "Modules/ModuleManager.h"
+#include "ChaosClothAsset/ClothAsset.h"
 #include "ChaosClothAsset/ClothAssetEditorStyle.h"
 #include "ChaosClothAsset/ClothAssetSKMClothingAssetFactory.h"
+#include "ChaosClothAsset/ClothAssetThumbnailRenderer.h"
 #include "ChaosClothAsset/ClothComponentEditorStyle.h"
 #include "ChaosClothAsset/ClothingAssetToClothAssetExporter.h"
 #include "ChaosClothAsset/SkeletalMeshConverter.h"
 #include "ChaosClothAsset/ClothAssetSKMSimulationEditorExtender.h"
 #include "ClothingSystemEditorInterfaceModule.h"
+#include "ThumbnailRendering/ThumbnailManager.h"
 
 namespace UE::Chaos::ClothAsset
 {
@@ -33,6 +36,8 @@ namespace UE::Chaos::ClothAsset
 			IModularFeatures::Get().RegisterModularFeature(IClothingAssetFactoryProvider::FeatureName, static_cast<IClothingAssetFactoryProvider*>(this));
 			IModularFeatures::Get().RegisterModularFeature(FClothingSystemEditorInterfaceModule::ExtenderFeatureName, &SKMSimulationEditorExtender);
 
+			// Register thumbnail renderer
+			UThumbnailManager::Get().RegisterCustomRenderer(UChaosClothAsset::StaticClass(), UChaosClothAssetThumbnailRenderer_Internal::StaticClass());
 		}
 
 		virtual void ShutdownModule() override
@@ -44,6 +49,9 @@ namespace UE::Chaos::ClothAsset
 				IModularFeatures::Get().UnregisterModularFeature(IClothAssetSkeletalMeshConverterClassProvider::FeatureName, static_cast<IClothAssetSkeletalMeshConverterClassProvider*>(this));
 				IModularFeatures::Get().UnregisterModularFeature(IClothingAssetFactoryProvider::FeatureName, static_cast<IClothingAssetFactoryProvider*>(this));
 				IModularFeatures::Get().UnregisterModularFeature(FClothingSystemEditorInterfaceModule::ExtenderFeatureName, &SKMSimulationEditorExtender);
+
+				// Unregister thumbnail renderer
+				UThumbnailManager::Get().UnregisterCustomRenderer(UChaosClothAsset::StaticClass());
 			}
 		}
 
diff --git a/Engine/Plugins/ChaosClothAsset/Source/ChaosClothAssetTools/Private/ChaosClothAsset/ClothAssetThumbnailRenderer.cpp b/Engine/Plugins/ChaosClothAsset/Source/ChaosClothAssetTools/Private/ChaosClothAsset/ClothAssetThumbnailRenderer.cpp
new file mode 100644
--- /dev/null
+++ b/Engine/Plugins/ChaosClothAsset/Source/ChaosClothAssetTools/Private/ChaosClothAsset/ClothAssetThumbnailRenderer.cpp
@@ -0,0 +1,177 @@
+// Copyright Epic Games, Inc. All Rights Reserved.
+
+#include "ChaosClothAsset/ClothAssetThumbnailRenderer.h"
+
+#include "ChaosClothAsset/ClothAsset.h"
+#include "ChaosClothAsset/ClothAssetBase.h"
+#include "ChaosClothAsset/ClothComponent.h"
+#include "ChaosClothAsset/ClothGeometryTools.h"
+#include "Components/SkeletalMeshComponent.h"
+#include "Engine/SkeletalMesh.h"
+#include "Rendering/SkeletalMeshRenderData.h"
+#include "ThumbnailRendering/SceneThumbnailInfo.h"
+#include "SceneView.h"
+#include "ShowFlags.h"
+
+#include UE_INLINE_GENERATED_CPP_BY_NAME(ClothAssetThumbnailRenderer)
+
+AChaosClothPreviewActor_Internal::AChaosClothPreviewActor_Internal()
+{
+	ClothComponent = CreateDefaultSubobject<UChaosClothComponent>(TEXT("ClothComponent0"));
+	RootComponent = ClothComponent;
+}
+
+namespace UE::Chaos::ClothAsset
+{
+	FThumbnailScene_Internal::FThumbnailScene_Internal() : FThumbnailPreviewScene()
+	{
+		bForceAllUsedMipsResident = false;
+
+		if (GetWorld())
+		{
+			FActorSpawnParameters SpawnInfo;
+			SpawnInfo.SpawnCollisionHandlingOverride = ESpawnActorCollisionHandlingMethod::AlwaysSpawn;
+			SpawnInfo.bNoFail = true;
+			SpawnInfo.ObjectFlags = RF_Transient;
+
+			PreviewActor = GetWorld()->SpawnActor<AChaosClothPreviewActor_Internal>(SpawnInfo);
+			PreviewActor->SetActorEnableCollision(false);
+		}
+	}
+
+	void FThumbnailScene_Internal::SetClothAsset(UChaosClothAssetBase* InClothAsset)
+	{
+		if (UChaosClothComponent* const ClothComponent = PreviewActor ? PreviewActor->GetClothComponent() : nullptr;
+			ensure(ClothComponent))
+		{
+			ClothComponent->SetAsset(InClothAsset);
+		}
+	}
+
+	void FThumbnailScene_Internal::GetViewMatrixParameters(
+		const float InFOVDegrees,
+		FVector& OutOrigin,
+		float& OutOrbitPitch,
+		float& OutOrbitYaw,
+		float& OutOrbitZoom) const
+	{	
+		if (UChaosClothComponent* const ClothComponent = PreviewActor ? PreviewActor->GetClothComponent() : nullptr;
+			ensure(ClothComponent))
+		{
+			FBoxSphereBounds Bounds = ClothComponent->Bounds;
+			Bounds = Bounds.ExpandBy(2.0f);
+		
+			const float HalfFOVRadians = FMath::DegreesToRadians<float>(InFOVDegrees) * 0.5f;
+			const float HalfMeshSize = static_cast<float>(Bounds.SphereRadius); 
+			const float TargetDistance = HalfMeshSize / FMath::Tan(HalfFOVRadians);
+
+			USceneThumbnailInfo* ThumbnailInfo = Cast<USceneThumbnailInfo>(ClothComponent->GetThumbnailInfo());
+			if (ThumbnailInfo)
+			{
+				if (TargetDistance + ThumbnailInfo->OrbitZoom < 0 )
+				{
+					ThumbnailInfo->OrbitZoom = -TargetDistance;
+				}
+			}
+			else
+			{
+				ThumbnailInfo = USceneThumbnailInfo::StaticClass()->GetDefaultObject<USceneThumbnailInfo>();
+			}
+
+			OutOrigin = -Bounds.Origin;
+			OutOrbitPitch = ThumbnailInfo->OrbitPitch;
+			OutOrbitYaw = ThumbnailInfo->OrbitYaw;
+			OutOrbitZoom = TargetDistance + ThumbnailInfo->OrbitZoom;
+		}
+		else
+		{
+			OutOrigin = FVector::ZeroVector;
+			OutOrbitPitch = 0.f;
+			OutOrbitYaw = 0.f;
+			OutOrbitZoom =0.f;
+		}
+	}
+}
+
+void UChaosClothAssetThumbnailRenderer_Internal::Draw(UObject* Object, int32 X, int32 Y, uint32 Width, uint32 Height, FRenderTarget* Viewport, FCanvas* Canvas, bool bAdditionalViewFamily)
+{
+	UChaosClothAssetBase* const ClothAsset = Cast<UChaosClothAssetBase>(Object);
+	if (!ClothAsset)
+	{
+		return;
+	}
+	
+	const TSharedRef<UE::Chaos::ClothAsset::FThumbnailScene_Internal> ThumbnailScene = ClothThumbnailSceneCache.EnsureThumbnailScene(Object);
+
+	// Check if the cloth asset has render data, if not render the sim mesh instead
+	FSkeletalMeshRenderData* const RenderData = ClothAsset->GetResourceForRendering();
+	const bool bHasRenderData = RenderData && RenderData->LODRenderData.Num() && RenderData->LODRenderData[0].GetTotalFaces();
+
+	if (!bHasRenderData && ClothAsset->HasValidClothSimulationModels() && ClothAsset->GetCollections(0).Num())
+	{
+		// Render sim mesh
+		const FSoftObjectPath RenderMaterialPathName(TEXT("/Engine/EditorMaterials/Cloth/CameraLitDoubleSided.CameraLitDoubleSided"));
+
+		const TSharedRef<const FManagedArrayCollection>& Collection = ClothAsset->GetCollections(0)[0];
+		UChaosClothAsset* const ClothAssetCopy = NewObject<UChaosClothAsset>(this, UChaosClothAsset::StaticClass());
+
+		TSharedRef<FManagedArrayCollection> CollectionCopy = MakeShared<FManagedArrayCollection>(*Collection);
+
+		// Reverse mesh normals to keep the same render look as inverted sim mesh in flat shaded visualization
+		constexpr bool bReverseSimMeshNormals = true;
+		constexpr bool bReverseSimMeshWindingOrder = false;
+		constexpr bool bReverseRenderMeshNormals = false;
+		constexpr bool bReverseRenderMeshWindingOrder = false;
+		const TArray<int32> EmptySelection;
+
+		UE::Chaos::ClothAsset::FClothGeometryTools::ReverseMesh(
+			CollectionCopy,
+			bReverseSimMeshNormals,
+			bReverseSimMeshWindingOrder,
+			bReverseRenderMeshNormals,
+			bReverseRenderMeshWindingOrder,
+			EmptySelection,
+			EmptySelection);
+
+		constexpr bool bSingleRenderPattern = true;
+		UE::Chaos::ClothAsset::FClothGeometryTools::CopySimMeshToRenderMesh(CollectionCopy, RenderMaterialPathName, bSingleRenderPattern);
+
+		const TArray<TSharedRef<const FManagedArrayCollection>> ClothCollectionsCopy = { MoveTemp(CollectionCopy) };
+
+		ClothAssetCopy->Build(ClothCollectionsCopy);
+		ThumbnailScene->SetClothAsset(ClothAssetCopy);
+	}
+	else
+	{
+		// Render render mesh
+		ThumbnailScene->SetClothAsset(ClothAsset);
+	}
+
+	FSceneViewFamilyContext ViewFamily( FSceneViewFamily::ConstructionValues( Viewport, ThumbnailScene->GetScene(), FEngineShowFlags(ESFIM_Game) )
+		.SetTime(UThumbnailRenderer::GetTime())
+		.SetAdditionalViewFamily(bAdditionalViewFamily));
+
+	ViewFamily.EngineShowFlags.DisableAdvancedFeatures();
+	ViewFamily.EngineShowFlags.MotionBlur = 0;
+	ViewFamily.EngineShowFlags.LOD = 0;
+
+	RenderViewFamily(Canvas, &ViewFamily, ThumbnailScene->CreateView(&ViewFamily, X, Y, Width, Height));
+	ThumbnailScene->SetClothAsset(nullptr);
+}
+
+bool UChaosClothAssetThumbnailRenderer_Internal::CanVisualizeAsset(UObject* Object)
+{
+	return Cast<UChaosClothAssetBase>(Object) != nullptr;
+}
+
+EThumbnailRenderFrequency UChaosClothAssetThumbnailRenderer_Internal::GetThumbnailRenderFrequency(UObject* Object) const
+{
+	UChaosClothAssetBase* const AsClothAsset = Cast<UChaosClothAssetBase>(Object);
+	return AsClothAsset && AsClothAsset->GetResourceForRendering() ? EThumbnailRenderFrequency::Realtime : EThumbnailRenderFrequency::OnPropertyChange;
+}
+
+void UChaosClothAssetThumbnailRenderer_Internal::BeginDestroy()
+{
+	ClothThumbnailSceneCache.Clear();
+	Super::BeginDestroy();
+}
diff --git a/Engine/Plugins/ChaosClothAssetEditorCore/Source/ChaosClothAssetEditor/ChaosClothAssetEditor.Build.cs b/Engine/Plugins/ChaosClothAssetEditorCore/Source/ChaosClothAssetEditor/ChaosClothAssetEditor.Build.cs
--- a/Engine/Plugins/ChaosClothAssetEditorCore/Source/ChaosClothAssetEditor/ChaosClothAssetEditor.Build.cs
+++ b/Engine/Plugins/ChaosClothAssetEditorCore/Source/ChaosClothAssetEditor/ChaosClothAssetEditor.Build.cs
@@ -5,16 +5,15 @@ using UnrealBuildTool;
 public class ChaosClothAssetEditor : ModuleRules
 {
 	public ChaosClothAssetEditor(ReadOnlyTargetRules Target) : base(Target)
-	{		
+	{
 		PublicDependencyModuleNames.AddRange(
 			new string[]
 			{
 				// ... add other public dependencies that you statically link with here ...
 				
 			}
-			);
-			
-		
+		);
+
 		PrivateDependencyModuleNames.AddRange(
 			new string[]
 			{
@@ -74,6 +73,8 @@ public class ChaosClothAssetEditor : ModuleRules
 				"DeveloperSettings"
 				// ... add private dependencies that you statically link with here ...	
 			}
-			);
+		);
+
+		PrivateIncludePaths.Add(System.IO.Path.Combine(ModuleDirectory, "Internal"));  // For thumbnail renderer
 	}
 }
diff --git a/Engine/Plugins/ChaosClothAssetEditorCore/Source/ChaosClothAssetEditor/Internal/ChaosClothAsset/ChaosClothAssetThumbnailRenderer.h b/Engine/Plugins/ChaosClothAssetEditorCore/Source/ChaosClothAssetEditor/Internal/ChaosClothAsset/ChaosClothAssetThumbnailRenderer.h
new file mode 100644
--- /dev/null
+++ b/Engine/Plugins/ChaosClothAssetEditorCore/Source/ChaosClothAssetEditor/Internal/ChaosClothAsset/ChaosClothAssetThumbnailRenderer.h
@@ -0,0 +1,73 @@
+// Copyright Epic Games, Inc. All Rights Reserved.
+
+#pragma once
+
+#include "HAL/Platform.h"
+#include "ThumbnailRendering/SkeletalMeshThumbnailRenderer.h"
+
+#include "ChaosClothAssetThumbnailRenderer.generated.h"
+
+class UChaosClothComponent;
+class UChaosClothAssetBase;
+
+/**
+ * Minimal actor with a cloth component used to render cloth assets previews
+ */
+UCLASS(MinimalAPI)
+class UE_DEPRECATED(5.8, "This class will be moved to the internal section of the ChaosClothAsset module.") AChaosClothPreviewActor : public AActor
+{
+	GENERATED_BODY()
+public:
+	AChaosClothPreviewActor();
+
+	UChaosClothComponent* GetClothComponent()
+	{
+		return ClothComponent;
+	}
+
+protected:
+	UPROPERTY()
+	TObjectPtr<UChaosClothComponent> ClothComponent;
+};
+
+namespace UE::Chaos::ClothAsset
+{
+	/**
+	 * Preview scene for a Cloth asset thumbnail
+	 */
+	class UE_DEPRECATED(5.8, "This class will be moved to the internal section of the ChaosClothAsset module.") FThumbnailScene : public FThumbnailPreviewScene
+	{
+	public:
+		FThumbnailScene();
+
+		void SetClothAsset(UChaosClothAssetBase* InClothAsset);
+
+		virtual void GetViewMatrixParameters(const float InFOVDegrees, FVector& OutOrigin, float& OutOrbitPitch, float& OutOrbitYaw, float& OutOrbitZoom) const override;
+
+	private:
+PRAGMA_DISABLE_DEPRECATION_WARNINGS  // UE_DEPRECATED(5.8, "This class will be moved to the internal section of the ChaosClothAsset module.") 
+		AChaosClothPreviewActor* PreviewActor;
+PRAGMA_ENABLE_DEPRECATION_WARNINGS
+	};
+}
+
+
+/**
+ * Renders a preview scene to be used as a thumbnail for a supported cloth asset
+ */
+UCLASS(MinimalAPI)
+class UE_DEPRECATED(5.8, "This class will be moved to the internal section of the ChaosClothAsset module.") UChaosClothAssetThumbnailRenderer : public UDefaultSizedThumbnailRenderer
+{
+	GENERATED_BODY()
+
+	virtual void Draw(UObject* Object, int32 X, int32 Y, uint32 Width, uint32 Height, FRenderTarget* Viewport, FCanvas* Canvas, bool bAdditionalViewFamily) override;
+	virtual bool CanVisualizeAsset(UObject* Object) override;
+	virtual EThumbnailRenderFrequency GetThumbnailRenderFrequency(UObject* Object) const override;
+
+	virtual void BeginDestroy() override;
+
+protected:
+PRAGMA_DISABLE_DEPRECATION_WARNINGS  // UE_DEPRECATED(5.8, "This class will be moved to the internal section of the ChaosClothAsset module.") 
+	TObjectInstanceThumbnailScene<UE::Chaos::ClothAsset::FThumbnailScene, 128> ClothThumbnailSceneCache;
+PRAGMA_ENABLE_DEPRECATION_WARNINGS
+};
diff --git a/Engine/Plugins/ChaosClothAssetEditorCore/Source/ChaosClothAssetEditor/Private/ChaosClothAsset/ChaosClothAssetEditorModule.cpp b/Engine/Plugins/ChaosClothAssetEditorCore/Source/ChaosClothAssetEditor/Private/ChaosClothAsset/ChaosClothAssetEditorModule.cpp
--- a/Engine/Plugins/ChaosClothAssetEditorCore/Source/ChaosClothAssetEditor/Private/ChaosClothAsset/ChaosClothAssetEditorModule.cpp
+++ b/Engine/Plugins/ChaosClothAssetEditorCore/Source/ChaosClothAssetEditor/Private/ChaosClothAsset/ChaosClothAssetEditorModule.cpp
@@ -4,7 +4,6 @@
 
 
 #include "ChaosClothAsset/AssetDefinition_ClothAsset.h"
-#include "ChaosClothAsset/ChaosClothAssetThumbnailRenderer.h"
 #include "ChaosClothAsset/ClothAsset.h"
 #include "ChaosClothAsset/ClothComponent.h"
 #include "ChaosClothAsset/ClothDataflowConstructionVisualization.h"
@@ -163,8 +162,6 @@ namespace UE::Chaos::ClothAsset
 
 		ClothAssetComponentBroker = MakeShareable(new FClothAssetComponentBroker);
 		FComponentAssetBrokerage::RegisterBroker(ClothAssetComponentBroker, UChaosClothComponent::StaticClass(), true, true);
-
-		UThumbnailManager::Get().RegisterCustomRenderer(UChaosClothAsset::StaticClass(), UChaosClothAssetThumbnailRenderer::StaticClass());
 	}
 
 	void FChaosClothAssetEditorModule::ShutdownModule()
diff --git a/Engine/Plugins/ChaosClothAssetEditorCore/Source/ChaosClothAssetEditor/Private/ChaosClothAsset/ChaosClothAssetThumbnailRenderer.cpp b/Engine/Plugins/ChaosClothAssetEditorCore/Source/ChaosClothAssetEditor/Private/ChaosClothAsset/ChaosClothAssetThumbnailRenderer.cpp
--- a/Engine/Plugins/ChaosClothAssetEditorCore/Source/ChaosClothAssetEditor/Private/ChaosClothAsset/ChaosClothAssetThumbnailRenderer.cpp
+++ b/Engine/Plugins/ChaosClothAssetEditorCore/Source/ChaosClothAssetEditor/Private/ChaosClothAsset/ChaosClothAssetThumbnailRenderer.cpp
@@ -1,5 +1,8 @@
 ﻿// Copyright Epic Games, Inc. All Rights Reserved.
 
+// UE_DEPRECATED(5.8, "This class will be moved to the internal section of the ChaosClothAsset module.")
+// TODO: Remove these files in 5.10.
+
 #include "ChaosClothAsset/ChaosClothAssetThumbnailRenderer.h"
 
 #include "ChaosClothAsset/ClothAsset.h"
@@ -32,7 +35,9 @@ namespace UE::Chaos::ClothAsset
 		SpawnInfo.bNoFail = true;
 		SpawnInfo.ObjectFlags = RF_Transient;
 
+PRAGMA_DISABLE_DEPRECATION_WARNINGS  // UE_DEPRECATED(5.8, "This class will be moved to the internal section of the ChaosClothAsset module.") 
 		PreviewActor = GetWorld()->SpawnActor<AChaosClothPreviewActor>( SpawnInfo );
+PRAGMA_ENABLE_DEPRECATION_WARNINGS
 		PreviewActor->SetActorEnableCollision(false);
 
 		check(PreviewActor);
@@ -40,7 +45,9 @@ namespace UE::Chaos::ClothAsset
 
 	void FThumbnailScene::SetClothAsset(UChaosClothAssetBase* InClothAsset)
 	{
-		UChaosClothComponent* ClothComponent = PreviewActor->GetClothComponent();
+PRAGMA_DISABLE_DEPRECATION_WARNINGS  // UE_DEPRECATED(5.8, "This class will be moved to the internal section of the ChaosClothAsset module.") 
+		UChaosClothComponent* const ClothComponent = PreviewActor->GetClothComponent();
+PRAGMA_ENABLE_DEPRECATION_WARNINGS
 		check(ClothComponent);
 		
 		ClothComponent->SetAsset(InClothAsset);
@@ -49,7 +56,9 @@ namespace UE::Chaos::ClothAsset
 	void FThumbnailScene::GetViewMatrixParameters(const float InFOVDegrees, FVector& OutOrigin, float& OutOrbitPitch, float& OutOrbitYaw,
 		float& OutOrbitZoom) const
 	{	
-		UChaosClothComponent* ClothComponent = PreviewActor->GetClothComponent();
+PRAGMA_DISABLE_DEPRECATION_WARNINGS  // UE_DEPRECATED(5.8, "This class will be moved to the internal section of the ChaosClothAsset module.") 
+		UChaosClothComponent* const ClothComponent = PreviewActor->GetClothComponent();
+PRAGMA_ENABLE_DEPRECATION_WARNINGS
 		check(ClothComponent);
 
 		FBoxSphereBounds Bounds = ClothComponent->Bounds;
@@ -87,7 +96,9 @@ void UChaosClothAssetThumbnailRenderer::Draw(UObject* Object, int32 X, int32 Y,
 		return;
 	}
 	
+PRAGMA_DISABLE_DEPRECATION_WARNINGS  // UE_DEPRECATED(5.8, "This class will be moved to the internal section of the ChaosClothAsset module.") 
 	TSharedRef<UE::Chaos::ClothAsset::FThumbnailScene> ThumbnailScene = ClothThumbnailSceneCache.EnsureThumbnailScene(Object);
+PRAGMA_ENABLE_DEPRECATION_WARNINGS
 
 	// Check if the cloth asset has render data, if not render the sim mesh instead
 	FSkeletalMeshRenderData* const RenderData = ClothAsset->GetResourceForRendering();
@@ -125,12 +136,16 @@ void UChaosClothAssetThumbnailRenderer::Draw(UObject* Object, int32 X, int32 Y,
 		const TArray<TSharedRef<const FManagedArrayCollection>> ClothCollectionsCopy = { MoveTemp(CollectionCopy) };
 
 		ClothAssetCopy->Build(ClothCollectionsCopy);
+PRAGMA_DISABLE_DEPRECATION_WARNINGS  // UE_DEPRECATED(5.8, "This class will be moved to the internal section of the ChaosClothAsset module.") 
 		ThumbnailScene->SetClothAsset(ClothAssetCopy);
+PRAGMA_ENABLE_DEPRECATION_WARNINGS
 	}
 	else
 	{
 		// Render render mesh
+PRAGMA_DISABLE_DEPRECATION_WARNINGS  // UE_DEPRECATED(5.8, "This class will be moved to the internal section of the ChaosClothAsset module.") 
 		ThumbnailScene->SetClothAsset(ClothAsset);
+PRAGMA_ENABLE_DEPRECATION_WARNINGS
 	}
 
 	FSceneViewFamilyContext ViewFamily( FSceneViewFamily::ConstructionValues( Viewport, ThumbnailScene->GetScene(), FEngineShowFlags(ESFIM_Game) )
@@ -141,8 +156,10 @@ void UChaosClothAssetThumbnailRenderer::Draw(UObject* Object, int32 X, int32 Y,
 	ViewFamily.EngineShowFlags.MotionBlur = 0;
 	ViewFamily.EngineShowFlags.LOD = 0;
 
+PRAGMA_DISABLE_DEPRECATION_WARNINGS  // UE_DEPRECATED(5.8, "This class will be moved to the internal section of the ChaosClothAsset module.") 
 	RenderViewFamily(Canvas, &ViewFamily, ThumbnailScene->CreateView(&ViewFamily, X, Y, Width, Height));
 	ThumbnailScene->SetClothAsset(nullptr);
+PRAGMA_ENABLE_DEPRECATION_WARNINGS
 }
 
 bool UChaosClothAssetThumbnailRenderer::CanVisualizeAsset(UObject* Object)
diff --git a/Engine/Plugins/ChaosClothAssetEditorCore/Source/ChaosClothAssetEditor/Public/ChaosClothAsset/ChaosClothAssetThumbnailRenderer.h b/Engine/Plugins/ChaosClothAssetEditorCore/Source/ChaosClothAssetEditor/Public/ChaosClothAsset/ChaosClothAssetThumbnailRenderer.h
--- a/Engine/Plugins/ChaosClothAssetEditorCore/Source/ChaosClothAssetEditor/Public/ChaosClothAsset/ChaosClothAssetThumbnailRenderer.h
+++ b/Engine/Plugins/ChaosClothAssetEditorCore/Source/ChaosClothAssetEditor/Public/ChaosClothAsset/ChaosClothAssetThumbnailRenderer.h
@@ -2,68 +2,7 @@
 
 #pragma once
 
-#include "HAL/Platform.h"
-#include "ThumbnailRendering/SkeletalMeshThumbnailRenderer.h"
+#include "Misc/CoreMiscDefines.h"
+UE_DEPRECATED_HEADER(5.8, "This header and its internal counterpart are scheduled for removal. Do not include either version.")
 
-#include "ChaosClothAssetThumbnailRenderer.generated.h"
-
-class UChaosClothComponent;
-class UChaosClothAssetBase;
-
-/**
- * Minimal actor with a cloth component used to render cloth assets previews
- */
-UCLASS(MinimalAPI)
-class AChaosClothPreviewActor : public AActor
-{
-	GENERATED_BODY()
-public:
-	AChaosClothPreviewActor();
-
-	UChaosClothComponent* GetClothComponent()
-	{
-		return ClothComponent;
-	}
-
-protected:
-	UPROPERTY()
-	TObjectPtr<UChaosClothComponent> ClothComponent;
-};
-
-namespace UE::Chaos::ClothAsset
-{
-	/**
-	 * Preview scene for a Cloth asset thumbnail
-	 */
-	class FThumbnailScene : public FThumbnailPreviewScene
-	{
-	public:
-		FThumbnailScene();
-
-		void SetClothAsset(UChaosClothAssetBase* InClothAsset);
-
-		virtual void GetViewMatrixParameters(const float InFOVDegrees, FVector& OutOrigin, float& OutOrbitPitch, float& OutOrbitYaw, float& OutOrbitZoom) const override;
-
-	private:
-		AChaosClothPreviewActor* PreviewActor;
-	};
-}
-
-
-/**
- * Renders a preview scene to be used as a thumbnail for a supported cloth asset
- */
-UCLASS(MinimalAPI)
-class UChaosClothAssetThumbnailRenderer : public UDefaultSizedThumbnailRenderer
-{
-	GENERATED_BODY()
-
-	virtual void Draw(UObject* Object, int32 X, int32 Y, uint32 Width, uint32 Height, FRenderTarget* Viewport, FCanvas* Canvas, bool bAdditionalViewFamily) override;
-	virtual bool CanVisualizeAsset(UObject* Object) override;
-	virtual EThumbnailRenderFrequency GetThumbnailRenderFrequency(UObject* Object) const override;
-
-	virtual void BeginDestroy() override;
-
-protected:
-	TObjectInstanceThumbnailScene<UE::Chaos::ClothAsset::FThumbnailScene, 128> ClothThumbnailSceneCache;
-};
+#include "ChaosClothAssetEditor/Internal/ChaosClothAsset/ChaosClothAssetThumbnailRenderer.h"
diff --git a/Engine/Plugins/Experimental/ChaosOutfitAsset/Source/ChaosOutfitAssetEditor/ChaosOutfitAssetEditor.Build.cs b/Engine/Plugins/Experimental/ChaosOutfitAsset/Source/ChaosOutfitAssetEditor/ChaosOutfitAssetEditor.Build.cs
--- a/Engine/Plugins/Experimental/ChaosOutfitAsset/Source/ChaosOutfitAssetEditor/ChaosOutfitAssetEditor.Build.cs
+++ b/Engine/Plugins/Experimental/ChaosOutfitAsset/Source/ChaosOutfitAssetEditor/ChaosOutfitAssetEditor.Build.cs
@@ -15,7 +15,7 @@ public class ChaosOutfitAssetEditor : ModuleRules
 				"BaseCharacterFXEditor",
 				"ChaosClothAsset",
 				"ChaosClothAssetEngine",
-				"ChaosClothAssetEditor",
+				"ChaosClothAssetTools",
 				"ChaosOutfitAssetEngine",
 				"ContentBrowser",
 				"Core",
@@ -33,5 +33,7 @@ public class ChaosOutfitAssetEditor : ModuleRules
 				"UnrealEd"
 			}
 		);
+
+		PrivateIncludePaths.Add(System.IO.Path.Combine(GetModuleDirectory("ChaosClothAssetTools"), "Internal"));  // For thumbnail renderer
 	}
 }
diff --git a/Engine/Plugins/Experimental/ChaosOutfitAsset/Source/ChaosOutfitAssetEditor/Private/ChaosOutfitAsset/ChaosOutfitAssetEditorModule.cpp b/Engine/Plugins/Experimental/ChaosOutfitAsset/Source/ChaosOutfitAssetEditor/Private/ChaosOutfitAsset/ChaosOutfitAssetEditorModule.cpp
--- a/Engine/Plugins/Experimental/ChaosOutfitAsset/Source/ChaosOutfitAssetEditor/Private/ChaosOutfitAsset/ChaosOutfitAssetEditorModule.cpp
+++ b/Engine/Plugins/Experimental/ChaosOutfitAsset/Source/ChaosOutfitAssetEditor/Private/ChaosOutfitAsset/ChaosOutfitAssetEditorModule.cpp
@@ -3,7 +3,7 @@
 #include "ChaosOutfitAsset/AssetDefinition_OutfitAsset.h"
 #include "Algo/AllOf.h"
 #include "AssetRegistry/AssetRegistryModule.h"
-#include "ChaosClothAsset/ChaosClothAssetThumbnailRenderer.h"
+#include "ChaosClothAsset/ClothAssetThumbnailRenderer.h"
 #include "ChaosOutfitAsset/OutfitAsset.h"
 #include "ChaosOutfitAsset/OutfitAssetEditorStyle.h"
 #include "ChaosOutfitAsset/OutfitEditorCommands.h"
@@ -38,7 +38,7 @@ namespace UE::Chaos::OutfitAsset
 			StartupCallbackDelegateHandle = UToolMenus::RegisterStartupCallback(
 				FSimpleMulticastDelegate::FDelegate::CreateRaw(this, &FOutfitAssetEditorModule::RegisterMenus));
 
-			UThumbnailManager::Get().RegisterCustomRenderer(UChaosOutfitAsset::StaticClass(), UChaosClothAssetThumbnailRenderer::StaticClass());
+			UThumbnailManager::Get().RegisterCustomRenderer(UChaosOutfitAsset::StaticClass(), UChaosClothAssetThumbnailRenderer_Internal::StaticClass());
 
 			// TODO Register visualizations (see FChaosClothAssetEditorModule::StartupModule() in ChaosClothAssetEditorModule.cpp)
 		}
@@ -51,6 +51,8 @@ namespace UE::Chaos::OutfitAsset
 
 				UToolMenus::UnRegisterStartupCallback(StartupCallbackDelegateHandle);
 				FOutfitEditorCommands::Unregister();
+
+				UThumbnailManager::Get().UnregisterCustomRenderer(UChaosOutfitAsset::StaticClass());
 			}
 		}
 		//~ End IModuleInterface interface
