Chaos Cloth Asset - Fixed cook data determinism caused by a call to SetNumSimAccessoryMeshes() not using the cooked only flag.

[RN] Chaos Cloth Asset - Fixed cook data determinism caused by a call to SetNumSimAccessoryMeshes() not using the cooked only flag.

#Synced-CL 50202394

diff --git a/Engine/Plugins/ChaosClothAsset/Source/ChaosClothAsset/Private/ChaosClothAsset/CollectionClothFacade.cpp b/Engine/Plugins/ChaosClothAsset/Source/ChaosClothAsset/Private/ChaosClothAsset/CollectionClothFacade.cpp
--- a/Engine/Plugins/ChaosClothAsset/Source/ChaosClothAsset/Private/ChaosClothAsset/CollectionClothFacade.cpp
+++ b/Engine/Plugins/ChaosClothAsset/Source/ChaosClothAsset/Private/ChaosClothAsset/CollectionClothFacade.cpp
@@ -989,9 +989,9 @@ namespace UE::Chaos::ClothAsset
 		FClothCollection::CopyArrayViewData(GetPreResizedSimPosition3D().Right(Other.GetNumSimVertices3D()), Other.GetPreResizedSimPosition3D());
 
 		// Sim AccessoryMeshes
-		SetNumSimAccessoryMeshes(0);
+		SetNumSimAccessoryMeshes(0, bCookedOnly);
 		const int32 OtherNumSimAccessoryMeshes = Other.GetNumSimAccessoryMeshes();
-		SetNumSimAccessoryMeshes(OtherNumSimAccessoryMeshes);
+		SetNumSimAccessoryMeshes(OtherNumSimAccessoryMeshes, bCookedOnly);
 		for (int32 MeshIndex = 0; MeshIndex < OtherNumSimAccessoryMeshes; ++MeshIndex)
 		{
 			GetSimAccessoryMesh(MeshIndex).Initialize(Other.GetSimAccessoryMesh(MeshIndex));
