Chaos Cloth Asset - Added console variable to disable thumbnail generation. This is useful for recovering Cloth or Outfit Assets with corrupted render data.

#rnx

#Synced-CL 50258746

diff --git a/Engine/Plugins/ChaosClothAsset/Source/ChaosClothAssetTools/Private/ChaosClothAsset/ClothAssetThumbnailRenderer.cpp b/Engine/Plugins/ChaosClothAsset/Source/ChaosClothAssetTools/Private/ChaosClothAsset/ClothAssetThumbnailRenderer.cpp
--- a/Engine/Plugins/ChaosClothAsset/Source/ChaosClothAssetTools/Private/ChaosClothAsset/ClothAssetThumbnailRenderer.cpp
+++ b/Engine/Plugins/ChaosClothAsset/Source/ChaosClothAssetTools/Private/ChaosClothAsset/ClothAssetThumbnailRenderer.cpp
@@ -8,6 +8,7 @@
 #include "ChaosClothAsset/ClothGeometryTools.h"
 #include "Components/SkeletalMeshComponent.h"
 #include "Engine/SkeletalMesh.h"
+#include "HAL/IConsoleManager.h"
 #include "Rendering/SkeletalMeshRenderData.h"
 #include "ThumbnailRendering/SceneThumbnailInfo.h"
 #include "SceneView.h"
@@ -23,6 +24,11 @@ AChaosClothPreviewActor_Internal::AChaosClothPreviewActor_Internal()
 
 namespace UE::Chaos::ClothAsset
 {
+	namespace Private
+	{
+		static TAutoConsoleVariable<bool> CVarThumbnailEnable(TEXT("p.ChaosClothAsset.Thumbnail.Enable"), true, TEXT("Enable Cloth Asset and Outfit Asset thumbnail generation.\n Default: True."));
+	}
+
 	FThumbnailScene_Internal::FThumbnailScene_Internal() : FThumbnailPreviewScene()
 	{
 		bForceAllUsedMipsResident = false;
@@ -161,7 +167,9 @@ void UChaosClothAssetThumbnailRenderer_Internal::Draw(UObject* Object, int32 X,
 
 bool UChaosClothAssetThumbnailRenderer_Internal::CanVisualizeAsset(UObject* Object)
 {
-	return Cast<UChaosClothAssetBase>(Object) != nullptr;
+	return
+		UE::Chaos::ClothAsset::Private::CVarThumbnailEnable.GetValueOnAnyThread() &&
+		Cast<UChaosClothAssetBase>(Object) != nullptr;
 }
 
 EThumbnailRenderFrequency UChaosClothAssetThumbnailRenderer_Internal::GetThumbnailRenderFrequency(UObject* Object) const
