Cloth Asset - Added new Static Mesh and Skeletal Mesh templates for the creation of new Cloth Assets.

[RN] Cloth Asset - Added new Static Mesh and Skeletal Mesh templates for the creation of new Cloth Assets.

#virtualized
#Synced-CL 50258746

diff --git a/Engine/Plugins/ChaosClothAsset/Source/ChaosClothAssetTools/ChaosClothAssetTools.Build.cs b/Engine/Plugins/ChaosClothAsset/Source/ChaosClothAssetTools/ChaosClothAssetTools.Build.cs
--- a/Engine/Plugins/ChaosClothAsset/Source/ChaosClothAssetTools/ChaosClothAssetTools.Build.cs
+++ b/Engine/Plugins/ChaosClothAsset/Source/ChaosClothAssetTools/ChaosClothAssetTools.Build.cs
@@ -29,6 +29,7 @@ public class ChaosClothAssetTools : ModuleRules
 				"DataflowEditor",
 				"DataflowEngine",
 				"SlateCore",
+				"ToolWidgets",
 				"UnrealEd",
 				"Projects",
 				"ClothPainter",                 // For clothing asset exporter
diff --git a/Engine/Plugins/ChaosClothAsset/Source/ChaosClothAssetTools/Private/ChaosClothAsset/ClothAssetFactory.cpp b/Engine/Plugins/ChaosClothAsset/Source/ChaosClothAssetTools/Private/ChaosClothAsset/ClothAssetFactory.cpp
--- a/Engine/Plugins/ChaosClothAsset/Source/ChaosClothAssetTools/Private/ChaosClothAsset/ClothAssetFactory.cpp
+++ b/Engine/Plugins/ChaosClothAsset/Source/ChaosClothAssetTools/Private/ChaosClothAsset/ClothAssetFactory.cpp
@@ -3,6 +3,7 @@
 #include "ChaosClothAsset/ClothAssetFactory.h"
 #include "ChaosClothAsset/ClothAsset.h"
 #include "Dataflow/AssetDefinition_DataflowAsset.h"
+#include "Dialog/SMessageDialog.h"
 #include "HAL/IConsoleManager.h"
 
 #include UE_INLINE_GENERATED_CPP_BY_NAME(ClothAssetFactory)
@@ -31,10 +32,56 @@ UObject* UChaosClothAssetFactory::FactoryCreateNew(UClass* Class, UObject* Paren
 {
 	using namespace UE::Chaos::ClothAsset;
 
-	static const FString TemplatePath = TEXT("/ChaosClothAsset/ClothAssetTemplate.ClothAssetTemplate");
+	// Select the template
+	const bool bHaveDataflowNodes = FModuleManager::Get().IsModuleLoaded(TEXT("ChaosClothAssetDataflowNodes"));
+	const bool bHaveUsdDataflowNodes = FModuleManager::Get().IsModuleLoaded(TEXT("ChaosClothAssetUsdDataflowNodes"));
+
+	TSharedRef<SMessageDialog> SelectTemplateMessageDialog = SNew(SMessageDialog)
+		.Title(FText(LOCTEXT("SelectTemplateTitle", "Select a Cloth Asset Template")))
+		.Message(LOCTEXT("SelectTemplateMessage", "Select a template for this newly created Cloth Asset:"))
+		.Buttons(
+			{
+				SMessageDialog::FButton(LOCTEXT("SelectNoTemplate", "No Dataflow"))
+					.SetToolTipText(LOCTEXT("SelectNoTemplateTooltip", "Don't add a Dataflow to this Cloth Asset. Useful when a Dataflow already exists."))
+					.SetPrimary(!bHaveDataflowNodes),
+				SMessageDialog::FButton(LOCTEXT("SelectEmptyTemplate", "Empty Dataflow"))
+					.SetToolTipText(LOCTEXT("SelectEmptyTemplateTooltip", "Add an empty Dataflow graph with a Cloth Asset Terminal node.\n(Requires the ChaosClothAssetDataflowNodes plug-in.)"))
+					.SetPrimary(false)
+					.SetIsEnabled(bHaveDataflowNodes),
+				SMessageDialog::FButton(LOCTEXT("SelectStaticMeshClothTemplate", "Static Mesh Cloth"))
+					.SetToolTipText(LOCTEXT("SelectStaticMeshClothTemplateTooltip", "Add a Dataflow graph that builds a Cloth Asset from Static Meshes.\n(Requires the ChaosClothAssetDataflowNodes plug-in.)"))
+					.SetPrimary(false)
+					.SetIsEnabled(bHaveDataflowNodes),
+				SMessageDialog::FButton(LOCTEXT("SelectSkeletalMeshClothTemplate", "Skeletal Mesh Cloth"))
+					.SetToolTipText(LOCTEXT("SelectSkeletalMeshClothTemplateTooltip", "Add a Dataflow graph that builds a Cloth Asset from Skeletal Meshes.\n(Requires the ChaosClothAssetDataflowNodes plug-in.)"))
+					.SetPrimary(bHaveDataflowNodes && !bHaveUsdDataflowNodes)
+					.SetIsEnabled(bHaveDataflowNodes),
+				SMessageDialog::FButton(LOCTEXT("SelectUsdClothTemplate", "USD Cloth"))
+					.SetToolTipText(LOCTEXT("SelectUsdClothTemplateTooltip", "Add a Dataflow graph that builds a Cloth Asset from an imported USD Cloth.\n(Requires the ChaosClothAssetDataflowNodes and ChaosClothAssetUsdDataflowNodes plug-ins.)"))
+					.SetPrimary(bHaveDataflowNodes && bHaveUsdDataflowNodes)
+					.SetIsEnabled(bHaveDataflowNodes && bHaveUsdDataflowNodes)
+			});
+
+	const FString* TemplatePath;
+	static const FString EmptyClothTemplatePath(TEXT("/ChaosClothAsset/DF_EmptyClothAssetTemplate.DF_EmptyClothAssetTemplate"));
+	static const FString StaticMeshClothTemplatePath(TEXT("/ChaosClothAsset/DF_StaticMeshClothTemplate.DF_StaticMeshClothTemplate"));
+	static const FString SkeletalMeshClothTemplatePath(TEXT("/ChaosClothAsset/DF_SkeletalMeshClothTemplate.DF_SkeletalMeshClothTemplate"));
+	static const FString UsdClothTemplatePath(TEXT("/ChaosClothAsset/ClothAssetTemplate.ClothAssetTemplate"));
+
+	switch (SelectTemplateMessageDialog->ShowModal())
+	{
+	default:
+	case 0: TemplatePath = nullptr; break;  // No Dataflow
+	case 1: TemplatePath = &EmptyClothTemplatePath; break;  // Empty Dataflow with terminal node
+	case 2: TemplatePath = &StaticMeshClothTemplatePath; break;  // Static Mesh Cloth
+	case 3: TemplatePath = &SkeletalMeshClothTemplatePath; break;  // Skeletal Mesh Cloth
+	case 4: TemplatePath = &UsdClothTemplatePath; break;  // USD Cloth
+	}
+
+	// Create a new asset using the provided template path
 	static const TCHAR* AssetPrefix = TEXT("CA_");
 	const bool bEmbedDataflow = Private::CVarCreateWithEmbeddedDataflow.GetValueOnAnyThread();
-	return UE::DataflowAssetDefinitionHelpers::FactoryCreateNew(Class, Parent, Name, Flags, &TemplatePath, bEmbedDataflow, AssetPrefix);
+	return UE::DataflowAssetDefinitionHelpers::FactoryCreateNew(Class, Parent, Name, Flags, TemplatePath, bEmbedDataflow, AssetPrefix);
 }
 
 FString UChaosClothAssetFactory::GetDefaultNewAssetName() const
