Chaos Cloth Asset - Fixed crash when proxy is simulating while a sim model changes.

[RN] Chaos Cloth Asset - Fixed crash when proxy is simulating while a sim model changes.

diff --git a/Engine/Plugins/ChaosClothAsset/Source/ChaosClothAssetEngine/Private/ChaosClothAsset/ClothSimulationProxy.cpp b/Engine/Plugins/ChaosClothAsset/Source/ChaosClothAssetEngine/Private/ChaosClothAsset/ClothSimulationProxy.cpp
--- a/Engine/Plugins/ChaosClothAsset/Source/ChaosClothAssetEngine/Private/ChaosClothAsset/ClothSimulationProxy.cpp
+++ b/Engine/Plugins/ChaosClothAsset/Source/ChaosClothAssetEngine/Private/ChaosClothAsset/ClothSimulationProxy.cpp
@@ -170,6 +170,7 @@ namespace UE::Chaos::ClothAsset
 		Meshes.Reset();
 		Colliders.Reset();
 		Cloths.Reset();
+		ClothSimulationModels.Reset();
 
 		// Create solver config simulation thread object first. Need to know which solver type we're creating.
 		const int32 SolverConfigIndex = Configs.Emplace(MakeUnique<FClothingSimulationConfig>(ClothComponentAdapter.GetSolverPropertyCollections()));
@@ -208,6 +209,13 @@ namespace UE::Chaos::ClothAsset
 						}), TEXT("The ClothComponentAdapter cannot return the same SimulationGroupId for two different Asset/ModelIndex pairs (unless INDEX_NONE)."));
 
 					const TSharedPtr<const FChaosClothSimulationModel> ClothSimulationModel = Asset->GetClothSimulationModel(ModelIndex);
+					if (Algo::NoneOf(ClothSimulationModels, [&ClothSimulationModel](const TSharedPtr<const FChaosClothSimulationModel>& InClothSimulationModel)
+						{
+							return InClothSimulationModel.Get() == ClothSimulationModel.Get();
+						}))
+					{
+						ClothSimulationModels.Emplace(ClothSimulationModel);  // Ownership has to transfer to this proxy when the model changes during the simulation
+					}
 
 					FString DebugName;
 #if !UE_BUILD_SHIPPING
