Chaos Cloth Asset - Defaulted the Cloth Asset creation to an empty Dataflow graph with Terminal node when the template selection modal box fails.

#rnx
#tests Passed Physics.AssetValidation.ChaosCloth.ChaosClothAssetCreationTest

#Synced-CL 50399480

diff --git a/Engine/Plugins/ChaosClothAsset/Source/ChaosClothAssetTools/Private/ChaosClothAsset/ClothAssetFactory.cpp b/Engine/Plugins/ChaosClothAsset/Source/ChaosClothAssetTools/Private/ChaosClothAsset/ClothAssetFactory.cpp
--- a/Engine/Plugins/ChaosClothAsset/Source/ChaosClothAssetTools/Private/ChaosClothAsset/ClothAssetFactory.cpp
+++ b/Engine/Plugins/ChaosClothAsset/Source/ChaosClothAssetTools/Private/ChaosClothAsset/ClothAssetFactory.cpp
@@ -70,8 +70,8 @@ UObject* UChaosClothAssetFactory::FactoryCreateNew(UClass* Class, UObject* Paren
 
 	switch (SelectTemplateMessageDialog->ShowModal())
 	{
-	default:
 	case 0: TemplatePath = nullptr; break;  // No Dataflow
+	default:  // Must use the empty dataflow to validate engine tests when the message dialog fails to display and return -1
 	case 1: TemplatePath = &EmptyClothTemplatePath; break;  // Empty Dataflow with terminal node
 	case 2: TemplatePath = &StaticMeshClothTemplatePath; break;  // Static Mesh Cloth
 	case 3: TemplatePath = &SkeletalMeshClothTemplatePath; break;  // Skeletal Mesh Cloth
