Chaos Cloth Asset - Replace bHasDataflowAsset with UFUNCTION getter HasDataflowAsset to fix the edit conditions that don't otherwise update when an asset is first created.

The behavior got broken after starting to use the IDataflowInstanceInterface to set the Dataflow assets, instead of the UClothAssetBase::SetDataflow() method that was also setting the edit condition boolean.
This fix is possible since simple boolean returning UFUNCTIOn can now be used in an EditCondition.

#rnx

#Synced-CL 51673482

diff --git a/Engine/Plugins/ChaosClothAsset/Source/ChaosClothAssetEngine/Private/ChaosClothAsset/ClothAssetBase.cpp b/Engine/Plugins/ChaosClothAsset/Source/ChaosClothAssetEngine/Private/ChaosClothAsset/ClothAssetBase.cpp
--- a/Engine/Plugins/ChaosClothAsset/Source/ChaosClothAssetEngine/Private/ChaosClothAsset/ClothAssetBase.cpp
+++ b/Engine/Plugins/ChaosClothAsset/Source/ChaosClothAssetEngine/Private/ChaosClothAsset/ClothAssetBase.cpp
@@ -370,8 +370,6 @@ void UChaosClothAssetBase::PostLoad()
 	GetDataflowInstance().PostLoad();
 
 #if WITH_EDITORONLY_DATA
-	bHasDataflowAsset = (GetDataflowInstance().GetDataflowAsset() != nullptr);
-
 	// Sync preview assets after moving from the old preview assets to the new BP variable system
 	SyncBPPreviewAsset(PreviewSceneSkeletalMesh, TEXT("SkeletalMesh"));  // TODO: Deprecate PreviewSceneSkeletalMesh and PreviewSceneAnimation
 	SyncBPPreviewAsset(PreviewSceneAnimation, TEXT("AnimationAsset"));   //       and remove these calls once the old Cloth Editor is retired
@@ -406,11 +404,6 @@ PRAGMA_DISABLE_DEPRECATION_WARNINGS
 	{
 		ReregisterComponents();
 	}
-	else if (PropertyChangedEvent.GetMemberPropertyName() == GET_MEMBER_NAME_CHECKED(UChaosClothAssetBase, DataflowInstance) && 
-		PropertyChangedEvent.GetPropertyName() == FName(TEXT("DataflowAsset")))  // Can't use GET_MEMBER_NAME_CHECKED, because DataflowAsset is private
-	{
-		bHasDataflowAsset = (GetDataflowInstance().GetDataflowAsset() != nullptr);
-	}
 PRAGMA_ENABLE_DEPRECATION_WARNINGS
 	// This invalidates the IDataflowContentOwner, not the Dataflow itself
 	InvalidateDataflowContents();
@@ -563,10 +556,6 @@ void UChaosClothAssetBase::SetDataflow(UDataflow* InDataflow)
 PRAGMA_DISABLE_DEPRECATION_WARNINGS
 	DataflowInstance.SetDataflowAsset(InDataflow);
 PRAGMA_ENABLE_DEPRECATION_WARNINGS
-
-#if WITH_EDITORONLY_DATA
-	bHasDataflowAsset = (GetDataflowInstance().GetDataflowAsset() != nullptr);
-#endif
 }
 
 UDataflow* UChaosClothAssetBase::GetDataflow()
diff --git a/Engine/Plugins/ChaosClothAsset/Source/ChaosClothAssetEngine/Public/ChaosClothAsset/ClothAsset.h b/Engine/Plugins/ChaosClothAsset/Source/ChaosClothAssetEngine/Public/ChaosClothAsset/ClothAsset.h
--- a/Engine/Plugins/ChaosClothAsset/Source/ChaosClothAssetEngine/Public/ChaosClothAsset/ClothAsset.h
+++ b/Engine/Plugins/ChaosClothAsset/Source/ChaosClothAssetEngine/Public/ChaosClothAsset/ClothAsset.h
@@ -249,11 +249,11 @@ private:
 	 * This is of limited use since this USkeleton's reference skeleton might not necessarily match the one created for this asset.
 	 * Set by the Dataflow evaluation.
 	 */
-	UPROPERTY(EditAnywhere, Setter = SetSkeleton, Category = Skeleton, Meta = (EditCondition = "!bHasDataflowAsset"))
+	UPROPERTY(EditAnywhere, Setter = SetSkeleton, Category = Skeleton, Meta = (EditCondition = "!HasDataflow()"))
 	TObjectPtr<USkeleton> Skeleton;
 
 	/** Physics asset used for collision. Set by the Dataflow evaluation. */
-	UPROPERTY(EditAnywhere, Category = Collision, Meta = (EditCondition = "!bHasDataflowAsset"))
+	UPROPERTY(EditAnywhere, Category = Collision, Meta = (EditCondition = "!HasDataflow()"))
 	TObjectPtr<UPhysicsAsset> PhysicsAsset;
 
 	/** Whether to blend positions between the skinned/simulated transitions of the cloth render mesh. Only used when not overriden in Dataflow by the ProxyDefomer node. */
diff --git a/Engine/Plugins/ChaosClothAsset/Source/ChaosClothAssetEngine/Public/ChaosClothAsset/ClothAssetBase.h b/Engine/Plugins/ChaosClothAsset/Source/ChaosClothAssetEngine/Public/ChaosClothAsset/ClothAssetBase.h
--- a/Engine/Plugins/ChaosClothAsset/Source/ChaosClothAssetEngine/Public/ChaosClothAsset/ClothAssetBase.h
+++ b/Engine/Plugins/ChaosClothAsset/Source/ChaosClothAssetEngine/Public/ChaosClothAsset/ClothAssetBase.h
@@ -257,6 +257,16 @@ public:
 	/** Return the Dataflow graph asset associated to this asset if any, const version. */
 	UE_API const UDataflow* GetDataflow() const;
 
+	/**
+	 * Return whether there is a Dataflow graph associated with this asset.
+	 * This function can be used in EditCondition.
+	 */
+	UFUNCTION()
+	bool HasDataflow() const
+	{
+		return GetDataflow() != nullptr;
+	}
+
 #if WITH_EDITORONLY_DATA
 	UE_API void SetPreviewSceneSkeletalMesh(class USkeletalMesh* Mesh);
 	UE_API class USkeletalMesh* GetPreviewSceneSkeletalMesh() const;
@@ -335,7 +345,7 @@ protected:
 	FDataflowInstance DataflowInstance;
 
 	/** List of materials for this cloth asset. Set by the Dataflow evaluation. */
-	UPROPERTY(EditAnywhere, Category = Materials, Meta = (EditCondition = "!bHasDataflowAsset"))
+	UPROPERTY(EditAnywhere, Category = Materials, Meta = (EditCondition = "!HasDataflow()"))
 	TArray<FSkeletalMaterial> Materials;
 
 	/** Struct containing information for each LOD level, such as materials to use, and when use the LOD. Not currently editable or customizable through the Dataflow. */
@@ -387,8 +397,8 @@ protected:
 	FReferenceSkeleton RefSkeleton;
 
 #if WITH_EDITORONLY_DATA
-	/** Property used in various edit conditions to enable property changes when this asset has no Dataflow asset. */
-	UPROPERTY(VisibleAnywhere, Category = Dataflow)
+	UE_DEPRECATED(5.8, "Use HasDataflow instead.")
+	UPROPERTY()
 	bool bHasDataflowAsset = false;
 #endif
 
