Cloth : Use the new template user experience
- Register cloth templates using the dataflow template registry

#Synced-CL 52448897

diff --git a/Engine/Plugins/ChaosClothAsset/Source/ChaosClothAssetTools/Private/ChaosClothAsset/ClothAssetFactory.cpp b/Engine/Plugins/ChaosClothAsset/Source/ChaosClothAssetTools/Private/ChaosClothAsset/ClothAssetFactory.cpp
--- a/Engine/Plugins/ChaosClothAsset/Source/ChaosClothAssetTools/Private/ChaosClothAsset/ClothAssetFactory.cpp
+++ b/Engine/Plugins/ChaosClothAsset/Source/ChaosClothAssetTools/Private/ChaosClothAsset/ClothAssetFactory.cpp
@@ -12,14 +12,6 @@
 
 DEFINE_LOG_CATEGORY_STATIC(LogChaosClothAssetFactory, Log, All);
 
-namespace UE::Chaos::ClothAsset::Private
-{
-	static TAutoConsoleVariable<bool> CVarCreateWithEmbeddedDataflow(
-		TEXT("p.ChaosClothAsset.CreateWithEmbeddedDataflow"),
-		true,
-		TEXT("Set how to create new Cloth Assets, either with embedded (default) or external (legacy) Dataflow assets.\n"));
-}
-
 UChaosClothAssetFactory::UChaosClothAssetFactory(const FObjectInitializer& ObjectInitializer)
 	: Super(ObjectInitializer)
 {
@@ -30,58 +22,18 @@ UChaosClothAssetFactory::UChaosClothAssetFactory(const FObjectInitializer& Objec
 
 UObject* UChaosClothAssetFactory::FactoryCreateNew(UClass* Class, UObject* Parent, FName Name, EObjectFlags Flags, UObject* /*Context*/, FFeedbackContext* /*Warn*/)
 {
-	using namespace UE::Chaos::ClothAsset;
-
-	// Select the template
-	const bool bHaveDataflowNodes = FModuleManager::Get().IsModuleLoaded(TEXT("ChaosClothAssetDataflowNodes"));
-	const bool bHaveUsdDataflowNodes = FModuleManager::Get().IsModuleLoaded(TEXT("ChaosClothAssetUsdDataflowNodes"));
-
-	TSharedRef<SMessageDialog> SelectTemplateMessageDialog = SNew(SMessageDialog)
-		.Title(FText(LOCTEXT("SelectTemplateTitle", "Select a Cloth Asset Template")))
-		.Message(LOCTEXT("SelectTemplateMessage", "Select a template for this newly created Cloth Asset:"))
-		.Buttons(
-			{
-				SMessageDialog::FButton(LOCTEXT("SelectNoTemplate", "No Dataflow"))
-					.SetToolTipText(LOCTEXT("SelectNoTemplateTooltip", "Don't add a Dataflow to this Cloth Asset. Useful when a Dataflow already exists."))
-					.SetPrimary(!bHaveDataflowNodes),
-				SMessageDialog::FButton(LOCTEXT("SelectEmptyTemplate", "Empty Dataflow"))
-					.SetToolTipText(LOCTEXT("SelectEmptyTemplateTooltip", "Add an empty Dataflow graph with a Cloth Asset Terminal node.\n(Requires the ChaosClothAssetDataflowNodes plug-in.)"))
-					.SetPrimary(false)
-					.SetIsEnabled(bHaveDataflowNodes),
-				SMessageDialog::FButton(LOCTEXT("SelectStaticMeshClothTemplate", "Static Mesh Cloth"))
-					.SetToolTipText(LOCTEXT("SelectStaticMeshClothTemplateTooltip", "Add a Dataflow graph that builds a Cloth Asset from Static Meshes.\n(Requires the ChaosClothAssetDataflowNodes plug-in.)"))
-					.SetPrimary(false)
-					.SetIsEnabled(bHaveDataflowNodes),
-				SMessageDialog::FButton(LOCTEXT("SelectSkeletalMeshClothTemplate", "Skeletal Mesh Cloth"))
-					.SetToolTipText(LOCTEXT("SelectSkeletalMeshClothTemplateTooltip", "Add a Dataflow graph that builds a Cloth Asset from Skeletal Meshes.\n(Requires the ChaosClothAssetDataflowNodes plug-in.)"))
-					.SetPrimary(bHaveDataflowNodes && !bHaveUsdDataflowNodes)
-					.SetIsEnabled(bHaveDataflowNodes),
-				SMessageDialog::FButton(LOCTEXT("SelectUsdClothTemplate", "USD Cloth"))
-					.SetToolTipText(LOCTEXT("SelectUsdClothTemplateTooltip", "Add a Dataflow graph that builds a Cloth Asset from an imported USD Cloth.\n(Requires the ChaosClothAssetDataflowNodes and ChaosClothAssetUsdDataflowNodes plug-ins.)"))
-					.SetPrimary(bHaveDataflowNodes && bHaveUsdDataflowNodes)
-					.SetIsEnabled(bHaveDataflowNodes && bHaveUsdDataflowNodes)
-			});
-
-	const FString* TemplatePath;
-	static const FString EmptyClothTemplatePath(TEXT("/ChaosClothAsset/DF_EmptyClothAssetTemplate.DF_EmptyClothAssetTemplate"));
-	static const FString StaticMeshClothTemplatePath(TEXT("/ChaosClothAsset/DF_StaticMeshClothTemplate.DF_StaticMeshClothTemplate"));
-	static const FString SkeletalMeshClothTemplatePath(TEXT("/ChaosClothAsset/DF_SkeletalMeshClothTemplate.DF_SkeletalMeshClothTemplate"));
-	static const FString UsdClothTemplatePath(TEXT("/ChaosClothAsset/ClothAssetTemplate.ClothAssetTemplate"));
-
-	switch (SelectTemplateMessageDialog->ShowModal())
+	UObject* const NewAsset = NewObject<UObject>(Parent, Class, Name, Flags | RF_Transactional | RF_Public | RF_Standalone);
+	if (!NewAsset)
 	{
-	case 0: TemplatePath = nullptr; break;  // No Dataflow
-	default:  // Must use the empty dataflow to validate engine tests when the message dialog fails to display and return -1
-	case 1: TemplatePath = &EmptyClothTemplatePath; break;  // Empty Dataflow with terminal node
-	case 2: TemplatePath = &StaticMeshClothTemplatePath; break;  // Static Mesh Cloth
-	case 3: TemplatePath = &SkeletalMeshClothTemplatePath; break;  // Skeletal Mesh Cloth
-	case 4: TemplatePath = &UsdClothTemplatePath; break;  // USD Cloth
+		const FText ErrorMessage = FText::Format(LOCTEXT("ErrorCreatingNewAsset", "Failed to create a new asset {0}."), FText::FromName(Name));
+		UE_LOG(LogChaosClothAssetFactory, Error, TEXT("%s"), *ErrorMessage.ToString());
+		return nullptr;
 	}
+	NewAsset->MarkPackageDirty();
 
-	// Create a new asset using the provided template path
-	static const TCHAR* AssetPrefix = TEXT("CA_");
-	const bool bEmbedDataflow = Private::CVarCreateWithEmbeddedDataflow.GetValueOnAnyThread();
-	return UE::DataflowAssetDefinitionHelpers::FactoryCreateNew(Class, Parent, Name, Flags, TemplatePath, bEmbedDataflow, AssetPrefix);
+	// No default blank option since cloth templates already provide one
+	UE::DataflowAssetDefinitionHelpers::SetDataflowFromTemplatePicker(NewAsset, /*bShowDefaultBlankOption*/ false);
+	return NewAsset;
 }
 
 FString UChaosClothAssetFactory::GetDefaultNewAssetName() const
diff --git a/Engine/Plugins/ChaosClothAssetEditorCore/Source/ChaosClothAssetEditorTools/Private/ChaosClothAsset/ChaosClothAssetEditorToolsModule.cpp b/Engine/Plugins/ChaosClothAssetEditorCore/Source/ChaosClothAssetEditorTools/Private/ChaosClothAsset/ChaosClothAssetEditorToolsModule.cpp
--- a/Engine/Plugins/ChaosClothAssetEditorCore/Source/ChaosClothAssetEditorTools/Private/ChaosClothAsset/ChaosClothAssetEditorToolsModule.cpp
+++ b/Engine/Plugins/ChaosClothAssetEditorCore/Source/ChaosClothAssetEditorTools/Private/ChaosClothAsset/ChaosClothAssetEditorToolsModule.cpp
@@ -3,12 +3,15 @@
 #include "Modules/ModuleManager.h"
 
 #include "Dataflow/DataflowToolRegistry.h"
+#include "Dataflow/DataflowTemplateRegistry.h"
 #include "DataflowEditorTools/DataflowEditorVertexAttributePaintTool.h"
+#include "ChaosClothAsset/ClothAsset.h"
 #include "ChaosClothAsset/ClothEditorToolBuilders.h"
 #include "ChaosClothAsset/SelectionNode.h"
 #include "ChaosClothAsset/TransferSkinWeightsNode.h"
 #include "ChaosClothAsset/WeightMapNode.h"
 #include "ChaosClothAsset/ClothToolActionCommandBindings.h"
+#include "Textures/SlateIcon.h"
 
 #define LOCTEXT_NAMESPACE "ChaosClothAssetEditorToolsModule"
 
@@ -44,6 +47,42 @@ namespace UE::Chaos::ClothAsset
 				ClothToolActions, 
 				FSlateIcon(FName("ClothStyle"), FName("ChaosClothAssetEditor.AddTransferSkinWeightsNode")),
 				LOCTEXT("AddTransferSkinWeightNodeButtonText", "Cloth Skinning Transfer"), FName("Cloth"));
+
+			const FSlateIcon ChaosClothAssetIcon("ClothAssetEditorStyle", "ClassThumbnail.ChaosClothAsset");
+
+			FDataflowTemplateRegistry::Get().RegisterTemplateAsset(
+				UChaosClothAsset::StaticClass(),
+				{
+					.AssetPath = FSoftObjectPath(TEXT("/ChaosClothAsset/DF_EmptyClothAssetTemplate.DF_EmptyClothAssetTemplate")),
+					.DisplayName = LOCTEXT("SelectEmptyTemplate", "Empty Dataflow"),
+					.Tooltip = LOCTEXT("SelectEmptyTemplateTooltip", "Add an empty Dataflow graph with a Cloth Asset Terminal node.\n(Requires the ChaosClothAssetDataflowNodes plug-in.)"),
+					 .Icon = ChaosClothAssetIcon
+				});
+
+			FDataflowTemplateRegistry::Get().RegisterTemplateAsset(
+				UChaosClothAsset::StaticClass(),
+				{
+					.AssetPath = FSoftObjectPath(TEXT("/ChaosClothAsset/DF_StaticMeshClothTemplate.DF_StaticMeshClothTemplate")),
+					.DisplayName = LOCTEXT("SelectStaticMeshClothTemplate", "Static Mesh Cloth"),
+					.Tooltip = LOCTEXT("SelectStaticMeshClothTemplateTooltip", "Add a Dataflow graph that builds a Cloth Asset from Static Meshes.\n(Requires the ChaosClothAssetDataflowNodes plug-in.)"),
+					 .Icon = ChaosClothAssetIcon
+				});
+			FDataflowTemplateRegistry::Get().RegisterTemplateAsset(
+				UChaosClothAsset::StaticClass(),
+				{
+					.AssetPath = FSoftObjectPath(TEXT("/ChaosClothAsset/DF_SkeletalMeshClothTemplate.DF_SkeletalMeshClothTemplate")),
+					.DisplayName = LOCTEXT("SelectSkeletalMeshClothTemplate", "Skeletal Mesh Cloth"),
+					.Tooltip = LOCTEXT("SelectSkeletalMeshClothTemplateTooltip", "Add a Dataflow graph that builds a Cloth Asset from Skeletal Meshes.\n(Requires the ChaosClothAssetDataflowNodes plug-in.)"),
+					 .Icon = ChaosClothAssetIcon
+				});
+			FDataflowTemplateRegistry::Get().RegisterTemplateAsset(
+				UChaosClothAsset::StaticClass(),
+				{
+					.AssetPath = FSoftObjectPath(TEXT("/ChaosClothAsset/ClothAssetTemplate.ClothAssetTemplate")),
+					.DisplayName = LOCTEXT("SelectUsdClothTemplate", "USD Cloth"),
+					.Tooltip = LOCTEXT("SelectUsdClothTemplateTooltip", "Add a Dataflow graph that builds a Cloth Asset from an imported USD Cloth.\n(Requires the ChaosClothAssetDataflowNodes and ChaosClothAssetUsdDataflowNodes plug-ins.)"),
+					 .Icon = ChaosClothAssetIcon
+				});
 		}
 
 		virtual void ShutdownModule() override
